Copying //UE4/Dev-Physics to Dev-Main (//UE4/Dev-Main) [at] 7703071
#rnx

diff --git a/Engine/Plugins/Experimental/ChaosCloth/Source/Chaos/Private/ChaosCloth/ChaosClothPrivate.h b/Engine/Plugins/Experimental/ChaosCloth/Source/Chaos/Private/ChaosCloth/ChaosClothPrivate.h
new file mode 100644
--- /dev/null
+++ b/Engine/Plugins/Experimental/ChaosCloth/Source/Chaos/Private/ChaosCloth/ChaosClothPrivate.h
@@ -0,0 +1,7 @@
+// Copyright 1998-2019 Epic Games, Inc. All Rights Reserved.
+
+#pragma once
+
+#include "CoreMinimal.h"
+
+DECLARE_LOG_CATEGORY_EXTERN(LogChaosCloth, Verbose, All);
diff --git a/Engine/Plugins/Experimental/ChaosCloth/Source/Chaos/Private/ChaosCloth/ChaosClothingSimulation.cpp b/Engine/Plugins/Experimental/ChaosCloth/Source/Chaos/Private/ChaosCloth/ChaosClothingSimulation.cpp
new file mode 100644
--- /dev/null
+++ b/Engine/Plugins/Experimental/ChaosCloth/Source/Chaos/Private/ChaosCloth/ChaosClothingSimulation.cpp
@@ -0,0 +1,985 @@
+// Copyright 1998-2019 Epic Games, Inc. All Rights Reserved.
+#include "ChaosCloth/ChaosClothingSimulation.h"
+#include "ChaosCloth/ChaosClothPrivate.h"
+
+#include "Async/ParallelFor.h"
+#include "Assets/ClothingAsset.h"
+#include "ClothingSimulation.h" // ClothingSystemRuntime
+#include "Components/SkeletalMeshComponent.h"
+#include "PhysicsEngine/PhysicsAsset.h"
+#include "Chaos/Box.h"
+#include "Chaos/Capsule.h"
+#include "Chaos/Cylinder.h"
+#include "Chaos/GeometryParticles.h"
+#include "Chaos/ImplicitObjectIntersection.h"
+#include "Chaos/ImplicitObjectUnion.h"
+#include "Chaos/Levelset.h"
+#include "Chaos/PBDAxialSpringConstraints.h"
+#include "Chaos/PBDBendingConstraints.h"
+#include "Chaos/PBDParticles.h"
+#include "Chaos/PBDSpringConstraints.h"
+#include "Chaos/PBDVolumeConstraint.h"
+#include "Chaos/PerParticleGravity.h"
+#include "Chaos/PerParticlePBDLongRangeConstraints.h"
+#include "Chaos/PerParticlePBDShapeConstraints.h"
+#include "Chaos/Plane.h"
+#include "Chaos/Sphere.h"
+#include "Chaos/TaperedCylinder.h"
+#include "Chaos/Transform.h"
+#include "Chaos/Utilities.h"
+#include "Chaos/Vector.h"
+
+#if WITH_PHYSX && !PLATFORM_LUMIN && !PLATFORM_ANDROID
+#include "PhysXIncludes.h" 
+#endif
+
+#include <functional>
+#include <memory>
+#include <unordered_map>
+#include <unordered_set>
+#include "Chaos/ErrorReporter.h"
+
+using namespace Chaos;
+
+static TAutoConsoleVariable<int32> CVarClothNumIterations(TEXT("physics.ClothNumIterations"), 1, TEXT(""));
+static TAutoConsoleVariable<float> CVarClothSelfCollisionThickness(TEXT("physics.ClothSelfCollisionThickness"), 2.f, TEXT(""));
+static TAutoConsoleVariable<float> CVarClothCollisionThickness(TEXT("physics.ClothCollisionThickness"), 1.2f, TEXT(""));
+static TAutoConsoleVariable<float> CVarClothCoefficientOfFriction(TEXT("physics.ClothCoefficientOfFriction"), 0.f, TEXT(""));
+static TAutoConsoleVariable<float> CVarClothDamping(TEXT("physics.ClothDamping"), 0.01f, TEXT(""));
+static TAutoConsoleVariable<float> CVarClothGravityMagnitude(TEXT("physics.ClothGravityMagnitude"), 490.f, TEXT(""));
+
+ClothingSimulation::ClothingSimulation()
+	: NumIterations(1)
+	, EdgeStiffness(1.f)
+	, BendingStiffness(1.f)
+	, AreaStiffness(1.f)
+	, VolumeStiffness(0.f)
+	, StrainLimitingStiffness(1.f)
+	, ShapeTargetStiffness(0.f)
+	, SelfCollisionThickness(2.f)
+	, CollisionThickness(1.2f)
+	, GravityMagnitude(490.f)
+	, bUseBendingElements(false)
+	, bUseTetrahedralConstraints(false)
+	, bUseThinShellVolumeConstraints(false)
+	, bUseSelfCollisions(false)
+	, bUseContinuousCollisionDetection(false)
+{}
+
+ClothingSimulation::~ClothingSimulation() 
+{}
+
+void ClothingSimulation::Initialize()
+{
+#if INCLUDE_CHAOS
+    NumIterations = CVarClothNumIterations.GetValueOnGameThread();
+    SelfCollisionThickness = CVarClothSelfCollisionThickness.GetValueOnGameThread();
+    CollisionThickness = CVarClothCollisionThickness.GetValueOnGameThread();
+    CoefficientOfFriction = CVarClothCoefficientOfFriction.GetValueOnGameThread();
+    Damping = CVarClothDamping.GetValueOnGameThread();
+    GravityMagnitude = CVarClothGravityMagnitude.GetValueOnGameThread();
+
+    Chaos::TPBDParticles<float, 3> LocalParticles;
+    Chaos::TKinematicGeometryClothParticles<float, 3> TRigidParticles;
+    Evolution.Reset(
+		new Chaos::TPBDEvolution<float, 3>(
+			MoveTemp(LocalParticles), 
+			MoveTemp(TRigidParticles), 
+			{}, // CollisionTriangles
+			NumIterations, 
+			CollisionThickness, 
+			SelfCollisionThickness, 
+			CoefficientOfFriction, 
+			Damping));
+    Evolution->CollisionParticles().AddArray(&BoneIndices);
+    Evolution->CollisionParticles().AddArray(&BaseTransforms);
+
+    if (GravityMagnitude)
+    {
+        Evolution->AddForceFunction(
+			Chaos::Utilities::GetDeformablesGravityFunction(
+				Chaos::TVector<float, 3>(0.f, 0.f, -1.f), 
+				GravityMagnitude));
+    }
+
+    Evolution->SetKinematicUpdateFunction(
+		[&](Chaos::TPBDParticles<float, 3>& ParticlesInput, const float Dt, const float LocalTime, const int32 Index) 
+		{
+			if (!OldAnimationPositions.IsValidIndex(Index) || ParticlesInput.InvM(Index) > 0)
+				return;
+			const float Alpha = (LocalTime - Time) / DeltaTime;
+			ParticlesInput.X(Index) = Alpha * AnimationPositions[Index] + (1.f - Alpha) * OldAnimationPositions[Index];
+		});
+
+	Evolution->SetCollisionKinematicUpdateFunction(
+//		[&](Chaos::TKinematicGeometryParticles<float, 3>& ParticlesInput, const float Dt, const float LocalTime, const int32 Index) 
+		[&](Chaos::TKinematicGeometryClothParticles<float, 3>& ParticlesInput, const float Dt, const float LocalTime, const int32 Index) 
+		{
+			checkSlow(DeltaTime > SMALL_NUMBER);
+			const float Alpha = (LocalTime - Time) / DeltaTime;
+			const Chaos::TVector<float, 3> NewX = 
+				Alpha * AnimationTransforms[Index].GetTranslation() + (1.f - Alpha) * OldAnimationTransforms[Index].GetTranslation();
+			ParticlesInput.V(Index) = (NewX - ParticlesInput.X(Index)) / DeltaTime;
+			ParticlesInput.X(Index) = NewX;
+			Chaos::TRotation<float, 3> NewR = FQuat::Slerp(OldAnimationTransforms[Index].GetRotation(), AnimationTransforms[Index].GetRotation(), Alpha);
+			Chaos::TRotation<float, 3> Delta = NewR * ParticlesInput.R(Index).Inverse();
+			Chaos::TVector<float, 3> Axis;
+			float Angle;
+			Delta.ToAxisAndAngle(Axis, Angle);
+			ParticlesInput.W(Index) = Axis * Angle / Dt;
+			ParticlesInput.R(Index) = NewR;
+		});
+
+    MaxDeltaTime = 1.0f;
+    ClampDeltaTime = 0.f;
+    Time = 0.f;
+#endif
+}
+
+void ClothingSimulation::CreateActor(USkeletalMeshComponent* InOwnerComponent, UClothingAssetBase* InAsset, int32 InSimDataIndex)
+{
+	MassMode = InOwnerComponent->MassMode; // uniform, total, density
+	UniformMass = InOwnerComponent->UniformMass;
+	TotalMass = InOwnerComponent->TotalMass;
+	Density = InOwnerComponent->Density;
+	MinMass = InOwnerComponent->MinPerParticleMass;
+
+	EdgeStiffness = InOwnerComponent->EdgeStiffness;
+    BendingStiffness = InOwnerComponent->BendingStiffness;
+    AreaStiffness = InOwnerComponent->AreaStiffness;
+    VolumeStiffness = InOwnerComponent->VolumeStiffness;
+    StrainLimitingStiffness = InOwnerComponent->StrainLimitingStiffness;
+    ShapeTargetStiffness = InOwnerComponent->ShapeTargetStiffness;
+    bUseBendingElements = InOwnerComponent->bUseBendingElements;
+    bUseTetrahedralConstraints = InOwnerComponent->bUseTetrahedralConstraints;
+    bUseThinShellVolumeConstraints = InOwnerComponent->bUseThinShellVolumeConstraints;
+    bUseSelfCollisions = InOwnerComponent->bUseSelfCollisions;
+    bUseContinuousCollisionDetection = InOwnerComponent->bUseContinuousCollisionDetection; // ccd
+
+	//Evolution->SetCCD(bUseContinuousCollisionDetection);
+	//Evolution->SetCCD(true); // ryan!!!
+
+	ClothingSimulationContext Context;
+    FillContext(InOwnerComponent, 0, &Context);
+
+    // TODO(mlentine): Support multiple assets.
+    Asset = Cast<UClothingAsset>(InAsset);
+    check(Asset->LodData.Num() == 1);
+    FClothLODData& AssetLodData = Asset->LodData[0];
+    FClothPhysicalMeshData& PhysMesh = AssetLodData.PhysicalMeshData;
+
+	// SkinPhysicsMesh() strips scale from RootBoneTransform ("Ignore any user scale.  
+	// It's already accounted for in our skinning matrices."), and returns all points 
+	// in that space.
+	FTransform RootBoneTransform = Context.BoneTransforms[Asset->ReferenceBoneIndex];
+	FClothingSimulationBase::SkinPhysicsMesh(
+		Asset, 
+		PhysMesh, // curr pos and norm
+		RootBoneTransform, 
+		Context.RefToLocals.GetData(), 
+		Context.RefToLocals.Num(), 
+		reinterpret_cast<TArray<FVector>&>(AnimationPositions), 
+		reinterpret_cast<TArray<FVector>&>(AnimationNormals));
+
+	// Transform points & normals to world space
+	RootBoneTransform.SetScale3D(FVector(1.0f));
+	const FTransform RootBoneWorldTransform = RootBoneTransform * Context.ComponentToWorld;
+	ParallelFor(AnimationPositions.Num(), 
+		[&](int32 Index) 
+		{
+			AnimationPositions[Index] = RootBoneWorldTransform.TransformPosition(AnimationPositions[Index]);
+	        AnimationNormals[Index] = RootBoneWorldTransform.TransformVector(AnimationNormals[Index]);
+		});
+
+    TPBDParticles<float, 3>& Particles = Evolution->Particles();
+    const uint32 Offset = Particles.Size();
+    Particles.AddParticles(PhysMesh.Vertices.Num());
+
+	if (IndexToRangeMap.Num() <= InSimDataIndex)
+        IndexToRangeMap.SetNum(InSimDataIndex + 1);
+    IndexToRangeMap[InSimDataIndex] = Chaos::TVector<uint32, 2>(Offset, Particles.Size());
+
+    for (uint32 i = Offset; i < Particles.Size(); ++i)
+    {
+        Particles.X(i) = AnimationPositions[i - Offset];
+        Particles.V(i) = Chaos::TVector<float, 3>(0.f, 0.f, 0.f);
+		// Initialize mass to 0, to be overridden later
+		Particles.M(i) = 0.f;
+    }
+
+	const int32 NumTriangles = PhysMesh.Indices.Num() / 3;
+	TArray<Chaos::TVector<int32, 3>> InputSurfaceElements;
+	InputSurfaceElements.Reserve(NumTriangles);
+    for (int i = 0; i < NumTriangles; ++i)
+    {
+        const int32 Index = 3 * i;
+        InputSurfaceElements.Add(
+			{static_cast<int32>(PhysMesh.Indices[Index]), 
+			 static_cast<int32>(PhysMesh.Indices[Index + 1]), 
+			 static_cast<int32>(PhysMesh.Indices[Index + 2])});
+    }
+	check(InputSurfaceElements.Num() == NumTriangles);
+	if (Meshes.Num() <= InSimDataIndex)
+	{
+		Meshes.SetNum(InSimDataIndex + 1);
+		FaceNormals.SetNum(InSimDataIndex + 1);
+		PointNormals.SetNum(InSimDataIndex + 1);
+	}
+	TUniquePtr<Chaos::TTriangleMesh<float>>& Mesh = Meshes[InSimDataIndex];
+    Mesh.Reset(new Chaos::TTriangleMesh<float>(MoveTemp(InputSurfaceElements)));
+	check(Mesh->GetNumElements() == NumTriangles);
+    const auto& SurfaceElements = Mesh->GetSurfaceElements();
+	Mesh->GetPointToTriangleMap(); // Builds map for later use by GetPointNormals().
+
+	// Assign per particle mass proportional to connected area.
+	float TotalArea = 0.0;
+	for (const Chaos::TVector<int32, 3>& Tri : SurfaceElements)
+	{
+		const float TriArea = 0.5 * Chaos::TVector<float, 3>::CrossProduct(
+			Particles.X(Tri[1] + Offset) - Particles.X(Tri[0] + Offset),
+			Particles.X(Tri[2] + Offset) - Particles.X(Tri[0] + Offset)).Size();
+		TotalArea += TriArea;
+		const float ThirdTriArea = TriArea / 3.0;
+		Particles.M(Tri[0] + Offset) += ThirdTriArea;
+		Particles.M(Tri[1] + Offset) += ThirdTriArea;
+		Particles.M(Tri[2] + Offset) += ThirdTriArea;
+	}
+	const TSet<int32> Vertices = Mesh->GetVertices();
+	switch (MassMode)
+	{
+	case EClothMassMode::UniformMass:
+		for (const int32 Vertex : Vertices)
+		{
+			Particles.M(Vertex) = UniformMass;
+		}
+		break;
+	case EClothMassMode::TotalMass:
+	{
+		const float MassPerUnitArea = TotalArea > 0.0 ? TotalMass / TotalArea : 1.0;
+		for (const int32 Vertex : Vertices)
+		{
+			Particles.M(Vertex) *= MassPerUnitArea;
+		}
+		break;
+	}
+	case EClothMassMode::Density:
+		for (const int32 Vertex : Vertices)
+		{
+			Particles.M(Vertex) *= Density;
+		}
+		break;
+	};
+	// Clamp and enslave
+	for (uint32 i = Offset; i < Particles.Size(); i++)
+	{
+		Particles.M(i) = FMath::Max(Particles.M(i), MinMass);
+		Particles.InvM(i) = PhysMesh.MaxDistances[i-Offset] > 0.1 ? 1.0/Particles.M(i) : 0.0;
+	}
+
+    // Add Model
+    if (ShapeTargetStiffness)
+    {
+        check(ShapeTargetStiffness > 0.f && ShapeTargetStiffness <= 1.f);
+		Evolution->AddPBDConstraintFunction([ShapeConstraints = Chaos::TPerParticlePBDShapeConstraints<float, 3>(Evolution->Particles(), AnimationPositions, ShapeTargetStiffness)](TPBDParticles<float, 3>& InParticles, const float Dt) {
+			ShapeConstraints.Apply(InParticles, Dt);
+		});
+    }
+    if (EdgeStiffness)
+    {
+        check(EdgeStiffness > 0.f && EdgeStiffness <= 1.f);
+        Evolution->AddPBDConstraintFunction([SpringConstraints = Chaos::TPBDSpringConstraints<float, 3>(Evolution->Particles(), SurfaceElements, EdgeStiffness)](TPBDParticles<float, 3>& InParticles, const float Dt) {
+            SpringConstraints.Apply(InParticles, Dt);
+        });
+    }
+    if (BendingStiffness)
+    {
+        check(BendingStiffness > 0.f && BendingStiffness <= 1.f);
+        if (bUseBendingElements)
+        {
+            TArray<Chaos::TVector<int32, 4>> BendingConstraints = Mesh->GetUniqueAdjacentElements();
+			Evolution->AddPBDConstraintFunction([BendConstraints = Chaos::TPBDBendingConstraints<float>(Evolution->Particles(), MoveTemp(BendingConstraints))](TPBDParticles<float, 3>& InParticles, const float Dt) {
+				BendConstraints.Apply(InParticles, Dt);
+			});
+        }
+        else
+        {
+            TArray<Chaos::TVector<int32, 2>> BendingConstraints = Mesh->GetUniqueAdjacentPoints();
+            Evolution->AddPBDConstraintFunction([SpringConstraints = Chaos::TPBDSpringConstraints<float, 3>(Evolution->Particles(), MoveTemp(BendingConstraints), BendingStiffness)](TPBDParticles<float, 3>& InParticles, const float Dt) {
+                SpringConstraints.Apply(InParticles, Dt);
+            });
+        }
+    }
+    if (AreaStiffness)
+    {
+        TArray<Chaos::TVector<int32, 3>> SurfaceConstraints = SurfaceElements;
+		Evolution->AddPBDConstraintFunction([SurfConstraints = Chaos::TPBDAxialSpringConstraints<float, 3>(Evolution->Particles(), MoveTemp(SurfaceConstraints), AreaStiffness)](TPBDParticles<float, 3>& InParticles, const float Dt) {
+			SurfConstraints.Apply(InParticles, Dt);
+		});
+    }
+    if (VolumeStiffness)
+    {
+        check(VolumeStiffness > 0.f && VolumeStiffness <= 1.f);
+        if (bUseTetrahedralConstraints)
+        {
+            // TODO(mlentine): Need to tetrahedralize surface to support this
+            check(false);
+        }
+        else if (bUseThinShellVolumeConstraints)
+        {
+            TArray<Chaos::TVector<int32, 2>> BendingConstraints = Mesh->GetUniqueAdjacentPoints();
+            TArray<Chaos::TVector<int32, 2>> DoubleBendingConstraints;
+            {
+                TMap<int32, TArray<int32>> BendingHash;
+                for (int32 i = 0; i < BendingConstraints.Num(); ++i)
+                {
+                    BendingHash.FindOrAdd(BendingConstraints[i][0]).Add(BendingConstraints[i][1]);
+                    BendingHash.FindOrAdd(BendingConstraints[i][1]).Add(BendingConstraints[i][0]);
+                }
+                TSet<Chaos::TVector<int32, 2>> Visited;
+                for (auto Elem : BendingHash)
+                {
+                    for (int32 i = 0; i < Elem.Value.Num(); ++i)
+                    {
+                        for (int32 j = i + 1; j < Elem.Value.Num(); ++j)
+                        {
+                            if (Elem.Value[i] == Elem.Value[j])
+                                continue;
+                            auto NewElem = Chaos::TVector<int32, 2>(Elem.Value[i], Elem.Value[j]);
+                            if (!Visited.Contains(NewElem))
+                            {
+                                DoubleBendingConstraints.Add(NewElem);
+                                Visited.Add(NewElem);
+                                Visited.Add(Chaos::TVector<int32, 2>(Elem.Value[j], Elem.Value[i]));
+                            }
+                        }
+                    }
+                }
+            }
+            Evolution->AddPBDConstraintFunction([SpringConstraints = Chaos::TPBDSpringConstraints<float, 3>(Evolution->Particles(), MoveTemp(DoubleBendingConstraints), VolumeStiffness)](TPBDParticles<float, 3>& InParticles, const float Dt) {
+                SpringConstraints.Apply(InParticles, Dt);
+            });
+        }
+        else
+        {
+            TArray<Chaos::TVector<int32, 3>> SurfaceConstraints = SurfaceElements;
+            Evolution->AddPBDConstraintFunction(std::bind(&Chaos::TPBDVolumeConstraint<float>::Apply,
+                                                          Chaos::TPBDVolumeConstraint<float>(Evolution->Particles(), MoveTemp(SurfaceConstraints), VolumeStiffness), std::placeholders::_1, std::placeholders::_2));
+        }
+    }
+    if (StrainLimitingStiffness)
+    {
+		check(Mesh->GetNumElements() > 0);
+        Evolution->AddPBDConstraintFunction(
+			std::bind(
+				static_cast
+				<
+					void (Chaos::TPerParticlePBDLongRangeConstraints<float, 3>::*)(
+						TPBDParticles<float, 3> & InParticles, 
+						const float Dt) const
+				>(&Chaos::TPerParticlePBDLongRangeConstraints<float, 3>::Apply),
+				Chaos::TPerParticlePBDLongRangeConstraints<float, 3>(
+					Evolution->Particles(), 
+					Mesh->GetPointToNeighborsMap(), 
+					10, // The max number of connected neighbors per particle.  ryan - What should this be?  Was k...
+					StrainLimitingStiffness), 
+				std::placeholders::_1, 
+				std::placeholders::_2));
+    }
+    // Add Self Collisions
+    if (bUseSelfCollisions)
+    {
+        // TODO(mlentine): Parallelize these for multiple meshes
+        Evolution->CollisionTriangles().Append(SurfaceElements);
+        for (uint32 i = Offset; i < Particles.Size(); ++i)
+        {
+            auto Neighbors = Mesh->GetNRing(i, 5);
+            for (const auto& Element : Neighbors)
+            {
+                check(i != Element);
+                Evolution->DisabledCollisionElements().Add(Chaos::TVector<int32, 2>(i, Element));
+                Evolution->DisabledCollisionElements().Add(Chaos::TVector<int32, 2>(Element, i));
+            }
+        }
+    }
+    // Add Collision Bodies
+    TGeometryClothParticles<float, 3>& CollisionParticles = Evolution->CollisionParticles();
+    //USkeletalMesh* TargetMesh = InOwnerComponent->SkeletalMesh;
+	USkeletalMesh* TargetMesh = CastChecked<USkeletalMesh>(Asset->GetOuter());
+
+		// TODO(mlentine): Support collision body activation on a per particle basis, preferably using a map but also can be a particle attribute
+    if (UPhysicsAsset* PhysAsset = Asset->PhysicsAsset)
+    {
+        for (const USkeletalBodySetup* BodySetup : PhysAsset->SkeletalBodySetups)
+        {
+			if (!BodySetup)
+				continue;
+
+			const int32 MeshBoneIndex = TargetMesh->RefSkeleton.FindBoneIndex(BodySetup->BoneName);
+			const int32 MappedBoneIndex = 
+				MeshBoneIndex != INDEX_NONE ? 
+				Asset->UsedBoneNames.AddUnique(BodySetup->BoneName) : INDEX_NONE;
+
+			const FKAggregateGeom& AggGeom = BodySetup->AggGeom; 
+			if(AggGeom.SphylElems.Num())
+            {
+                uint32 OldSize = CollisionParticles.Size();
+                CollisionParticles.AddParticles(AggGeom.SphylElems.Num());
+                for (uint32 i = OldSize; i < CollisionParticles.Size(); ++i)
+                {
+                    const auto& Capsule = AggGeom.SphylElems[i - OldSize];
+                    if (Capsule.Length == 0)
+                    {
+                        CollisionParticles.SetDynamicGeometry(i, MakeUnique<Chaos::TSphere<float,3>>(Chaos::TVector<float, 3>(0), Capsule.Radius));
+                    }
+                    else
+                    {
+                        Chaos::TVector<float, 3> half_extents(0, 0, Capsule.Length / 2);
+                        CollisionParticles.SetDynamicGeometry(
+							i, 
+							MakeUnique<Chaos::TCapsule<float>>(
+								-half_extents, half_extents, Capsule.Radius));
+                    }
+                    BaseTransforms[i] = Chaos::TRigidTransform<float, 3>(Capsule.Center, Capsule.Rotation.Quaternion());
+                    BoneIndices[i] = MappedBoneIndex;
+                }
+            }
+			if(AggGeom.SphereElems.Num())
+            {
+                uint32 OldSize = CollisionParticles.Size();
+                CollisionParticles.AddParticles(AggGeom.SphereElems.Num());
+                for (uint32 i = OldSize; i < CollisionParticles.Size(); ++i)
+                {
+                    const auto& CollisionSphere = AggGeom.SphereElems[i - OldSize];
+                    CollisionParticles.SetDynamicGeometry(
+						i, 
+						MakeUnique<Chaos::TSphere<float,3>>(
+							Chaos::TVector<float, 3>(0.f, 0.f, 0.f), 
+							CollisionSphere.Radius));
+                    BaseTransforms[i] = Chaos::TRigidTransform<float, 3>(CollisionSphere.Center, Chaos::TRotation<float, 3>::MakeFromEuler(Chaos::TVector<float, 3>(0.f, 0.f, 0.f)));
+                    BoneIndices[i] = MappedBoneIndex;
+                }
+            }
+			if(AggGeom.BoxElems.Num())
+            {
+                uint32 OldSize = CollisionParticles.Size();
+                CollisionParticles.AddParticles(AggGeom.BoxElems.Num());
+                for (uint32 i = OldSize; i < CollisionParticles.Size(); ++i)
+                {
+                    const auto& Box = AggGeom.BoxElems[i - OldSize];
+                    Chaos::TVector<float, 3> half_extents(Box.X / 2.f, Box.Y / 2.f, Box.Z / 2.f);
+                    CollisionParticles.SetDynamicGeometry(i, MakeUnique<Chaos::TBox<float,3>>(-half_extents, half_extents));
+                    BaseTransforms[i] = Chaos::TRigidTransform<float, 3>(Box.Center, Box.Rotation.Quaternion());
+                    BoneIndices[i] = MappedBoneIndex;
+                }
+            }
+            /*
+			if(AggGeom.TaperedCapsuleElems.Num())
+			{
+                int32 OldSize = CollisionParticles.Size();
+                CollisionParticles.AddParticles(AggGeom.TaperedCapsuleElems.Num());
+                for (int32 i = OldSize; i < CollisionParticles.Size(); ++i)
+                {
+                    const auto& Capsule = AggGeom.TaperedCapsuleElems[i - OldSize];
+                    if (Capsule.Length == 0)
+                    {
+                        CollisionParticles.Geometry(i) = new Chaos::Sphere<float, 3>(Chaos::TVector<float, 3>(0), Capsule.Radius1 > Capsule.Radius0 ? Capsule.Radius1 : Capsule.Radius0);
+                    }
+                    else
+                    {
+                        TArray<TUniquePtr<TImplicitObject<float, 3>>> Objects;
+                        Chaos::TVector<float, 3> half_extents(0, 0, Capsule.Length / 2);
+                        Objects.Add(TUniquePtr<Chaos::TImplicitObject<float, 3>>(
+                            new Chaos::TTaperedCylinder<float>(-half_extents, half_extents, Capsule.Radius1, Capsule.Radius0)));
+                        Objects.Add(TUniquePtr<Chaos::TImplicitObject<float, 3>>(
+                            new Chaos::Sphere<float, 3>(-half_extents, Capsule.Radius1)));
+                        Objects.Add(TUniquePtr<Chaos::TImplicitObject<float, 3>>(
+                            new Chaos::Sphere<float, 3>(half_extents, Capsule.Radius0)));
+                        CollisionParticles.Geometry(i) = new Chaos::TImplicitObjectUnion<float, 3>(MoveTemp(Objects));
+                    }
+					BaseTransforms[i] = Chaos::TRigidTransform<float, 3>(Capsule.Center, Capsule.Rotation.Quaternion());
+                    BoneIndices[i] = MappedBoneIndex;
+                }
+			}*/
+			if(AggGeom.ConvexElems.Num())
+            {
+// Collision bodies are stored in PhysX specific data structures so they can only be imported if we enable PhysX.
+#if WITH_PHYSX && !PLATFORM_LUMIN && !PLATFORM_ANDROID
+                uint32 OldSize = CollisionParticles.Size();
+                CollisionParticles.AddParticles(AggGeom.ConvexElems.Num());
+                for (uint32 i = OldSize; i < CollisionParticles.Size(); ++i)
+                {
+                    const auto& CollisionBody = AggGeom.ConvexElems[i - OldSize];
+                    TArray<Chaos::TVector<int32, 3>> CollisionMeshElements;
+                    const auto PhysXMesh = CollisionBody.GetConvexMesh();
+                    for (int32 j = 0; j < static_cast<int32>(PhysXMesh->getNbPolygons()); ++j)
+                    {
+                        physx::PxHullPolygon Poly;
+                        PhysXMesh->getPolygonData(j, Poly);
+                        check(Poly.mNbVerts == 3);
+                        const auto Indices = PhysXMesh->getIndexBuffer() + Poly.mIndexBase;
+                        CollisionMeshElements.Add(Chaos::TVector<int32, 3>(Indices[0], Indices[1], Indices[2]));
+                    }
+                    Chaos::TParticles<float, 3> CollisionMeshParticles;
+                    CollisionMeshParticles.AddParticles(CollisionBody.VertexData.Num());
+                    for (uint32 j = 0; j < CollisionMeshParticles.Size(); ++j)
+                    {
+                        CollisionMeshParticles.X(j) = CollisionBody.VertexData[j];
+                    }
+                    Chaos::TBox<float, 3> BoundingBox(CollisionMeshParticles.X(0), CollisionMeshParticles.X(0));
+                    for (uint32 j = 1; j < CollisionMeshParticles.Size(); ++j)
+                    {
+                        BoundingBox.GrowToInclude(CollisionMeshParticles.X(i));
+                    }
+                    int32 MaxAxisSize = 100;
+                    int32 MaxAxis;
+                    const auto Extents = BoundingBox.Extents();
+                    if (Extents[0] > Extents[1] && Extents[0] > Extents[2])
+                    {
+                        MaxAxis = 0;
+                    }
+                    else if (Extents[1] > Extents[2])
+                    {
+                        MaxAxis = 1;
+                    }
+                    else
+                    {
+                        MaxAxis = 2;
+                    }
+                    Chaos::TUniformGrid<float, 3> Grid(BoundingBox.Min(), BoundingBox.Max(), Chaos::TVector<int32, 3>(100 * Extents[0] / Extents[MaxAxis], 100 * Extents[0] / Extents[MaxAxis], 100 * Extents[0] / Extents[MaxAxis]));
+                    Chaos::TTriangleMesh<float> CollisionMesh(MoveTemp(CollisionMeshElements));
+					Chaos::FErrorReporter ErrorReporter;
+                    CollisionParticles.SetDynamicGeometry(i, MakeUnique<Chaos::TLevelSet<float,3>>(ErrorReporter, Grid, CollisionMeshParticles, CollisionMesh));
+                    BaseTransforms[i] = Chaos::TRigidTransform<float, 3>(Chaos::TVector<float, 3>(0.f, 0.f, 0.f), Chaos::TRotation<float, 3>::MakeFromEuler(Chaos::TVector<float, 3>(0.f, 0.f, 0.f)));
+                    BoneIndices[i] = MappedBoneIndex;
+                }
+#endif
+            }
+        } // end for
+    } // end if PhysAsset
+
+	// We can't just use AddExternalCollisions() because we need to add entries for bone mappings and lookups.
+	const FClothCollisionData& LodCollData = AssetLodData.CollisionData;
+	if (LodCollData.Spheres.Num() || LodCollData.SphereConnections.Num() || LodCollData.Convexes.Num())
+	{
+		UE_LOG(LogSkeletalMesh, Warning,
+			TEXT("Actor '%s' component '%s' has %d sphere, %d capsule, and %d "
+				"convex collision objects for physics authored as part of a LOD construct, "
+				"probably by the Apex cloth authoring system.  This is deprecated.  "
+				"Please update your asset!"),
+			*InOwnerComponent->GetOwner()->GetName(),
+			*InOwnerComponent->GetName(),
+			LodCollData.Spheres.Num(),
+			LodCollData.SphereConnections.Num(),
+			LodCollData.Convexes.Num());
+
+		TSet<int32> CapsuleEnds;
+		if (LodCollData.SphereConnections.Num())
+		{
+			const uint32 Size = CollisionParticles.Size();
+			CollisionParticles.AddParticles(LodCollData.SphereConnections.Num());
+			CapsuleEnds.Reserve(LodCollData.SphereConnections.Num() * 2);
+			for (uint32 i = Size; i < CollisionParticles.Size(); ++i)
+			{
+				// This data was pulled from a FKSphylElem, which is a capsule.  So 
+				// it should only have 1 radius, and the BoneIndex for both spheres
+				// should be the same.
+				const FClothCollisionPrim_SphereConnection& Connection = LodCollData.SphereConnections[i - Size];
+
+				const int32 SphereIndex0 = Connection.SphereIndices[0];
+				const int32 SphereIndex1 = Connection.SphereIndices[1];
+				checkSlow(SphereIndex0 != SphereIndex1);
+				const float Radius = LodCollData.Spheres[SphereIndex0].Radius;
+				checkSlow(LodCollData.Spheres[SphereIndex0].Radius - LodCollData.Spheres[SphereIndex1].Radius < SMALL_NUMBER);
+				const Chaos::TVector<float, 3> X0 = LodCollData.Spheres[SphereIndex0].LocalPosition;
+				const Chaos::TVector<float, 3> X1 = LodCollData.Spheres[SphereIndex1].LocalPosition;
+
+				checkSlow(LodCollData.Spheres[SphereIndex0].BoneIndex == LodCollData.Spheres[SphereIndex1].BoneIndex);
+				const int32 BoneIndex = LodCollData.Spheres[SphereIndex0].BoneIndex;
+				const FName BoneName = 
+					TargetMesh->RefSkeleton.IsValidIndex(BoneIndex) ? 
+					TargetMesh->RefSkeleton.GetBoneName(BoneIndex) : NAME_None;
+				const int32 MappedBoneIndex = BoneName != NAME_None ?
+					Asset->UsedBoneNames.AddUnique(BoneName) : INDEX_NONE;
+				BoneIndices[i] = MappedBoneIndex;
+
+				const Chaos::TVector<float, 3> Center = (X0 + X1) * 0.5;
+				const Chaos::TVector<float, 3> Axis = X1 - X0;
+				const float HalfHeight = Axis.Size() * 0.5;
+
+				// We construct a capsule centered at the origin along the Z axis, and
+				// then move it into place with X and R.
+				CollisionParticles.X(i) = Center;
+				CollisionParticles.R(i) = Chaos::TRotation<float, 3>::FromRotatedVector(Chaos::TVector<float, 3>::AxisVector(2), Axis.GetSafeNormal());
+				BaseTransforms[i] = Chaos::TRigidTransform<float, 3>(CollisionParticles.X(i), CollisionParticles.R(i));
+
+				CollisionParticles.SetDynamicGeometry(
+					i,
+					MakeUnique<Chaos::TCapsule<float>>(
+						Chaos::TVector<float, 3>(0.f, 0.f, -HalfHeight), // Min
+						Chaos::TVector<float, 3>(0.f, 0.f, HalfHeight), // Max
+						Radius));
+				IndexAndCapsuleCollisionMap.Add(MakePair(i, Connection));
+
+				// Skip spheres added as end caps for the capsule.
+				CapsuleEnds.Add(SphereIndex0);
+				CapsuleEnds.Add(SphereIndex1);
+			}
+		}
+		if (LodCollData.Spheres.Num() - CapsuleEnds.Num())
+		{
+			const uint32 Size = CollisionParticles.Size();
+			CollisionParticles.AddParticles(LodCollData.Spheres.Num() - CapsuleEnds.Num());
+			// i = Spheres index, j = CollisionParticles index
+			for (uint32 i = 0, j = Size; i < (uint32)LodCollData.Spheres.Num(); ++i)
+			{
+				// Skip spheres that are the end caps of capsules.
+				if (CapsuleEnds.Contains(i))
+					continue;
+
+				const FClothCollisionPrim_Sphere& CollisionSphere = LodCollData.Spheres[i];
+
+				const int32 BoneIndex = CollisionSphere.BoneIndex;
+				const FName BoneName = 
+					TargetMesh->RefSkeleton.IsValidIndex(BoneIndex) ? 
+					TargetMesh->RefSkeleton.GetBoneName(BoneIndex) : NAME_None;
+				const int32 MappedBoneIndex = BoneName != NAME_None ?
+					Asset->UsedBoneNames.AddUnique(BoneName) : INDEX_NONE;
+				BoneIndices[j] = MappedBoneIndex;
+
+				CollisionParticles.X(j) = CollisionSphere.LocalPosition;
+				CollisionParticles.R(j) = Chaos::TRotation<float, 3>(Chaos::TVector<float, 3>::AxisVector(0), 0.f);
+				BaseTransforms[j] = Chaos::TRigidTransform<float, 3>(CollisionParticles.X(j), CollisionParticles.R(j));
+
+				CollisionParticles.SetDynamicGeometry(
+					j,
+					MakeUnique<Chaos::TSphere<float, 3>>(
+						Chaos::TVector<float, 3>(0.f, 0.f, 0.f), 
+						CollisionSphere.Radius));
+				IndexAndSphereCollisionMap.Add(MakePair(j, CollisionSphere));
+				j++;
+			}
+		}
+
+	} // end if LodCollData
+
+    AnimationTransforms.SetNum(BaseTransforms.Num());
+    for (uint32 i = 0; i < CollisionParticles.Size(); ++i)
+    {
+		const int32 BoneIndex = BoneIndices[i];
+        if (Asset->UsedBoneIndices.IsValidIndex(BoneIndex))
+        {
+            const int32 MappedIndex = Asset->UsedBoneIndices[BoneIndex];
+            if (Context.BoneTransforms.IsValidIndex(MappedIndex))
+            {
+                const FTransform& BoneTransform = Context.BoneTransforms[MappedIndex];
+                AnimationTransforms[i] = BaseTransforms[i] * BoneTransform * Context.ComponentToWorld;
+                CollisionParticles.X(i) = AnimationTransforms[i].GetTranslation();
+                CollisionParticles.R(i) = AnimationTransforms[i].GetRotation();
+            }
+		}
+    }
+}
+
+void ClothingSimulation::FillContext(USkeletalMeshComponent* InComponent, float InDeltaTime, IClothingSimulationContext* InOutContext)
+{
+    ClothingSimulationContext* Context = static_cast<ClothingSimulationContext*>(InOutContext);
+    Context->ComponentToWorld = InComponent->GetComponentToWorld();
+    Context->DeltaTime = ClampDeltaTime > 0 ? std::min(InDeltaTime, ClampDeltaTime) : InDeltaTime;
+
+	Context->RefToLocals.Reset();
+    InComponent->GetCurrentRefToLocalMatrices(Context->RefToLocals, 0);
+
+	const USkeletalMesh* SkelMesh = InComponent->SkeletalMesh;
+    if (USkinnedMeshComponent* MasterComponent = InComponent->MasterPoseComponent.Get())
+    {
+		const TArray<int32>& MasterBoneMap = InComponent->GetMasterBoneMap();
+        int32 NumBones = MasterBoneMap.Num();
+        if (NumBones == 0)
+        {
+            if (InComponent->SkeletalMesh)
+            {
+                // This case indicates an invalid master pose component (e.g. no skeletal mesh)
+                NumBones = InComponent->SkeletalMesh->RefSkeleton.GetNum();
+            }
+			Context->BoneTransforms.Reset(NumBones);
+			Context->BoneTransforms.AddDefaulted(NumBones);
+        }
+        else
+        {
+            Context->BoneTransforms.Reset(NumBones);
+            Context->BoneTransforms.AddDefaulted(NumBones);
+			const TArray<FTransform>& MasterTransforms = MasterComponent->GetComponentSpaceTransforms();
+            for (int32 BoneIndex = 0; BoneIndex < NumBones; ++BoneIndex)
+            {
+                bool bFoundMaster = false;
+                if (MasterBoneMap.IsValidIndex(BoneIndex))
+                {
+                    const int32 MasterIndex = MasterBoneMap[BoneIndex];
+                    if (MasterTransforms.IsValidIndex(MasterIndex))
+                    {
+                        Context->BoneTransforms[BoneIndex] = MasterTransforms[MasterIndex];
+                        bFoundMaster = true;
+                    }
+                }
+
+                if (!bFoundMaster && SkelMesh)
+                {
+                    const int32 ParentIndex = SkelMesh->RefSkeleton.GetParentIndex(BoneIndex);
+					check(ParentIndex < BoneIndex);
+					Context->BoneTransforms[BoneIndex] = 
+						Context->BoneTransforms.IsValidIndex(ParentIndex) && ParentIndex < BoneIndex ?
+						Context->BoneTransforms[ParentIndex] * SkelMesh->RefSkeleton.GetRefBonePose()[BoneIndex] :
+                        SkelMesh->RefSkeleton.GetRefBonePose()[BoneIndex];
+                }
+            }
+        }
+    }
+    else
+    {
+        Context->BoneTransforms = InComponent->GetComponentSpaceTransforms();
+    }
+}
+
+void ClothingSimulation::Simulate(IClothingSimulationContext* InContext)
+{
+    ClothingSimulationContext* Context = static_cast<ClothingSimulationContext*>(InContext);
+    if (Context->DeltaTime == 0)
+        return;
+
+	// Get New Animation Positions and Normals
+    OldAnimationTransforms = AnimationTransforms;
+    OldAnimationPositions = AnimationPositions;
+    const FClothLODData& AssetLodData = Asset->LodData[0];
+    const FClothPhysicalMeshData& PhysMesh = AssetLodData.PhysicalMeshData;
+
+	FTransform RootBoneTransform = Context->BoneTransforms[Asset->ReferenceBoneIndex];
+    FClothingSimulationBase::SkinPhysicsMesh(
+		Asset, 
+		PhysMesh, 
+		RootBoneTransform, 
+		Context->RefToLocals.GetData(), 
+		Context->RefToLocals.Num(), 
+		reinterpret_cast<TArray<FVector>&>(AnimationPositions), 
+		reinterpret_cast<TArray<FVector>&>(AnimationNormals));
+
+	RootBoneTransform.SetScale3D(FVector(1.0f));
+	
+	// Removing Context->ComponentToWorld means the sim doesn't see updates to the component level xf
+	const FTransform RootBoneWorldTransform = RootBoneTransform * Context->ComponentToWorld;
+
+    ParallelFor(AnimationPositions.Num(), 
+		[&](int32 Index) 
+		{
+			AnimationPositions[Index] = RootBoneWorldTransform.TransformPosition(AnimationPositions[Index]);
+			AnimationNormals[Index] = RootBoneWorldTransform.TransformVector(AnimationNormals[Index]);
+		});
+
+	// Collision bodies
+    TGeometryClothParticles<float, 3>& CollisionParticles = Evolution->CollisionParticles();
+    for (uint32 i = 0; i < CollisionParticles.Size(); ++i)
+    {
+        if (BoneIndices[i] != INDEX_NONE)
+        {
+            const int32 MappedIndex = Asset->UsedBoneIndices[BoneIndices[i]];
+            if (MappedIndex != INDEX_NONE)
+            {
+                const FTransform& BoneTransform = Context->BoneTransforms[MappedIndex];
+                AnimationTransforms[i] = BaseTransforms[i] * BoneTransform * Context->ComponentToWorld;
+            }
+        }
+    }
+    // Advance Sim
+    DeltaTime = Context->DeltaTime;
+    while (Context->DeltaTime > MaxDeltaTime)
+    {
+        Evolution->AdvanceOneTimeStep(MaxDeltaTime);
+        Context->DeltaTime -= MaxDeltaTime;
+    }
+    Evolution->AdvanceOneTimeStep(Context->DeltaTime);
+    Time += DeltaTime;
+}
+
+void ClothingSimulation::GetSimulationData(
+	TMap<int32, FClothSimulData>& OutData, 
+	USkeletalMeshComponent* InOwnerComponent, 
+	USkinnedMeshComponent* InOverrideComponent) const
+{
+	const FTransform& OwnerTransform = InOwnerComponent->GetComponentTransform();
+    for (int32 i = 0; i < IndexToRangeMap.Num(); ++i)
+    {	
+		const TUniquePtr<Chaos::TTriangleMesh<float>>& Mesh = Meshes[i];
+		if (!Mesh) 
+			continue;
+		Mesh->GetFaceNormals(FaceNormals[i], Evolution->Particles().X(), false);
+		Mesh->GetPointNormals(PointNormals[i], FaceNormals[i], false);
+
+		FClothSimulData& Data = OutData.FindOrAdd(i);
+		Data.Reset();
+
+		const TArray<FTransform>& ComponentSpaceTransforms = InOverrideComponent ?
+			InOverrideComponent->GetComponentSpaceTransforms() :
+			InOwnerComponent->GetComponentSpaceTransforms();
+		if (!ComponentSpaceTransforms.IsValidIndex(Asset->ReferenceBoneIndex))
+		{
+			UE_LOG(LogSkeletalMesh, Warning,
+				TEXT("Failed to write back clothing simulation data for component '%s' as bone transforms are invalid."),
+				*InOwnerComponent->GetName());
+			check(false);
+			continue;
+		}
+
+		FTransform RootBoneTransform = ComponentSpaceTransforms[Asset->ReferenceBoneIndex];
+		RootBoneTransform.SetScale3D(FVector(1.0f));
+		RootBoneTransform *= OwnerTransform;
+		Data.Transform = RootBoneTransform;
+		Data.ComponentRelativeTransform = OwnerTransform.Inverse();
+
+		const Chaos::TVector<uint32, 2>& VertexDomain = IndexToRangeMap[i];
+		const uint32 VertexRange = VertexDomain[1] - VertexDomain[0];
+		Data.Positions.SetNum(VertexRange);
+        Data.Normals.SetNum(VertexRange);
+		for (uint32 j = VertexDomain[0]; j < VertexDomain[1]; ++j)
+        {
+			const uint32 LocalIndex = j - VertexDomain[0];
+            Data.Positions[LocalIndex] = Evolution->Particles().X(j);
+            Data.Normals[LocalIndex] = PointNormals[i][LocalIndex];
+		}
+    }
+}
+
+void ClothingSimulation::AddExternalCollisions(const FClothCollisionData& InData)
+{
+    TGeometryClothParticles<float, 3>& CollisionParticles = Evolution->CollisionParticles();
+	TSet<int32> CapsuleEnds;
+
+	const uint32 NumParticles0 = CollisionParticles.Size();
+
+	if(InData.SphereConnections.Num())
+	{
+		const uint32 Size = CollisionParticles.Size();
+        CollisionParticles.AddParticles(InData.SphereConnections.Num());
+		CapsuleEnds.Reserve(InData.SphereConnections.Num() * 2);
+        for (uint32 i = Size; i < CollisionParticles.Size(); ++i)
+        {
+			// This data was pulled from a FKSphylElem, which is a capsule.  So 
+			// it should only have 1 radius, and the BoneIndex for both spheres
+			// should be the same.
+			const FClothCollisionPrim_SphereConnection& Connection = InData.SphereConnections[i - Size];
+			const int32 SphereIndex0 = Connection.SphereIndices[0];
+			const int32 SphereIndex1 = Connection.SphereIndices[1];
+			checkSlow(SphereIndex0 != SphereIndex1);
+			const float Radius = InData.Spheres[SphereIndex0].Radius;
+			checkSlow(InData.Spheres[SphereIndex0].Radius - InData.Spheres[SphereIndex1].Radius < SMALL_NUMBER);
+			const Chaos::TVector<float, 3> X0 = InData.Spheres[SphereIndex0].LocalPosition;
+			const Chaos::TVector<float, 3> X1 = InData.Spheres[SphereIndex1].LocalPosition;
+			const int32 BoneIndex = InData.Spheres[SphereIndex0].BoneIndex;
+			checkSlow(InData.Spheres[SphereIndex0].BoneIndex == InData.Spheres[SphereIndex1].BoneIndex);
+
+			const Chaos::TVector<float, 3> Center = (X0 + X1) * 0.5;
+			const Chaos::TVector<float, 3> Axis = X1 - X0;
+			const float HalfHeight = Axis.Size() * 0.5;
+
+			// We construct a capsule centered at the origin along the Z axis, and
+			// then move it into place with X and R.
+			CollisionParticles.X(i) = Center;
+//			CollisionParticles.X(i) = Chaos::TVector<float,3>(0.f);
+			CollisionParticles.R(i) = Chaos::TRotation<float, 3>::FromRotatedVector(Chaos::TVector<float,3>::AxisVector(2), Axis.GetSafeNormal());
+//			CollisionParticles.R(i) = Chaos::TRotation<float, 3>();
+			BaseTransforms[i] = Chaos::TRigidTransform<float, 3>(CollisionParticles.X(i), CollisionParticles.R(i));
+
+			CollisionParticles.SetDynamicGeometry(
+				i,
+				MakeUnique<Chaos::TCapsule<float>>(
+					Chaos::TVector<float, 3>(0.f, 0.f, -HalfHeight), // Min
+					Chaos::TVector<float, 3>(0.f, 0.f, HalfHeight), // Max
+					Radius));
+            IndexAndCapsuleCollisionMap.Add(MakePair(i, Connection));
+
+			// Skip spheres added as end caps for the capsule.
+			CapsuleEnds.Add(SphereIndex0);
+			CapsuleEnds.Add(SphereIndex1);
+		}
+	}
+/**/
+	if(InData.Spheres.Num() - CapsuleEnds.Num())
+	{
+        const uint32 Size = CollisionParticles.Size();
+        CollisionParticles.AddParticles(InData.Spheres.Num() - CapsuleEnds.Num());
+        for (uint32 i = 0, j = Size; i < (uint32)InData.Spheres.Num(); ++i)
+        {
+			// Skip spheres that are the end caps of capsules.
+			if (CapsuleEnds.Contains(i))
+				continue;
+			
+			const auto& CollisionSphere = InData.Spheres[i];
+//            CollisionParticles.X(j) = CollisionSphere.LocalPosition;
+            CollisionParticles.X(j) = Chaos::TVector<float, 3>(0.f);
+			CollisionParticles.R(j) = Chaos::TRotation<float, 3>(Chaos::TVector<float, 3>::AxisVector(0), 0.f); //::MakeFromEuler(Chaos::TVector<float, 3>(0.f, 0.f, 0.f));
+            CollisionParticles.SetDynamicGeometry(
+				j, 
+				MakeUnique<Chaos::TSphere<float,3>>(
+					CollisionSphere.LocalPosition, //Chaos::TVector<float, 3>(0.f, 0.f, 0.f), 
+					CollisionSphere.Radius));
+            IndexAndSphereCollisionMap.Add(MakePair(j, CollisionSphere));
+			j++;
+        }
+    }
+	if(InData.Convexes.Num())
+    {
+        const uint32 Size = CollisionParticles.Size();
+        CollisionParticles.AddParticles(InData.Convexes.Num());
+        for (uint32 i = Size; i < CollisionParticles.Size(); ++i)
+        {
+            const auto& convex = InData.Convexes[i - Size];
+            CollisionParticles.X(i) = Chaos::TVector<float, 3>(0.f, 0.f, 0.f);
+            CollisionParticles.R(i) = Chaos::TRotation<float, 3>::MakeFromEuler(Chaos::TVector<float, 3>(0.f, 0.f, 0.f));
+            TArray<TUniquePtr<TImplicitObject<float, 3>>> Planes;
+            for (int32 j = 0; j < convex.Planes.Num(); ++j)
+            {
+                Planes.Add(TUniquePtr<Chaos::TImplicitObject<float, 3>>(
+                    new Chaos::TPlane<float, 3>(Chaos::TVector<float, 3>(0.f, 0.f, convex.Planes[j].W / convex.Planes[j].Z),
+                                                   Chaos::TVector<float, 3>(convex.Planes[j].X, convex.Planes[j].Y, convex.Planes[j].Z))));
+            }
+            CollisionParticles.SetDynamicGeometry(i, MakeUnique<Chaos::TImplicitObjectIntersection<float, 3>>(MoveTemp(Planes)));
+            IndexAndConvexCollisionMap.Add(MakePair(i, convex));
+        }
+    }
+}
+
+void ClothingSimulation::ClearExternalCollisions()
+{
+	TGeometryClothParticles<float, 3>& CollisionParticles = Evolution->CollisionParticles();
+    CollisionParticles.Resize(0);
+
+	IndexAndSphereCollisionMap.Reset();
+    IndexAndCapsuleCollisionMap.Reset();
+    IndexAndConvexCollisionMap.Reset();
+}
+
+void ClothingSimulation::GetCollisions(FClothCollisionData& OutCollisions, bool bIncludeExternal) const
+{
+    OutCollisions.Spheres.Reset();
+    OutCollisions.SphereConnections.Reset();
+    OutCollisions.Convexes.Reset();
+    for (const auto& IndexSphere : IndexAndSphereCollisionMap)
+    {
+        if (Evolution->Collided(IndexSphere.First))
+        {
+            OutCollisions.Spheres.Add(IndexSphere.Second);
+        }
+    }
+	for (const auto& IndexCapsule : IndexAndCapsuleCollisionMap)
+	{
+		if (Evolution->Collided(IndexCapsule.First))
+		{
+			OutCollisions.SphereConnections.Add(IndexCapsule.Second);
+		}
+	}
+    for (const auto& IndexConvex : IndexAndConvexCollisionMap)
+    {
+        if (Evolution->Collided(IndexConvex.First))
+        {
+            OutCollisions.Convexes.Add(IndexConvex.Second);
+        }
+    }
+}
diff --git a/Engine/Plugins/Experimental/ChaosCloth/Source/Chaos/Private/ChaosCloth/ChaosClothingSimulationFactory.cpp b/Engine/Plugins/Experimental/ChaosCloth/Source/Chaos/Private/ChaosCloth/ChaosClothingSimulationFactory.cpp
new file mode 100644
--- /dev/null
+++ b/Engine/Plugins/Experimental/ChaosCloth/Source/Chaos/Private/ChaosCloth/ChaosClothingSimulationFactory.cpp
@@ -0,0 +1,32 @@
+// Copyright 1998-2019 Epic Games, Inc. All Rights Reserved.
+
+#include "ChaosCloth/ChaosClothingSimulationFactory.h"
+#include "ChaosCloth/ChaosClothingSimulation.h"
+
+#include "Assets/ClothingAsset.h"
+
+IClothingSimulation* UChaosClothingSimulationFactory::CreateSimulation()
+{
+	IClothingSimulation* Simulation = new Chaos::ClothingSimulation();
+	return Simulation;
+}
+
+void UChaosClothingSimulationFactory::DestroySimulation(IClothingSimulation* InSimulation)
+{
+    delete InSimulation;
+}
+
+bool UChaosClothingSimulationFactory::SupportsAsset(UClothingAssetBase* InAsset)
+{
+    return true;
+}
+
+bool UChaosClothingSimulationFactory::SupportsRuntimeInteraction()
+{
+    return false;
+}
+
+UClothingSimulationInteractor* UChaosClothingSimulationFactory::CreateInteractor()
+{
+	return nullptr;
+}
\ No newline at end of file
diff --git a/Engine/Plugins/Experimental/ChaosCloth/Source/Chaos/Public/ChaosCloth/ChaosClothingSimulation.h b/Engine/Plugins/Experimental/ChaosCloth/Source/Chaos/Public/ChaosCloth/ChaosClothingSimulation.h
new file mode 100644
--- /dev/null
+++ b/Engine/Plugins/Experimental/ChaosCloth/Source/Chaos/Public/ChaosCloth/ChaosClothingSimulation.h
@@ -0,0 +1,104 @@
+// Copyright 1998-2019 Epic Games, Inc. All Rights Reserved.
+#pragma once
+
+#include "ClothingSimulationInterface.h"
+
+#include "Assets/ClothingAsset.h"
+#include "Chaos/ArrayCollectionArray.h"
+#include "Chaos/PBDEvolution.h"
+#include "Chaos/Transform.h"
+#include "Chaos/TriangleMesh.h"
+#include "Components/SkeletalMeshComponent.h"
+
+namespace Chaos
+{
+	class ClothingSimulationContext : public IClothingSimulationContext
+	{
+	public:
+		ClothingSimulationContext() {}
+		~ClothingSimulationContext() {}
+
+		float DeltaTime;
+		TArray<FMatrix> RefToLocals;
+		TArray<FTransform> BoneTransforms;
+		FTransform ComponentToWorld;
+	};
+
+	class ClothingSimulation : public IClothingSimulation
+	{
+	public:
+		ClothingSimulation();
+		virtual ~ClothingSimulation();
+
+	protected:
+		void Initialize();
+		void CreateActor(USkeletalMeshComponent* InOwnerComponent, UClothingAssetBase* InAsset, int32 SimDataIndex);
+		IClothingSimulationContext* CreateContext() { return new ClothingSimulationContext(); }
+		void FillContext(USkeletalMeshComponent* InComponent, float InDeltaTime, IClothingSimulationContext* InOutContext);
+		void Shutdown() {}
+		bool ShouldSimulate() const { return true; }
+		void Simulate(IClothingSimulationContext* InContext);
+    void DestroyActors() {}
+    void DestroyContext(IClothingSimulationContext* InContext) { delete InContext; }
+		void GetSimulationData(TMap<int32, FClothSimulData>& OutData, USkeletalMeshComponent* InOwnerComponent, USkinnedMeshComponent* InOverrideComponent) const;
+
+		FBoxSphereBounds GetBounds(const USkeletalMeshComponent* InOwnerComponent) const
+		{ return FBoxSphereBounds(Evolution->Particles().X().GetData(), Evolution->Particles().Size()); }
+
+		void AddExternalCollisions(const FClothCollisionData& InData);
+		void ClearExternalCollisions();
+		void GetCollisions(FClothCollisionData& OutCollisions, bool bIncludeExternal = true) const;
+
+	private:
+		// UE Collision Data (Needed only for GetCollisions)
+		TArray<Pair<uint32, FClothCollisionPrim_Sphere>> IndexAndSphereCollisionMap;
+		TArray<Pair<uint32, FClothCollisionPrim_SphereConnection>> IndexAndCapsuleCollisionMap;
+		TArray<Pair<uint32, FClothCollisionPrim_Convex>> IndexAndConvexCollisionMap;
+		// Animation Data
+		UClothingAsset* Asset;
+		TArray<Chaos::TRigidTransform<float, 3>> OldAnimationTransforms;
+		TArray<Chaos::TRigidTransform<float, 3>> AnimationTransforms;
+		TArray<Chaos::TVector<float, 3>> OldAnimationPositions;
+		TArray<Chaos::TVector<float, 3>> AnimationPositions;
+		TArray<Chaos::TVector<float, 3>> AnimationNormals;
+		Chaos::TArrayCollectionArray<int32> BoneIndices;
+		Chaos::TArrayCollectionArray<Chaos::TRigidTransform<float, 3>> BaseTransforms;
+		// Sim Data
+		TArray<Chaos::TVector<uint32, 2>> IndexToRangeMap;
+
+		TArray<TUniquePtr<Chaos::TTriangleMesh<float>>> Meshes;
+		mutable TArray<TArray<Chaos::TVector<float, 3>>> FaceNormals;
+		mutable TArray<TArray<Chaos::TVector<float, 3>>> PointNormals;
+
+		TUniquePtr<Chaos::TPBDEvolution<float, 3>> Evolution;
+		float Time;
+		float DeltaTime;
+		float MaxDeltaTime;
+		float ClampDeltaTime;
+		// Parameters that should be set in the ui
+		int32 NumIterations;
+
+		EClothMassMode MassMode;
+		float UniformMass;
+		float TotalMass;
+		float Density;
+		float MinMass;
+
+		float EdgeStiffness;
+		float BendingStiffness;
+		float AreaStiffness;
+		float VolumeStiffness;
+		float StrainLimitingStiffness;
+		float ShapeTargetStiffness;
+		float SelfCollisionThickness;
+		float CollisionThickness;
+		float CoefficientOfFriction;
+		float Damping;
+		float GravityMagnitude;
+		bool bUseBendingElements;
+		bool bUseTetrahedralConstraints;
+		bool bUseThinShellVolumeConstraints;
+		bool bUseSelfCollisions;
+		bool bUseContinuousCollisionDetection;
+	};
+} // namespace Chaos
\ No newline at end of file
diff --git a/Engine/Plugins/Experimental/ChaosCloth/Source/Chaos/Public/ChaosCloth/ChaosClothingSimulationFactory.h b/Engine/Plugins/Experimental/ChaosCloth/Source/Chaos/Public/ChaosCloth/ChaosClothingSimulationFactory.h
new file mode 100644
--- /dev/null
+++ b/Engine/Plugins/Experimental/ChaosCloth/Source/Chaos/Public/ChaosCloth/ChaosClothingSimulationFactory.h
@@ -0,0 +1,19 @@
+// Copyright 1998-2019 Epic Games, Inc. All Rights Reserved.
+
+#pragma once
+
+#include "ClothingSimulationFactoryInterface.h"
+
+#include "ChaosClothingSimulationFactory.generated.h" // Ryan - really???
+
+UCLASS()
+class CHAOSCLOTH_API UChaosClothingSimulationFactory : public UClothingSimulationFactory
+{
+    GENERATED_BODY()
+  public:
+    virtual IClothingSimulation* CreateSimulation() override;
+    virtual void DestroySimulation(IClothingSimulation* InSimulation) override;
+    virtual bool SupportsAsset(UClothingAssetBase* InAsset) override;
+    virtual bool SupportsRuntimeInteraction() override;
+	virtual UClothingSimulationInteractor* CreateInteractor() override;
+};
diff --git a/Engine/Plugins/Experimental/ChaosNiagara/Source/ChaosNiagara/Classes/NiagaraDataInterfaceChaosDestruction.h b/Engine/Plugins/Experimental/ChaosNiagara/Source/ChaosNiagara/Classes/NiagaraDataInterfaceChaosDestruction.h
--- a/Engine/Plugins/Experimental/ChaosNiagara/Source/ChaosNiagara/Classes/NiagaraDataInterfaceChaosDestruction.h
+++ b/Engine/Plugins/Experimental/ChaosNiagara/Source/ChaosNiagara/Classes/NiagaraDataInterfaceChaosDestruction.h
@@ -9,10 +9,15 @@
 #include "VectorVM.h"
 #include "NiagaraDataInterface.h"
 #include "Chaos/ChaosSolver.h"
+#include "EventsData.h"
 #include "Chaos/ChaosSolverActor.h"
 #include "GeometryCollection/GeometryCollectionActor.h"
 #include "NiagaraDataInterfaceChaosDestruction.generated.h"
 
+#if INCLUDE_CHAOS
+struct PhysicsProxyWrapper;
+#endif
+
 USTRUCT()
 struct FChaosDestructionEvent
 {
@@ -86,7 +91,7 @@ struct FSolverData
 	{}
 
 	TSharedPtr<FPhysScene_Chaos> PhysScene;
-	Chaos::FPBDRigidsSolver* Solver;
+	Chaos::FPhysicsSolver* Solver;
 };
 #endif
 
@@ -528,13 +533,14 @@ public:
 	EDebugTypeEnum DebugType;
 
 	/* ParticleIndex to process collisionData for */
-	UPROPERTY(EditAnywhere, Category = "Debug Settings", meta = (DisplayName = "ParticleIndex to Process"))
-	int32 ParticleIndexToProcess;
+//	UPROPERTY(EditAnywhere, Category = "Debug Settings", meta = (DisplayName = "ParticleIndex to Process"))
+//	TGeometryParticleHandle<float, 3>* ParticleToProcess;
 
 	//----------------------------------------------------------------------------
 	// UObject Interface
 	virtual void PostInitProperties() override;
 	virtual void PostLoad() override;
+	virtual void BeginDestroy() override;
 #if WITH_EDITOR
 	virtual void PostEditChangeProperty(struct FPropertyChangedEvent& PropertyChangedEvent) override;
 #endif
@@ -681,13 +687,7 @@ protected:
 #if INCLUDE_CHAOS
 	void ResetInstData(FNDIChaosDestruction_InstanceData* InstData);
 
-	bool GetAllCollisionsAndMaps(FSolverData SolverData,
-								 TArray<Chaos::TCollisionDataExt<float, 3>>& AllCollisionsArray,
-								 TArray<SolverObjectWrapper>& SolverObjectReverseMappingArray,
-								 TArray<int32>& ParticleIndexReverseMappingArray,
-								 TMap<ISolverObjectBase*, TArray<int32>>& AllCollisionsIndicesBySolverObjectMap,
-								 float& TimeData_MapsCreated);
-
+	void HandleCollisionEvents(const Chaos::FCollisionEventData& Event);
 	void FilterAllCollisions(TArray<Chaos::TCollisionDataExt<float, 3>>& AllCollisionsArray);
 
 	void SortCollisions(TArray<Chaos::TCollisionDataExt<float, 3>>& CollisionsArray);
@@ -701,13 +701,7 @@ protected:
 									  float TimeData_MapsCreated,
 									  int32 IdxSolver);
 	
-	bool GetAllBreakingsAndMaps(FSolverData SolverData,
-								TArray<Chaos::TBreakingDataExt<float, 3>>& AllBreakingsArray,
-								TArray<SolverObjectWrapper>& SolverObjectReverseMappingArray,
-								TArray<int32>& ParticleIndexReverseMappingArray,
-								TMap<ISolverObjectBase*, TArray<int32>>& AllBreakingsIndicesBySolverObjectMap,
-								float& TimeData_MapsCreated);
-
+	void HandleBreakingEvents(const Chaos::FBreakingEventData& Event);
 	void FilterAllBreakings(TArray<Chaos::TBreakingDataExt<float, 3>>& AllBreakingsArray);
 
 	void SortBreakings(TArray<Chaos::TBreakingDataExt<float, 3>>& BreakingsArray);
@@ -721,13 +715,7 @@ protected:
 									 float TimeData_MapsCreated,
 									 int32 IdxSolver);
 	   	  
-	bool GetAllTrailingsAndMaps(FSolverData SolverData,
-								TArray<Chaos::TTrailingDataExt<float, 3>>& AllTrailingsArray,
-								TArray<SolverObjectWrapper>& SolverObjectReverseMappingArray,
-								TArray<int32>& ParticleIndexReverseMappingArray,
-								TMap<ISolverObjectBase*, TArray<int32>>& AllTrailingsIndicesBySolverObjectMap,
-								float& TimeData_MapsCreated);
-
+	void HandleTrailingEvents(const Chaos::FTrailingEventData& Event);
 	void FilterAllTrailings(TArray<Chaos::TTrailingDataExt<float, 3>>& AllTrailingsArray);
 
 	void SortTrailings(TArray<Chaos::TTrailingDataExt<float, 3>>& TrailingsArray);
@@ -768,7 +756,12 @@ protected:
 
 #if INCLUDE_CHAOS
 	TArray<FSolverData> Solvers;
+	TArray<Chaos::TCollisionDataExt<float, 3>> CollisionEvents;
+	TArray<Chaos::TBreakingDataExt<float, 3>> BreakingEvents;
+	TArray<Chaos::TTrailingDataExt<float, 3>> TrailingEvents;
+
 #endif
+
 };
 
 struct FNiagaraDataInterfaceProxyChaosDestruction : public FNiagaraDataInterfaceProxy
diff --git a/Engine/Plugins/Experimental/ChaosNiagara/Source/ChaosNiagara/Private/NiagaraDataInterfaceChaosDestruction.cpp b/Engine/Plugins/Experimental/ChaosNiagara/Source/ChaosNiagara/Private/NiagaraDataInterfaceChaosDestruction.cpp
--- a/Engine/Plugins/Experimental/ChaosNiagara/Source/ChaosNiagara/Private/NiagaraDataInterfaceChaosDestruction.cpp
+++ b/Engine/Plugins/Experimental/ChaosNiagara/Source/ChaosNiagara/Private/NiagaraDataInterfaceChaosDestruction.cpp
@@ -5,10 +5,10 @@
 #include "Misc/FileHelper.h"
 #include "NiagaraShader.h"
 #include "ShaderParameterUtils.h"
-#include "PBDRigidsSolver.h"
+#include "PhysicsSolver.h"
 #include "Niagara/Private/NiagaraStats.h"
 #include "Chaos/PBDCollisionConstraintUtil.h"
-#include "SolverObjects/GeometryCollectionPhysicsObject.h"
+#include "PhysicsProxy/GeometryCollectionPhysicsProxy.h"
 #include "Chaos/PBDCollisionTypes.h"
 #include "PhysicalMaterials/PhysicalMaterial.h"
 #include "NiagaraComponent.h"
@@ -37,14 +37,14 @@ DECLARE_DWORD_COUNTER_STAT(TEXT("NumBreakingsToSpawnParticles"), STAT_NiagaraNum
 DECLARE_DWORD_COUNTER_STAT(TEXT("NumParticlesSpawnedFromCollisions"), STAT_NiagaraNumParticlesSpawnedFromCollisions, STATGROUP_Niagara);
 DECLARE_DWORD_COUNTER_STAT(TEXT("NumParticlesSpawnedFromTrailings"), STAT_NiagaraNumParticlesSpawnedFromTrailings, STATGROUP_Niagara);
 DECLARE_DWORD_COUNTER_STAT(TEXT("NumParticlesSpawnedFromBreaking"), STAT_NiagaraNumParticlesSpawnedFromBreakings, STATGROUP_Niagara);
-DECLARE_MEMORY_STAT(TEXT("SolverObjectReverseMapping"), STAT_SolverObjectReverseMappingMemory, STATGROUP_Niagara);
+DECLARE_MEMORY_STAT(TEXT("PhysicsProxyReverseMapping"), STAT_PhysicsProxyReverseMappingMemory, STATGROUP_Niagara);
 DECLARE_MEMORY_STAT(TEXT("ParticleIndexReverseMapping"), STAT_ParticleIndexReverseMappingMemory, STATGROUP_Niagara);
 DECLARE_MEMORY_STAT(TEXT("AllCollisionsData"), STAT_AllCollisionsDataMemory, STATGROUP_Niagara);
-DECLARE_MEMORY_STAT(TEXT("AllCollisionsIndicesBySolverObject"), STAT_AllCollisionsIndicesBySolverObjectMemory, STATGROUP_Niagara);
+DECLARE_MEMORY_STAT(TEXT("AllCollisionsIndicesByPhysicsProxy"), STAT_AllCollisionsIndicesByPhysicsProxyMemory, STATGROUP_Niagara);
 DECLARE_MEMORY_STAT(TEXT("AllBreakingsData"), STAT_AllBreakingsDataMemory, STATGROUP_Niagara);
-DECLARE_MEMORY_STAT(TEXT("AllBreakingsIndicesBySolverObject"), STAT_AllBreakingsIndicesBySolverObjectMemory, STATGROUP_Niagara);
+DECLARE_MEMORY_STAT(TEXT("AllBreakingsIndicesByPhysicsProxy"), STAT_AllBreakingsIndicesByPhysicsProxyMemory, STATGROUP_Niagara);
 DECLARE_MEMORY_STAT(TEXT("AllTrailingsData"), STAT_AllTrailingsDataMemory, STATGROUP_Niagara);
-DECLARE_MEMORY_STAT(TEXT("AllTrailingsIndicesBySolverObject"), STAT_AllTrailingsIndicesBySolverObjectMemory, STATGROUP_Niagara);
+DECLARE_MEMORY_STAT(TEXT("AllTrailingsIndicesByPhysicsProxy"), STAT_AllTrailingsIndicesByPhysicsProxyMemory, STATGROUP_Niagara);
 
 // Name of all the functions available in the data interface
 static const FName GetPositionName("GetPosition");
@@ -111,7 +111,7 @@ UNiagaraDataInterfaceChaosDestruction::UNiagaraDataInterfaceChaosDestruction(FOb
 	, FinalVelocityMagnitudeMinMax(FVector2D(-1.f, -1.f))
 	, MaxLatency(1.f)
 	, DebugType(EDebugTypeEnum::ChaosNiagara_DebugType_NoDebug)
-	, ParticleIndexToProcess(-1)
+//	, ParticleIndexToProcess(-1)
 	, LastSpawnedPointID(-1)
 	, LastSpawnTime(-1.f)
 	, SolverTime(0.f)
@@ -171,9 +171,32 @@ void UNiagaraDataInterfaceChaosDestruction::PostLoad()
 	LastSpawnedPointID = -1;
 	LastSpawnTime = -1.f;
 	TimeStampOfLastProcessedData = -1.f;
+
+#if WITH_CHAOS
+	FPhysScene* Scene = GetWorld()->GetPhysicsScene();
+	Scene->RegisterEventHandler<Chaos::FCollisionEventData>(Chaos::EEventType::Collision, this, &UNiagaraDataInterfaceChaosDestruction::HandleCollisionEvents);
+	Scene->RegisterEventHandler<Chaos::FBreakingEventData>(Chaos::EEventType::Breaking, this, &UNiagaraDataInterfaceChaosDestruction::HandleBreakingEvents);
+	Scene->RegisterEventHandler<Chaos::FTrailingEventData>(Chaos::EEventType::Trailing, this, &UNiagaraDataInterfaceChaosDestruction::HandleTrailingEvents);
+#endif
+
 	PushToRenderThread();
 }
 
+void UNiagaraDataInterfaceChaosDestruction::BeginDestroy()
+{
+	Super::BeginDestroy();
+
+#if WITH_CHAOS
+	FPhysScene* Scene = GetWorld()->GetPhysicsScene();
+	if (Scene)
+	{
+		Scene->UnregisterEventHandler(Chaos::EEventType::Collision, this);
+		Scene->UnregisterEventHandler(Chaos::EEventType::Breaking, this);
+		Scene->UnregisterEventHandler(Chaos::EEventType::Trailing, this);
+	}
+#endif
+}
+
 #if WITH_EDITOR
 
 void UNiagaraDataInterfaceChaosDestruction::PostEditChangeProperty(struct FPropertyChangedEvent& PropertyChangedEvent)
@@ -313,7 +336,7 @@ bool UNiagaraDataInterfaceChaosDestruction::CopyToInternal(UNiagaraDataInterface
 		DestinationChaosDestruction->FinalVelocityMagnitudeMinMax = FinalVelocityMagnitudeMinMax;
 		DestinationChaosDestruction->MaxLatency = MaxLatency;
 		DestinationChaosDestruction->DebugType = DebugType;
-		DestinationChaosDestruction->ParticleIndexToProcess = ParticleIndexToProcess;
+		//DestinationChaosDestruction->ParticleIndexToProcess = ParticleIndexToProcess;
 		DestinationChaosDestruction->LastSpawnedPointID = LastSpawnedPointID;
 		DestinationChaosDestruction->LastSpawnTime = LastSpawnTime;
 		DestinationChaosDestruction->TimeStampOfLastProcessedData = TimeStampOfLastProcessedData;
@@ -394,8 +417,8 @@ bool UNiagaraDataInterfaceChaosDestruction::Equals(const UNiagaraDataInterface*
 		&& OtherChaosDestruction->VelocityOffsetMax == VelocityOffsetMax
 		&& OtherChaosDestruction->FinalVelocityMagnitudeMinMax == FinalVelocityMagnitudeMinMax
 		&& OtherChaosDestruction->MaxLatency == MaxLatency
-		&& OtherChaosDestruction->DebugType == DebugType
-		&& OtherChaosDestruction->ParticleIndexToProcess == ParticleIndexToProcess;
+		&& OtherChaosDestruction->DebugType == DebugType;
+		//&& OtherChaosDestruction->ParticleIndexToProcess == ParticleIndexToProcess;
 }
 
 int32 UNiagaraDataInterfaceChaosDestruction::PerInstanceDataSize()const
@@ -436,7 +459,7 @@ bool UNiagaraDataInterfaceChaosDestruction::InitPerInstanceData(void* PerInstanc
 		{
 			if (SolverActor)
 			{
-				if (Chaos::FPBDRigidsSolver* Solver = SolverActor->GetSolver())
+				if (Chaos::FPhysicsSolver* Solver = SolverActor->GetSolver())
 				{
 					int32 NewIdx = Solvers.Add(FSolverData());
 
@@ -481,7 +504,7 @@ void UNiagaraDataInterfaceChaosDestruction::DestroyPerInstanceData(void* PerInst
 #if INCLUDE_CHAOS
 void GetMeshExtData(FSolverData SolverData,
 					const int32 ParticleIndex,
-					const TArray<SolverObjectWrapper>& SolverObjectReverseMapping,
+					const TArray<PhysicsProxyWrapper>& PhysicsProxyReverseMapping,
 					const TArray<int32>& ParticleIndexReverseMapping,
 					float& BoundingboxVolume,
 					float& BoundingboxExtentMin,
@@ -499,13 +522,13 @@ void GetMeshExtData(FSolverData SolverData,
 		BoundingboxExtentMax = 100.f;
 		SurfaceType = 0;
 	}
-	else if (SolverObjectReverseMapping[ParticleIndex].Type == ESolverObjectType::GeometryCollectionType)
+	else if (PhysicsProxyReverseMapping[ParticleIndex].Type == EPhysicsProxyType::GeometryCollectionType)
 	{
 		// Since we are touching game objects below, I want to make sure that we're not off in some random thread.
 		ensure(IsInGameThread());
-		if (ISolverObjectBase* SolverObject = SolverObjectReverseMapping[ParticleIndex].SolverObject)
+		if (IPhysicsProxyBase* PhysicsProxy = PhysicsProxyReverseMapping[ParticleIndex].PhysicsProxy)
 		{
-			if (UGeometryCollectionComponent* GeometryCollectionComponent = SolverData.PhysScene->GetOwningComponent<UGeometryCollectionComponent>(SolverObject))
+			if (UGeometryCollectionComponent* GeometryCollectionComponent = SolverData.PhysScene->GetOwningComponent<UGeometryCollectionComponent>(PhysicsProxy))
 			{
 				if (const UGeometryCollection* GeometryCollection = GeometryCollectionComponent->GetRestCollection())
 				{
@@ -551,9 +574,9 @@ void GetMeshExtData(FSolverData SolverData,
 						}
 					}
 				}
-				if (const FGeometryCollectionPhysicsObject* PhysicsObject = GeometryCollectionComponent->GetPhysicsObject())
+				if (const FGeometryCollectionPhysicsProxy* GeomCollectionPhysicsProxy = GeometryCollectionComponent->GetPhysicsProxy())
 				{
-					const FGeometryCollectionResults& PhysResult = PhysicsObject->GetPhysicsResults().GetGameDataForRead();
+					const FGeometryCollectionResults& PhysResult = GeomCollectionPhysicsProxy->GetPhysicsResults().GetGameDataForRead();
 					Transform = PhysResult.ParticleToWorldTransforms[ParticleIndex - PhysResult.BaseIndex];
 				}
 			}
@@ -563,7 +586,7 @@ void GetMeshExtData(FSolverData SolverData,
 
 void GetMesPhysicalData(FSolverData SolverData,
 						const int32 ParticleIndex,
-						const TArray<SolverObjectWrapper>& SolverObjectReverseMapping,
+						const TArray<PhysicsProxyWrapper>& PhysicsProxyReverseMapping,
 						const TArray<int32>& ParticleIndexReverseMapping,
 						FLinearColor& Color,
 						float& Friction,
@@ -582,14 +605,14 @@ void GetMesPhysicalData(FSolverData SolverData,
 		Restitution = 0.3f;
 		Density = 1.0f;
 	}
-	else if (SolverObjectReverseMapping[ParticleIndex].Type == ESolverObjectType::GeometryCollectionType)
+	else if (PhysicsProxyReverseMapping[ParticleIndex].Type == EPhysicsProxyType::GeometryCollectionType)
 	{
 		// Since we are touching game objects below, I want to make sure that we're not off in some random thread.
 		ensure(IsInGameThread());
 
-		if (ISolverObjectBase* SolverObject = SolverObjectReverseMapping[ParticleIndex].SolverObject)
+		if (IPhysicsProxyBase* PhysicsProxy = PhysicsProxyReverseMapping[ParticleIndex].PhysicsProxy)
 		{
-			if (UGeometryCollectionComponent* GeometryCollectionComponent = SolverData.PhysScene->GetOwningComponent<UGeometryCollectionComponent>(SolverObject))
+			if (UGeometryCollectionComponent* GeometryCollectionComponent = SolverData.PhysScene->GetOwningComponent<UGeometryCollectionComponent>(PhysicsProxy))
 			{
 				if (const UGeometryCollection* GeometryCollection = GeometryCollectionComponent->GetRestCollection())
 				{
@@ -656,118 +679,46 @@ void GetMesPhysicalData(FSolverData SolverData,
 #endif
 
 #if INCLUDE_CHAOS
-bool UNiagaraDataInterfaceChaosDestruction::GetAllCollisionsAndMaps(FSolverData SolverData,
-																	TArray<Chaos::TCollisionDataExt<float, 3>>& AllCollisionsArray,
-																	TArray<SolverObjectWrapper>& SolverObjectReverseMappingArray,
-																	TArray<int32>& ParticleIndexReverseMappingArray,
-																	TMap<ISolverObjectBase*, TArray<int32>>& AllCollisionsIndicesBySolverObjectMap,
-																	float& TimeData_MapsCreated)
-{
-	QUICK_SCOPE_CYCLE_COUNTER(STAT_GetAllCollisionsAndMaps);
-
-	// ----------------------------------------------------------------------------------------------------------
-	// GETTING DATA FROM PBDRIGIDSOLVER
-	// ----------------------------------------------------------------------------------------------------------
-	Chaos::FPBDRigidsSolver::FScopedGetEventsData ScopedAccess = SolverData.Solver->ScopedGetEventsData();
-
-	const Chaos::FPBDRigidsSolver::FAllCollisionDataMaps& AllCollisionData_Maps = ScopedAccess.GetAllCollisions_Maps();
-	if (AllCollisionData_Maps.IsValid())
-	{
-		if ((*AllCollisionData_Maps.AllCollisionData).TimeCreated == TimeStampOfLastProcessedData)
-		{
-			return false;
-		}
-
-		if (bGetExternalCollisionData)
-		{
-			// NOTE: TODO: These arrays and maps are not used due to GetMeshExtData and GetMesPhysicalData not currently being used due perf concerns
-			if (AllCollisionData_Maps.SolverObjectReverseMapping)
-			{
-				SolverObjectReverseMappingArray = (*AllCollisionData_Maps.SolverObjectReverseMapping).SolverObjectReverseMappingArray;
-			}
-			else
-			{
-				return false;
-			}
-			if (AllCollisionData_Maps.ParticleIndexReverseMapping)
-			{
-				ParticleIndexReverseMappingArray = (*AllCollisionData_Maps.ParticleIndexReverseMapping).ParticleIndexReverseMappingArray;
-			}
-			else
-			{
-				return false;
-			}
-			if (AllCollisionData_Maps.AllCollisionsIndicesBySolverObject)
-			{
-				AllCollisionsIndicesBySolverObjectMap = (*AllCollisionData_Maps.AllCollisionsIndicesBySolverObject).AllCollisionsIndicesBySolverObjectMap;
-			}
-			else
-			{
-				return false;
-			}
-		}
-		
-		if (AllCollisionData_Maps.AllCollisionData)
-		{
-			// Copy data from *AllCollisionData_Maps.AllCollisionData into AllCollisionsArray
-			// Also get Boundingbox related data and SurfaceType and save it as well
-			AllCollisionsArray.AddUninitialized((*AllCollisionData_Maps.AllCollisionData).AllCollisionsArray.Num());
 
-			int32 Idx = 0;
-			for (Chaos::TCollisionDataExt<float, 3>& Collision : AllCollisionsArray)
-			{
-				Collision = (*AllCollisionData_Maps.AllCollisionData).AllCollisionsArray[Idx];
-
-				// #GM: Disable this for now for perf
-				/*
-				GetMeshExtData(SolverData,
-					AllCollisionsArray[Idx].ParticleIndexMesh == INDEX_NONE ? AllCollisionsArray[Idx].ParticleIndex : AllCollisionsArray[Idx].ParticleIndexMesh,
-					SolverObjectReverseMappingArray,
-					ParticleIndexReverseMappingArray,
-					AllCollisionsArray[Idx].BoundingboxVolume,
-					AllCollisionsArray[Idx].BoundingboxExtentMin,
-					AllCollisionsArray[Idx].BoundingboxExtentMax,
-					AllCollisionsArray[Idx].SurfaceType);
-				*/
-				AllCollisionsArray[Idx].BoundingboxVolume = 1000000.f;
-				AllCollisionsArray[Idx].BoundingboxExtentMin = 100.f;
-				AllCollisionsArray[Idx].BoundingboxExtentMax = 100.f;
-				AllCollisionsArray[Idx].SurfaceType = 0;
-
-				Idx++;
-			}
-		}
-		else
-		{
-			return false;
-		}
-	}
-	else
-	{
-		return false;
+void UNiagaraDataInterfaceChaosDestruction::HandleCollisionEvents(const Chaos::FCollisionEventData& Event)
+{
+	ensure(IsInGameThread());
+	
+	// Copy data from Event into AllCollisionsArray
+	// Also get Boundingbox related data and SurfaceType and save it as well
+	CollisionEvents.AddUninitialized(Event.CollisionData.AllCollisionsArray.Num());
+
+	int32 Idx = 0;
+	for (Chaos::TCollisionDataExt<float, 3>& Collision : CollisionEvents)
+	{
+		Collision = Event.CollisionData.AllCollisionsArray[Idx];
+
+		// #GM: Disable this for now for perf
+		/*
+		GetMeshExtData(SolverData,
+			AllCollisionsArray[Idx].ParticleIndexMesh == INDEX_NONE ? AllCollisionsArray[Idx].ParticleIndex : AllCollisionsArray[Idx].ParticleIndexMesh,
+			PhysicsProxyReverseMappingArray,
+			ParticleIndexReverseMappingArray,
+			AllCollisionsArray[Idx].BoundingboxVolume,
+			AllCollisionsArray[Idx].BoundingboxExtentMin,
+			AllCollisionsArray[Idx].BoundingboxExtentMax,
+			AllCollisionsArray[Idx].SurfaceType);
+		*/
+		CollisionEvents[Idx].BoundingboxVolume = 1000000.f;
+		CollisionEvents[Idx].BoundingboxExtentMin = 100.f;
+		CollisionEvents[Idx].BoundingboxExtentMax = 100.f;
+		CollisionEvents[Idx].SurfaceType = 0;
+
+		Idx++;
 	}
-
-	TimeData_MapsCreated = (*AllCollisionData_Maps.AllCollisionData).TimeCreated;
-	TimeStampOfLastProcessedData = TimeData_MapsCreated;
-	// ----------------------------------------------------------------------------------------------------------
-	// END OF GETTING DATA FROM PBDRIGIDSOLVER
-	// ----------------------------------------------------------------------------------------------------------
-
-	// Option one:
-	// - deep copy the data and release after deep copy
-	//
-	// Option two:
-	// - reference the data and release it at the end of processing it
-	// ----------------------------------------------------------------------------------------------------------
-
-	return true;
 }
 
+
 void UNiagaraDataInterfaceChaosDestruction::FilterAllCollisions(TArray<Chaos::TCollisionDataExt<float, 3>>& AllCollisionsArray)
 {
 	QUICK_SCOPE_CYCLE_COUNTER(STAT_FilterAllCollisions);
 
-	if (ParticleIndexToProcess != -1 ||
+	if (/*ParticleToProcess != nullptr ||*/
 		ImpulseToSpawnMinMax.X > 0.f ||
 		ImpulseToSpawnMinMax.Y > 0.f ||
 		SpeedToSpawnMinMax.X > 0.f ||
@@ -802,7 +753,7 @@ void UNiagaraDataInterfaceChaosDestruction::FilterAllCollisions(TArray<Chaos::TC
 			float CollisionAccumulatedImpulseSquared = AllCollisionsArray[IdxCollision].AccumulatedImpulse.SizeSquared();
 			float CollisionSpeedSquared = AllCollisionsArray[IdxCollision].Velocity1.SizeSquared();
 
-			if ((ParticleIndexToProcess != -1 && AllCollisionsArray[IdxCollision].ParticleIndex != ParticleIndexToProcess) ||
+			if (/*(ParticleToProcess != nullptr && AllCollisionsArrayInOut[IdxCollision].Particle != ParticleToProcess) ||*/
 				(ImpulseToSpawnMinMax.X > 0.f && ImpulseToSpawnMinMax.Y < 0.f && CollisionAccumulatedImpulseSquared < MinImpulseToSpawnSquared) ||
 				(ImpulseToSpawnMinMax.X < 0.f && ImpulseToSpawnMinMax.Y > 0.f && CollisionAccumulatedImpulseSquared > MaxImpulseToSpawnSquared) ||
 				(ImpulseToSpawnMinMax.X > 0.f && ImpulseToSpawnMinMax.Y > 0.f && (CollisionAccumulatedImpulseSquared < MinImpulseToSpawnSquared || CollisionAccumulatedImpulseSquared > MaxImpulseToSpawnSquared)) ||
@@ -1111,7 +1062,7 @@ int32 UNiagaraDataInterfaceChaosDestruction::SpawnParticlesFromCollision(FSolver
 			}
 			else if (DebugType == EDebugTypeEnum::ChaosNiagara_DebugType_ColorByParticleIndex)
 			{
-				ParticleColor = ColorArray[Collision.ParticleIndex % ColorArray.Num()];
+				//ParticleColor = ColorArray[Collision.ParticleIndex % ColorArray.Num()];
 			}
 
 			// Store principal data
@@ -1139,7 +1090,6 @@ int32 UNiagaraDataInterfaceChaosDestruction::SpawnParticlesFromCollision(FSolver
 	
 		return NumParticles;
 	}
-
 	return 0;
 }
 
@@ -1148,40 +1098,16 @@ bool UNiagaraDataInterfaceChaosDestruction::CollisionCallback(FNDIChaosDestructi
 	int32 IdxSolver = 0;
 	for (FSolverData SolverData : Solvers)
 	{
-		if (SolverData.Solver->GetGenerateCollisionData() && SolverData.Solver->GetSolverTime() > 0.f && MaxNumberOfDataEntriesToSpawn > 0)
+		if (SolverData.Solver->GetEventFilters()->IsCollisionEventEnabled() && CollisionEvents.Num() > 0 && SolverData.Solver->GetSolverTime() > 0.f && MaxNumberOfDataEntriesToSpawn > 0)
 		{
-			TArray<Chaos::TCollisionDataExt<float, 3>> AllCollisionsArray;
-			TArray<SolverObjectWrapper> SolverObjectReverseMappingArray;
-			TArray<int32> ParticleIndexReverseMappingArray;
-			TMap<ISolverObjectBase*, TArray<int32>> AllCollisionsIndicesBySolverObjectMap;
-			float TimeData_MapsCreated;
-
-			if (!GetAllCollisionsAndMaps(SolverData,
-										 AllCollisionsArray,
-										 SolverObjectReverseMappingArray,
-										 ParticleIndexReverseMappingArray,
-										 AllCollisionsIndicesBySolverObjectMap,
-										 TimeData_MapsCreated))
-			{
-				return false;
-			}
+			TArray<Chaos::TCollisionDataExt<float, 3>>& AllCollisionsArray = CollisionEvents;
+			float TimeData_MapsCreated = 0.0f;
 
 #if STATS
 			{
 				QUICK_SCOPE_CYCLE_COUNTER(STAT_GatherMemoryStats);
 				size_t SizeOfAllCollisions = sizeof(Chaos::TCollisionData<float, 3>) * AllCollisionsArray.Num();
-				size_t SizeOfSolverObjectReverseMapping = sizeof(SolverObjectWrapper) * SolverObjectReverseMappingArray.Num();
-				size_t SizeOfParticleIndexReverseMapping = sizeof(int32) * ParticleIndexReverseMappingArray.Num();
-
-				size_t SizeOfAllCollisionsIndicesBySolverObject = 0;
-				for (auto& Elem : AllCollisionsIndicesBySolverObjectMap)
-				{
-					SizeOfAllCollisionsIndicesBySolverObject += sizeof(int32) * (Elem.Value).Num();
-				}
 				SET_MEMORY_STAT(STAT_AllCollisionsDataMemory, SizeOfAllCollisions);
-				SET_MEMORY_STAT(STAT_SolverObjectReverseMappingMemory, SizeOfSolverObjectReverseMapping);
-				SET_MEMORY_STAT(STAT_ParticleIndexReverseMappingMemory, SizeOfParticleIndexReverseMapping);
-				SET_MEMORY_STAT(STAT_AllCollisionsIndicesBySolverObjectMemory, SizeOfAllCollisionsIndicesBySolverObject);
 			}
 
 			INC_DWORD_STAT_BY(STAT_NiagaraNumAllCollisions, AllCollisionsArray.Num());
@@ -1220,7 +1146,7 @@ bool UNiagaraDataInterfaceChaosDestruction::CollisionCallback(FNDIChaosDestructi
 						/*
 						GetMesPhysicalData(SolverData,
 							CollisionsToSpawnArray[IdxCollision].ParticleIndexMesh == INDEX_NONE ? CollisionsToSpawnArray[IdxCollision].ParticleIndex : CollisionsToSpawnArray[IdxCollision].ParticleIndexMesh,
-							SolverObjectReverseMappingArray,
+							PhysicsProxyReverseMappingArray,
 							ParticleIndexReverseMappingArray,
 							Color,
 							Friction,
@@ -1250,145 +1176,71 @@ bool UNiagaraDataInterfaceChaosDestruction::CollisionCallback(FNDIChaosDestructi
 	return false;
 }
 
-bool UNiagaraDataInterfaceChaosDestruction::GetAllBreakingsAndMaps(FSolverData SolverData,
-																   TArray<Chaos::TBreakingDataExt<float, 3>>& AllBreakingsArray,
-																   TArray<SolverObjectWrapper>& SolverObjectReverseMappingArray,
-																   TArray<int32>& ParticleIndexReverseMappingArray,
-																   TMap<ISolverObjectBase*, TArray<int32>>& AllBreakingIndicesBySolverObjectMap,
-																   float& TimeData_MapsCreated)
+void UNiagaraDataInterfaceChaosDestruction::HandleBreakingEvents(const Chaos::FBreakingEventData& Event)
 {
-	// ----------------------------------------------------------------------------------------------------------
-	// GETTING DATA FROM PBDRIGIDSOLVER
-	// ----------------------------------------------------------------------------------------------------------
-	Chaos::FPBDRigidsSolver::FScopedGetEventsData ScopedAccess = SolverData.Solver->ScopedGetEventsData();
-	Chaos::FPBDRigidsSolver::FAllBreakingDataMaps AllBreakingData_Maps = ScopedAccess.GetAllBreakings_Maps();
-	if (AllBreakingData_Maps.IsValid())
-	{
-		if (bGetExternalBreakingData) 
-		{
-			if ((*AllBreakingData_Maps.AllBreakingData).TimeCreated == TimeStampOfLastProcessedData)
-			{
-				return false;
-			}
+	ensure(IsInGameThread());
 
-			if (AllBreakingData_Maps.SolverObjectReverseMapping)
-			{
-				SolverObjectReverseMappingArray = (*AllBreakingData_Maps.SolverObjectReverseMapping).SolverObjectReverseMappingArray;
-			}
-			else
-			{
-				return false;
-			}
-			if (AllBreakingData_Maps.ParticleIndexReverseMapping)
-			{
-				ParticleIndexReverseMappingArray = (*AllBreakingData_Maps.ParticleIndexReverseMapping).ParticleIndexReverseMappingArray;
-			}
-			else
-			{
-				return false;
-			}
-			if (AllBreakingData_Maps.AllBreakingsIndicesBySolverObject)
-			{
-				AllBreakingIndicesBySolverObjectMap = (*AllBreakingData_Maps.AllBreakingsIndicesBySolverObject).AllBreakingsIndicesBySolverObjectMap;
-			}
-			else
-			{
-				return false;
-			}
-		}
-			
-		if (AllBreakingData_Maps.AllBreakingData)
-		{
-			ensure(IsInGameThread());
-			
-			// Copy data from *AllBreakingData_Maps.AllBreakingData into AllBreakingsArray
-			// Also get Boundingbox related data and SurfaceType and save it as well
-			AllBreakingsArray.InsertZeroed(0, (*AllBreakingData_Maps.AllBreakingData).AllBreakingsArray.Num());
-			//UE_LOG(LogScript, Warning, TEXT("(*AllBreakingData_Maps.AllBreakingData).AllBreakingsArray.Num() = %d"), (*AllBreakingData_Maps.AllBreakingData).AllBreakingsArray.Num());
-			int32 Idx = 0;
-			for (Chaos::TBreakingDataExt<float, 3>& Breaking : AllBreakingsArray)
-			{	
-				if (!ensure(AllBreakingsArray.Num() == (*AllBreakingData_Maps.AllBreakingData).AllBreakingsArray.Num())) {
-					continue;
-				}
-				Breaking = (*AllBreakingData_Maps.AllBreakingData).AllBreakingsArray[Idx];
+	// Copy data from *AllBreakingData_Maps.AllBreakingData into AllBreakingsArray
+	// Also get Boundingbox related data and SurfaceType and save it as well
+	BreakingEvents.InsertZeroed(0, Event.BreakingData.AllBreakingsArray.Num());
+	//UE_LOG(LogScript, Warning, TEXT("(*AllBreakingData_Maps.AllBreakingData).AllBreakingsArray.Num() = %d"), (*AllBreakingData_Maps.AllBreakingData).AllBreakingsArray.Num());
+	int32 Idx = 0;
+	for (Chaos::TBreakingDataExt<float, 3>& Breaking : BreakingEvents)
+	{
+		Breaking = Event.BreakingData.AllBreakingsArray[Idx];
 
-				// #GM: Disable this for now for perf
-				// TODO(mv): Temporarily re-enable this for Jon to get materials, transform and bounding box. Will optimize in the coming week. 
-				if (bGetExternalBreakingData)
-				{
-					Chaos::TRigidTransform<float, 3> Transform;
-					UPhysicalMaterial* Material = nullptr;
-					GetMeshExtData(SolverData,
-						AllBreakingsArray[Idx].ParticleIndexMesh == INDEX_NONE ? AllBreakingsArray[Idx].ParticleIndex : AllBreakingsArray[Idx].ParticleIndexMesh,
-						SolverObjectReverseMappingArray,
-						ParticleIndexReverseMappingArray,
-						AllBreakingsArray[Idx].BoundingboxVolume,
-						AllBreakingsArray[Idx].BoundingboxExtentMin,
-						AllBreakingsArray[Idx].BoundingboxExtentMax,
-						AllBreakingsArray[Idx].BoundingBox,
-						AllBreakingsArray[Idx].SurfaceType,
-						Transform,
-						Material);
-					//UE_LOG(LogScript, Warning, TEXT("GetAllBreakingsAndMaps[%d]: SurfaceType = %d"), Idx, AllBreakingsArray[Idx].SurfaceType);
-					AllBreakingsArray[Idx].TransformTranslation = Transform.GetTranslation();
-					AllBreakingsArray[Idx].TransformRotation = Transform.GetRotation();
-					AllBreakingsArray[Idx].TransformScale = Transform.GetScale3D();
-					if (Material)
-					{
-						AllBreakingsArray[Idx].PhysicalMaterialName = Material->GetFName();
-					}
-					else
-					{
-						AllBreakingsArray[Idx].PhysicalMaterialName = FName();
-					}
-				}
-				else
-				{
-					AllBreakingsArray[Idx].BoundingboxVolume = 1000000.f;
-					AllBreakingsArray[Idx].BoundingboxExtentMin = 100.0f;
-					AllBreakingsArray[Idx].BoundingboxExtentMax = 100.0f;
-					AllBreakingsArray[Idx].BoundingBox = FBox(FVector(-100.0f, -100.0f, -100.0f), FVector(100.0f, 100.0f, 100.0f));
-					AllBreakingsArray[Idx].SurfaceType = 0;
-					AllBreakingsArray[Idx].TransformTranslation = FVector(0.0f, 0.0f, 0.0f);
-					AllBreakingsArray[Idx].TransformRotation = FQuat(0.0f, 0.0f, 0.0f, 1.0f);
-					AllBreakingsArray[Idx].TransformScale = FVector(1.0f, 1.0f, 1.0f);
-					AllBreakingsArray[Idx].PhysicalMaterialName = FName();
-				}
-				
-				Idx++;
-			}
+		// #GM: Disable this for now for perf
+		// TODO(mv): Temporarily re-enable this for Jon to get materials, transform and bounding box. Will optimize in the coming week. 
+		if (bGetExternalBreakingData)
+		{
+			//Chaos::TRigidTransform<float, 3> Transform;
+			//UPhysicalMaterial* Material = nullptr;
+			//GetMeshExtData(SolverData,
+			//	AllBreakingsArray[Idx].ParticleIndexMesh == INDEX_NONE ? AllBreakingsArray[Idx].ParticleIndex : AllBreakingsArray[Idx].ParticleIndexMesh,
+			//	PhysicsProxyReverseMappingArray,
+			//	ParticleIndexReverseMappingArray,
+			//	AllBreakingsArray[Idx].BoundingboxVolume,
+			//	AllBreakingsArray[Idx].BoundingboxExtentMin,
+			//	AllBreakingsArray[Idx].BoundingboxExtentMax,
+			//	AllBreakingsArray[Idx].BoundingBox,
+			//	AllBreakingsArray[Idx].SurfaceType,
+			//	Transform,
+			//	Material);
+			////UE_LOG(LogScript, Warning, TEXT("GetAllBreakingsAndMaps[%d]: SurfaceType = %d"), Idx, AllBreakingsArray[Idx].SurfaceType);
+			//AllBreakingsArray[Idx].TransformTranslation = Transform.GetTranslation();
+			//AllBreakingsArray[Idx].TransformRotation = Transform.GetRotation();
+			//AllBreakingsArray[Idx].TransformScale = Transform.GetScale3D();
+			//if (Material)
+			//{
+			//	AllBreakingsArray[Idx].PhysicalMaterialName = Material->GetFName();
+			//}
+			//else
+			//{
+			//	AllBreakingsArray[Idx].PhysicalMaterialName = FName();
+			//}
 		}
 		else
 		{
-			return false;
+			BreakingEvents[Idx].BoundingboxVolume = 1000000.f;
+			BreakingEvents[Idx].BoundingboxExtentMin = 100.0f;
+			BreakingEvents[Idx].BoundingboxExtentMax = 100.0f;
+			BreakingEvents[Idx].BoundingBox = FBox(FVector(-100.0f, -100.0f, -100.0f), FVector(100.0f, 100.0f, 100.0f));
+			BreakingEvents[Idx].SurfaceType = 0;
+			BreakingEvents[Idx].TransformTranslation = FVector(0.0f, 0.0f, 0.0f);
+			BreakingEvents[Idx].TransformRotation = FQuat(0.0f, 0.0f, 0.0f, 1.0f);
+			BreakingEvents[Idx].TransformScale = FVector(1.0f, 1.0f, 1.0f);
+			BreakingEvents[Idx].PhysicalMaterialName = FName();
 		}
-	}
-	else
-	{
-		return false;
-	}
-
-	TimeData_MapsCreated = (*AllBreakingData_Maps.AllBreakingData).TimeCreated;
-	TimeStampOfLastProcessedData = TimeData_MapsCreated;
-	// ----------------------------------------------------------------------------------------------------------
-	// END OF GETTING DATA FROM PBDRIGIDSOLVER
-	// ----------------------------------------------------------------------------------------------------------
 
-	// Option one:
-	// - deep copy the data and release after deep copy
-	//
-	// Option two:
-	// - reference the data and release it at the end of processing it
-	// ----------------------------------------------------------------------------------------------------------
-
-	return true;
+		Idx++;
+	}
 }
 
+
 void UNiagaraDataInterfaceChaosDestruction::FilterAllBreakings(TArray<Chaos::TBreakingDataExt<float, 3>>& AllBreakingsArray)
 {
 	if (bApplyMaterialsFilter || 
-		ParticleIndexToProcess != -1 ||
+	//	ParticleIndexToProcess != -1 ||
 		SpeedToSpawnMinMax.X > 0.f ||
 		SpeedToSpawnMinMax.Y > 0.f ||
 		MassToSpawnMinMax.X > 0.f ||
@@ -1441,7 +1293,7 @@ void UNiagaraDataInterfaceChaosDestruction::FilterAllBreakings(TArray<Chaos::TBr
 			float BreakingSpeedSquared = AllBreakingsArray[IdxBreaking].Velocity.SizeSquared();
 
 			if (!(bApplyMaterialsFilter && IsMaterialInFilter(AllBreakingsArray[IdxBreaking].PhysicalMaterialName)) ||
-				(ParticleIndexToProcess != -1 && AllBreakingsArray[IdxBreaking].ParticleIndex != ParticleIndexToProcess) ||
+				//(ParticleIndexToProcess != -1 && AllBreakingsArray[IdxBreaking].ParticleIndex != ParticleIndexToProcess) ||
 				(SpeedToSpawnMinMax.X > 0.f && SpeedToSpawnMinMax.Y < 0.f && BreakingSpeedSquared < MinSpeedToSpawnSquared) ||
 				(SpeedToSpawnMinMax.X < 0.f && SpeedToSpawnMinMax.Y > 0.f && BreakingSpeedSquared > MaxSpeedToSpawnSquared) ||
 				(SpeedToSpawnMinMax.X > 0.f && SpeedToSpawnMinMax.Y > 0.f && (BreakingSpeedSquared < MinSpeedToSpawnSquared || BreakingSpeedSquared > MaxSpeedToSpawnSquared)) ||
@@ -1774,38 +1626,28 @@ bool UNiagaraDataInterfaceChaosDestruction::BreakingCallback(FNDIChaosDestructio
 	int32 IdxSolver = 0;
 	for (FSolverData SolverData : Solvers)
 	{
-		if (SolverData.Solver->GetGenerateBreakingData() && SolverData.Solver->GetSolverTime() > 0.f && MaxNumberOfDataEntriesToSpawn > 0)
+		if (SolverData.Solver->GetEventFilters()->IsBreakingEventEnabled() && BreakingEvents.Num() > 0 && SolverData.Solver->GetSolverTime() > 0.f && MaxNumberOfDataEntriesToSpawn > 0)
 		{
-			TArray<Chaos::TBreakingDataExt<float, 3>> AllBreakingsArray;
-			TArray<SolverObjectWrapper> SolverObjectReverseMappingArray;
+			TArray<Chaos::TBreakingDataExt<float, 3>>& AllBreakingsArray = BreakingEvents;
+			TArray<PhysicsProxyWrapper> PhysicsProxyReverseMappingArray;
 			TArray<int32> ParticleIndexReverseMappingArray;
-			TMap<ISolverObjectBase*, TArray<int32>> AllBreakingsIndicesBySolverObjectMap;
-			float TimeData_MapsCreated;
-
-			if (!GetAllBreakingsAndMaps(SolverData,
-										AllBreakingsArray,
-										SolverObjectReverseMappingArray,
-										ParticleIndexReverseMappingArray,
-										AllBreakingsIndicesBySolverObjectMap,
-										TimeData_MapsCreated))
-			{
-				return false;
-			}
+			TMap<IPhysicsProxyBase*, TArray<int32>> AllBreakingsIndicesByPhysicsProxyMap;
+			float TimeData_MapsCreated = 0.0f;
 
 			{
 				size_t SizeOfAllBreakings = sizeof(Chaos::TBreakingData<float, 3>) * AllBreakingsArray.Num();
-				size_t SizeOfSolverObjectReverseMapping = sizeof(SolverObjectWrapper) * SolverObjectReverseMappingArray.Num();
+				size_t SizeOfPhysicsProxyReverseMapping = sizeof(PhysicsProxyWrapper) * PhysicsProxyReverseMappingArray.Num();
 				size_t SizeOfParticleIndexReverseMapping = sizeof(int32) * ParticleIndexReverseMappingArray.Num();
 
-				size_t SizeOfAllBreakingsIndicesBySolverObject = 0;
-				for (auto& Elem : AllBreakingsIndicesBySolverObjectMap)
+				size_t SizeOfAllBreakingsIndicesByPhysicsProxy = 0;
+				for (auto& Elem : AllBreakingsIndicesByPhysicsProxyMap)
 				{
-					SizeOfAllBreakingsIndicesBySolverObject += sizeof(int32) * (Elem.Value).Num();
+					SizeOfAllBreakingsIndicesByPhysicsProxy += sizeof(int32) * (Elem.Value).Num();
 				}
 				SET_MEMORY_STAT(STAT_AllBreakingsDataMemory, SizeOfAllBreakings);
-				SET_MEMORY_STAT(STAT_SolverObjectReverseMappingMemory, SizeOfSolverObjectReverseMapping);
+				SET_MEMORY_STAT(STAT_PhysicsProxyReverseMappingMemory, SizeOfPhysicsProxyReverseMapping);
 				SET_MEMORY_STAT(STAT_ParticleIndexReverseMappingMemory, SizeOfParticleIndexReverseMapping);
-				SET_MEMORY_STAT(STAT_AllBreakingsIndicesBySolverObjectMemory, SizeOfAllBreakingsIndicesBySolverObject);
+				SET_MEMORY_STAT(STAT_AllBreakingsIndicesByPhysicsProxyMemory, SizeOfAllBreakingsIndicesByPhysicsProxy);
 			}
 
 			INC_DWORD_STAT_BY(STAT_NiagaraNumAllBreakings, AllBreakingsArray.Num());
@@ -1845,7 +1687,7 @@ bool UNiagaraDataInterfaceChaosDestruction::BreakingCallback(FNDIChaosDestructio
 							// #GM: Disable this for now for perf
 							GetMesPhysicalData(SolverData,
 											   BreakingsToSpawnArray[IdxBreaking].ParticleIndexMesh == INDEX_NONE ? BreakingsToSpawnArray[IdxBreaking].ParticleIndex : BreakingsToSpawnArray[IdxBreaking].ParticleIndexMesh,
-											   SolverObjectReverseMappingArray,
+											   PhysicsProxyReverseMappingArray,
 											   ParticleIndexReverseMappingArray,
 											   Color,
 											   Friction,
@@ -1878,106 +1720,42 @@ bool UNiagaraDataInterfaceChaosDestruction::BreakingCallback(FNDIChaosDestructio
 
 }
 
-bool UNiagaraDataInterfaceChaosDestruction::GetAllTrailingsAndMaps(FSolverData SolverData,
-																   TArray<Chaos::TTrailingDataExt<float, 3>>& AllTrailingsArray,
-																   TArray<SolverObjectWrapper>& SolverObjectReverseMappingArray,
-																   TArray<int32>& ParticleIndexReverseMappingArray,
-																   TMap<ISolverObjectBase*, TArray<int32>>& AllTrailingsIndicesBySolverObjectMap,
-																   float& TimeData_MapsCreated)
+void UNiagaraDataInterfaceChaosDestruction::HandleTrailingEvents(const Chaos::FTrailingEventData& Event)
 {
-	// ----------------------------------------------------------------------------------------------------------
-	// GETTING DATA FROM PBDRIGIDSOLVER
-	// ----------------------------------------------------------------------------------------------------------
-	Chaos::FPBDRigidsSolver::FScopedGetEventsData ScopedAccess = SolverData.Solver->ScopedGetEventsData();
-	Chaos::FPBDRigidsSolver::FAllTrailingDataMaps AllTrailingData_Maps = ScopedAccess.GetAllTrailings_Maps();
-
-	if (AllTrailingData_Maps.IsValid())
-	{
-		if (bGetExternalTrailingData) 
-		{
-			if ((*AllTrailingData_Maps.AllTrailingData).TimeCreated == TimeStampOfLastProcessedData)
-			{
-				return false;
-			}
-
-			if (AllTrailingData_Maps.SolverObjectReverseMapping)
-			{
-				SolverObjectReverseMappingArray = (*AllTrailingData_Maps.SolverObjectReverseMapping).SolverObjectReverseMappingArray;
-			}
-			else
-			{
-				return false;
-			}
-			if (AllTrailingData_Maps.ParticleIndexReverseMapping)
-			{
-				ParticleIndexReverseMappingArray = (*AllTrailingData_Maps.ParticleIndexReverseMapping).ParticleIndexReverseMappingArray;
-			}
-			else
-			{
-				return false;
-			}
-			if (AllTrailingData_Maps.AllTrailingsIndicesBySolverObject)
-			{
-				AllTrailingsIndicesBySolverObjectMap = (*AllTrailingData_Maps.AllTrailingsIndicesBySolverObject).AllTrailingsIndicesBySolverObjectMap;
-			}
-			else
-			{
-				return false;
-			}
-		}
-
-		if (AllTrailingData_Maps.AllTrailingData)
-		{
-			// Copy data from *AllTrailingData_Maps.AllTrailingData into AllTrailingsArray
-			// Also get Boundingbox related data and SurfaceType and save it as well
-			AllTrailingsArray.AddUninitialized((*AllTrailingData_Maps.AllTrailingData).AllTrailingsArray.Num());
-
-			int32 Idx = 0;
-			for (Chaos::TTrailingDataExt<float, 3>& Trailing : AllTrailingsArray)
-			{
-				Trailing = (*AllTrailingData_Maps.AllTrailingData).AllTrailingsArray[Idx];
-
-				// #GM: Disable this for now for perf
-				/*
-				GetMeshExtData(SolverData,
-					AllTrailingsArray[Idx].ParticleIndexMesh == INDEX_NONE ? AllTrailingsArray[Idx].ParticleIndex : AllTrailingsArray[Idx].ParticleIndexMesh,
-					SolverObjectReverseMappingArray,
-					ParticleIndexReverseMappingArray,
-					AllTrailingsArray[Idx].BoundingboxVolume,
-					AllTrailingsArray[Idx].BoundingboxExtentMin,
-					AllTrailingsArray[Idx].BoundingboxExtentMax,
-					AllTrailingsArray[Idx].SurfaceType);
-				*/
-				AllTrailingsArray[Idx].BoundingboxVolume = 1000000.f;
-				AllTrailingsArray[Idx].BoundingboxExtentMin = 100.f;
-				AllTrailingsArray[Idx].BoundingboxExtentMax = 100.f;
-				AllTrailingsArray[Idx].SurfaceType = 0;
-
-				Idx++;
-			}
-		}
-		else
-		{
-			return false;
-		}
+	ensure(IsInGameThread());
+
+	// Copy data from *AllTrailingData_Maps.AllTrailingData into AllTrailingsArray
+	// Also get Boundingbox related data and SurfaceType and save it as well
+	TrailingEvents.AddUninitialized(Event.TrailingData.AllTrailingsArray.Num());
+
+	int32 Idx = 0;
+	for (Chaos::TTrailingDataExt<float, 3>& Trailing : TrailingEvents)
+	{
+		Trailing = Event.TrailingData.AllTrailingsArray[Idx];
+
+		// #GM: Disable this for now for perf
+		/*
+		GetMeshExtData(SolverData,
+			AllTrailingsArray[Idx].ParticleIndexMesh == INDEX_NONE ? AllTrailingsArray[Idx].ParticleIndex : AllTrailingsArray[Idx].ParticleIndexMesh,
+			PhysicsProxyReverseMappingArray,
+			ParticleIndexReverseMappingArray,
+			AllTrailingsArray[Idx].BoundingboxVolume,
+			AllTrailingsArray[Idx].BoundingboxExtentMin,
+			AllTrailingsArray[Idx].BoundingboxExtentMax,
+			AllTrailingsArray[Idx].SurfaceType);
+		*/
+		TrailingEvents[Idx].BoundingboxVolume = 1000000.f;
+		TrailingEvents[Idx].BoundingboxExtentMin = 100.f;
+		TrailingEvents[Idx].BoundingboxExtentMax = 100.f;
+		TrailingEvents[Idx].SurfaceType = 0;
+
+		Idx++;
 	}
-	else
-	{
-		return false;
-	}
-
-	TimeData_MapsCreated = (*AllTrailingData_Maps.AllTrailingData).TimeCreated;
-	TimeStampOfLastProcessedData = TimeData_MapsCreated;
-	// ----------------------------------------------------------------------------------------------------------
-	// END OF GETTING DATA FROM PBDRIGIDSOLVER
-	// ----------------------------------------------------------------------------------------------------------
-
-	return true;
 }
 
 void UNiagaraDataInterfaceChaosDestruction::FilterAllTrailings(TArray<Chaos::TTrailingDataExt<float, 3>>& AllTrailingsArray)
 {
-	if (ParticleIndexToProcess != -1 ||
+	if (/*ParticleIndexToProcess != -1 ||*/
 		SpeedToSpawnMinMax.X > 0.f ||
 		SpeedToSpawnMinMax.Y > 0.f ||
 		MassToSpawnMinMax.X > 0.f ||
@@ -2007,7 +1785,7 @@ void UNiagaraDataInterfaceChaosDestruction::FilterAllTrailings(TArray<Chaos::TTr
 		{
 			float TrailingSpeedSquared = AllTrailingsArray[IdxTrailing].Velocity.SizeSquared();
 
-			if ((ParticleIndexToProcess != -1 && AllTrailingsArray[IdxTrailing].ParticleIndex != ParticleIndexToProcess) ||
+			if (/*(ParticleIndexToProcess != -1 && AllTrailingsArray[IdxTrailing].ParticleIndex != ParticleIndexToProcess) ||*/
 				(SpeedToSpawnMinMax.X > 0.f && SpeedToSpawnMinMax.Y < 0.f && TrailingSpeedSquared < MinSpeedToSpawnSquared) ||
 				(SpeedToSpawnMinMax.X < 0.f && SpeedToSpawnMinMax.Y > 0.f && TrailingSpeedSquared > MaxSpeedToSpawnSquared) ||
 				(SpeedToSpawnMinMax.X > 0.f && SpeedToSpawnMinMax.Y > 0.f && (TrailingSpeedSquared < MinSpeedToSpawnSquared || TrailingSpeedSquared > MaxSpeedToSpawnSquared)) ||
@@ -2205,7 +1983,7 @@ int32 UNiagaraDataInterfaceChaosDestruction::SpawnParticlesFromTrailing(FSolverD
 			}
 			else if (DebugType == EDebugTypeEnum::ChaosNiagara_DebugType_ColorByParticleIndex)
 			{
-				ParticleColor = ColorArray[Trailing.ParticleIndex % ColorArray.Num()];
+				//ParticleColor = ColorArray[Trailing.ParticleIndex % ColorArray.Num()];
 			}
 
 			// Store principal data
@@ -2244,38 +2022,28 @@ bool UNiagaraDataInterfaceChaosDestruction::TrailingCallback(FNDIChaosDestructio
 	int32 IdxSolver = 0;
 	for (FSolverData SolverData : Solvers)
 	{
-		if (SolverData.Solver->GetGenerateTrailingData() && SolverData.Solver->GetSolverTime() > 0.f && MaxNumberOfDataEntriesToSpawn > 0)
+		if (SolverData.Solver->GetEventFilters()->IsTrailingEventEnabled() && TrailingEvents.Num() > 0 && SolverData.Solver->GetSolverTime() > 0.f && MaxNumberOfDataEntriesToSpawn > 0)
 		{
-			TArray<Chaos::TTrailingDataExt<float, 3>> AllTrailingsArray;
-			TArray<SolverObjectWrapper> SolverObjectReverseMappingArray;
+			TArray<Chaos::TTrailingDataExt<float, 3>>& AllTrailingsArray = TrailingEvents;
+			TArray<PhysicsProxyWrapper> PhysicsProxyReverseMappingArray;
 			TArray<int32> ParticleIndexReverseMappingArray;
-			TMap<ISolverObjectBase*, TArray<int32>> AllTrailingsIndicesBySolverObjectMap;
-			float TimeData_MapsCreated;
-
-			if (!GetAllTrailingsAndMaps(SolverData,
-				AllTrailingsArray,
-				SolverObjectReverseMappingArray,
-				ParticleIndexReverseMappingArray,
-				AllTrailingsIndicesBySolverObjectMap,
-				TimeData_MapsCreated))
-			{
-				return false;
-			}
+			TMap<IPhysicsProxyBase*, TArray<int32>> AllTrailingsIndicesByPhysicsProxyMap;
+			float TimeData_MapsCreated = 0.0f;
 
 			{
 				size_t SizeOfAllTrailings = sizeof(Chaos::TTrailingData<float, 3>) * AllTrailingsArray.Num();
-				size_t SizeOfSolverObjectReverseMapping = sizeof(SolverObjectWrapper) * SolverObjectReverseMappingArray.Num();
+				size_t SizeOfPhysicsProxyReverseMapping = sizeof(PhysicsProxyWrapper) * PhysicsProxyReverseMappingArray.Num();
 				size_t SizeOfParticleIndexReverseMapping = sizeof(int32) * ParticleIndexReverseMappingArray.Num();
 
-				size_t SizeOfAllTrailingsIndicesBySolverObject = 0;
-				for (auto& Elem : AllTrailingsIndicesBySolverObjectMap)
+				size_t SizeOfAllTrailingsIndicesByPhysicsProxy = 0;
+				for (auto& Elem : AllTrailingsIndicesByPhysicsProxyMap)
 				{
-					SizeOfAllTrailingsIndicesBySolverObject += sizeof(int32) * (Elem.Value).Num();
+					SizeOfAllTrailingsIndicesByPhysicsProxy += sizeof(int32) * (Elem.Value).Num();
 				}
 				SET_MEMORY_STAT(STAT_AllTrailingsDataMemory, SizeOfAllTrailings);
-				SET_MEMORY_STAT(STAT_SolverObjectReverseMappingMemory, SizeOfSolverObjectReverseMapping);
+				SET_MEMORY_STAT(STAT_PhysicsProxyReverseMappingMemory, SizeOfPhysicsProxyReverseMapping);
 				SET_MEMORY_STAT(STAT_ParticleIndexReverseMappingMemory, SizeOfParticleIndexReverseMapping);
-				SET_MEMORY_STAT(STAT_AllTrailingsIndicesBySolverObjectMemory, SizeOfAllTrailingsIndicesBySolverObject);
+				SET_MEMORY_STAT(STAT_AllTrailingsIndicesByPhysicsProxyMemory, SizeOfAllTrailingsIndicesByPhysicsProxy);
 			}
 
 			INC_DWORD_STAT_BY(STAT_NiagaraNumAllTrailings, AllTrailingsArray.Num());
@@ -2312,7 +2080,7 @@ bool UNiagaraDataInterfaceChaosDestruction::TrailingCallback(FNDIChaosDestructio
 						/*
 						GetMesPhysicalData(SolverData,
 										   TrailingsToSpawnArray[IdxTrailing].ParticleIndexMesh == INDEX_NONE ? TrailingsToSpawnArray[IdxTrailing].ParticleIndex : TrailingsToSpawnArray[IdxTrailing].ParticleIndexMesh,
-										   SolverObjectReverseMappingArray,
+										   PhysicsProxyReverseMappingArray,
 										   ParticleIndexReverseMappingArray,
 										   Color,
 										   Friction,
