Fix non-unity compile errors
#rnx

diff --git a/Engine/Plugins/Experimental/ChaosCloth/Source/Chaos/Private/ChaosCloth/ChaosClothingSimulation.cpp b/Engine/Plugins/Experimental/ChaosCloth/Source/Chaos/Private/ChaosCloth/ChaosClothingSimulation.cpp
--- a/Engine/Plugins/Experimental/ChaosCloth/Source/Chaos/Private/ChaosCloth/ChaosClothingSimulation.cpp
+++ b/Engine/Plugins/Experimental/ChaosCloth/Source/Chaos/Private/ChaosCloth/ChaosClothingSimulation.cpp
@@ -7,6 +7,7 @@
 #include "ClothingSimulation.h" // ClothingSystemRuntimeInterface
 #include "Utils/ClothingMeshUtils.h" // ClothingSystemRuntimeCommon
 #include "Components/SkeletalMeshComponent.h"
+#include "Materials/Material.h"
 #include "PhysicsEngine/PhysicsAsset.h"
 #include "Chaos/Box.h"
 #include "Chaos/Capsule.h"
diff --git a/Engine/Plugins/Experimental/ChaosClothEditor/Source/Chaos/Private/ChaosClothEditor/ChaosClothEditorModule.cpp b/Engine/Plugins/Experimental/ChaosClothEditor/Source/Chaos/Private/ChaosClothEditor/ChaosClothEditorModule.cpp
--- a/Engine/Plugins/Experimental/ChaosClothEditor/Source/Chaos/Private/ChaosClothEditor/ChaosClothEditorModule.cpp
+++ b/Engine/Plugins/Experimental/ChaosClothEditor/Source/Chaos/Private/ChaosClothEditor/ChaosClothEditorModule.cpp
@@ -6,6 +6,7 @@
 
 #include "ClothingSystemEditorInterfaceModule.h"
 #include "Features/IModularFeatures.h"
+#include "Modules/ModuleManager.h"
 
 IMPLEMENT_MODULE(FChaosClothEditorModule, ChaosClothEditor);
 DEFINE_LOG_CATEGORY(LogChaosClothEditor);
