Cloth: Replace bind functions with lambdas in chaos cloth, to avoid using the std library.

diff --git a/Engine/Plugins/Experimental/ChaosCloth/Source/Chaos/Private/ChaosCloth/ChaosClothingSimulation.cpp b/Engine/Plugins/Experimental/ChaosCloth/Source/Chaos/Private/ChaosCloth/ChaosClothingSimulation.cpp
--- a/Engine/Plugins/Experimental/ChaosCloth/Source/Chaos/Private/ChaosCloth/ChaosClothingSimulation.cpp
+++ b/Engine/Plugins/Experimental/ChaosCloth/Source/Chaos/Private/ChaosCloth/ChaosClothingSimulation.cpp
@@ -35,7 +35,6 @@
 #include "PhysXIncludes.h"
 #endif
 
-#include <functional>
 #include <memory>
 #include <unordered_map>
 #include <unordered_set>
@@ -374,29 +373,27 @@ void ClothingSimulation::CreateActor(USkeletalMeshComponent* InOwnerComponent, U
         else
         {
             TArray<Chaos::TVector<int32, 3>> SurfaceConstraints = SurfaceElements;
-            Evolution->AddPBDConstraintFunction(std::bind(&Chaos::TPBDVolumeConstraint<float>::Apply,
-                                                          Chaos::TPBDVolumeConstraint<float>(Evolution->Particles(), MoveTemp(SurfaceConstraints), VolumeStiffness), std::placeholders::_1, std::placeholders::_2));
+			Chaos::TPBDVolumeConstraint<float> PBDVolumeConstraint (Evolution->Particles(), MoveTemp(SurfaceConstraints));
+			Evolution->AddPBDConstraintFunction([=](TPBDParticles<float, 3>& InParticles, const float Dt)
+			{
+				PBDVolumeConstraint.Apply(InParticles, Dt);
+			});            
         }
     }
     if (StrainLimitingStiffness)
     {
 		check(Mesh->GetNumElements() > 0);
-        Evolution->AddPBDConstraintFunction(
-			std::bind(
-				static_cast
-				<
-					void (Chaos::TPerParticlePBDLongRangeConstraints<float, 3>::*)(
-						TPBDParticles<float, 3> & InParticles,
-						const float Dt) const
-				>(&Chaos::TPerParticlePBDLongRangeConstraints<float, 3>::Apply),
-				Chaos::TPerParticlePBDLongRangeConstraints<float, 3>(
-					Evolution->Particles(),
-					Mesh->GetPointToNeighborsMap(),
-					10, // The max number of connected neighbors per particle.  ryan - What should this be?  Was k...
-					StrainLimitingStiffness),
-				std::placeholders::_1,
-				std::placeholders::_2));
-    }
+		Chaos::TPerParticlePBDLongRangeConstraints<float, 3> PerParticlePBDLongRangeConstraints(
+			Evolution->Particles(),
+			Mesh->GetPointToNeighborsMap(),
+			10, // The max number of connected neighbors per particle.  ryan - What should this be?  Was k...
+			StrainLimitingStiffness);
+
+		Evolution->AddPBDConstraintFunction([=](TPBDParticles<float, 3>& InParticles, const float Dt)
+		{
+			PerParticlePBDLongRangeConstraints.Apply(InParticles, Dt);
+		});
+	}
 
 	// Maximum Distance Constraints
 	const UEnum* const MeshTargets = PhysMesh->GetFloatArrayTargets();	
