Cloth: Make lambda captures explicit

diff --git a/Engine/Plugins/Experimental/ChaosCloth/Source/Chaos/Private/ChaosCloth/ChaosClothingSimulation.cpp b/Engine/Plugins/Experimental/ChaosCloth/Source/Chaos/Private/ChaosCloth/ChaosClothingSimulation.cpp
--- a/Engine/Plugins/Experimental/ChaosCloth/Source/Chaos/Private/ChaosCloth/ChaosClothingSimulation.cpp
+++ b/Engine/Plugins/Experimental/ChaosCloth/Source/Chaos/Private/ChaosCloth/ChaosClothingSimulation.cpp
@@ -374,7 +374,7 @@ void ClothingSimulation::CreateActor(USkeletalMeshComponent* InOwnerComponent, U
         {
             TArray<Chaos::TVector<int32, 3>> SurfaceConstraints = SurfaceElements;
 			Chaos::TPBDVolumeConstraint<float> PBDVolumeConstraint (Evolution->Particles(), MoveTemp(SurfaceConstraints));
-			Evolution->AddPBDConstraintFunction([=](TPBDParticles<float, 3>& InParticles, const float Dt)
+			Evolution->AddPBDConstraintFunction([PBDVolumeConstraint](TPBDParticles<float, 3>& InParticles, const float Dt)
 			{
 				PBDVolumeConstraint.Apply(InParticles, Dt);
 			});            
@@ -389,7 +389,7 @@ void ClothingSimulation::CreateActor(USkeletalMeshComponent* InOwnerComponent, U
 			10, // The max number of connected neighbors per particle.  ryan - What should this be?  Was k...
 			StrainLimitingStiffness);
 
-		Evolution->AddPBDConstraintFunction([=](TPBDParticles<float, 3>& InParticles, const float Dt)
+		Evolution->AddPBDConstraintFunction([PerParticlePBDLongRangeConstraints](TPBDParticles<float, 3>& InParticles, const float Dt)
 		{
 			PerParticlePBDLongRangeConstraints.Apply(InParticles, Dt);
 		});
@@ -397,12 +397,12 @@ void ClothingSimulation::CreateActor(USkeletalMeshComponent* InOwnerComponent, U
 
 	// Maximum Distance Constraints
 	const UEnum* const MeshTargets = PhysMesh->GetFloatArrayTargets();	
-	const uint32 PhysMeshMaxDistanceIndex = MeshTargets->GetValueByName(TEXT("MaxDistance"));;
+	const uint32 PhysMeshMaxDistanceIndex = MeshTargets->GetValueByName(TEXT("MaxDistance"));
 	if (PhysMesh->GetFloatArray(PhysMeshMaxDistanceIndex)->Num() > 0)
 	{
 		check(Mesh->GetNumElements() > 0);
 		Chaos::PBDSphericalConstraint<float, 3> SphericalContraint(Offset, PhysMesh->GetFloatArray(PhysMeshMaxDistanceIndex)->Num(), true, &AnimationPositions, PhysMesh->GetFloatArray(PhysMeshMaxDistanceIndex));
-		Evolution->AddPBDConstraintFunction([=](TPBDParticles<float, 3>& InParticles, const float Dt)
+		Evolution->AddPBDConstraintFunction([SphericalContraint](TPBDParticles<float, 3>& InParticles, const float Dt)
 		{
 			SphericalContraint.Apply(InParticles, Dt);
 		});
@@ -418,7 +418,7 @@ void ClothingSimulation::CreateActor(USkeletalMeshComponent* InOwnerComponent, U
 
 		Chaos::PBDSphericalConstraint<float, 3> SphericalContraint(Offset, PhysMesh->GetFloatArray(PhysMeshBackstopRadiusIndex)->Num(), false, &AnimationPositions, 
 			PhysMesh->GetFloatArray(PhysMeshBackstopRadiusIndex), PhysMesh->GetFloatArray(PhysMeshBackstopDistanceIndex), &AnimationNormals);
-		Evolution->AddPBDConstraintFunction([=](TPBDParticles<float, 3>& InParticles, const float Dt)
+		Evolution->AddPBDConstraintFunction([SphericalContraint](TPBDParticles<float, 3>& InParticles, const float Dt)
 		{
 			SphericalContraint.Apply(InParticles, Dt);
 		});		
