Fix Chaos Cloth's external collision setup and the collision transform updates.
Also:
-Add the missing DestroyActors() implementation to stop leaking particles when the cloth asset are modified.
-Reimplement the collision visualization to draw the actual simulation collisions rather than the asset collision volumes.
-Split the collision creation from CreateActors into new functions.
-Fixed the size of the capsule colllision in the Chaos Cloth visualization.
-Slight general tidy up.

diff --git a/Engine/Plugins/Experimental/ChaosCloth/Source/Chaos/Private/ChaosCloth/ChaosClothPrivate.h b/Engine/Plugins/Experimental/ChaosCloth/Source/Chaos/Private/ChaosCloth/ChaosClothPrivate.h
--- a/Engine/Plugins/Experimental/ChaosCloth/Source/Chaos/Private/ChaosCloth/ChaosClothPrivate.h
+++ b/Engine/Plugins/Experimental/ChaosCloth/Source/Chaos/Private/ChaosCloth/ChaosClothPrivate.h
@@ -4,4 +4,4 @@
 
 #include "CoreMinimal.h"
 
-DECLARE_LOG_CATEGORY_EXTERN(LogChaosCloth, Verbose, All);
+DECLARE_LOG_CATEGORY_EXTERN(LogChaosCloth, Log, All);
diff --git a/Engine/Plugins/Experimental/ChaosCloth/Source/Chaos/Private/ChaosCloth/ChaosClothingSimulation.cpp b/Engine/Plugins/Experimental/ChaosCloth/Source/Chaos/Private/ChaosCloth/ChaosClothingSimulation.cpp
--- a/Engine/Plugins/Experimental/ChaosCloth/Source/Chaos/Private/ChaosCloth/ChaosClothingSimulation.cpp
+++ b/Engine/Plugins/Experimental/ChaosCloth/Source/Chaos/Private/ChaosCloth/ChaosClothingSimulation.cpp
@@ -50,7 +50,8 @@ static TAutoConsoleVariable<float> CVarClothDamping(TEXT("physics.ClothDamping")
 static TAutoConsoleVariable<float> CVarClothGravityMagnitude(TEXT("physics.ClothGravityMagnitude"), 490.f, TEXT(""));
 
 ClothingSimulation::ClothingSimulation()
-	: NumIterations(1)
+	: ExternalCollisionsOffset(0)
+	, NumIterations(1)
 	, EdgeStiffness(1.f)
 	, BendingStiffness(1.f)
 	, AreaStiffness(1.f)
@@ -115,10 +116,10 @@ void ClothingSimulation::Initialize()
 			checkSlow(DeltaTime > SMALL_NUMBER);
 			const float Alpha = (LocalTime - Time) / DeltaTime;
 			const Chaos::TVector<float, 3> NewX =
-				Alpha * AnimationTransforms[Index].GetTranslation() + (1.f - Alpha) * OldAnimationTransforms[Index].GetTranslation();
+				Alpha * CollisionTransforms[Index].GetTranslation() + (1.f - Alpha) * OldCollisionTransforms[Index].GetTranslation();
 			ParticlesInput.V(Index) = (NewX - ParticlesInput.X(Index)) / DeltaTime;
 			ParticlesInput.X(Index) = NewX;
-			Chaos::TRotation<float, 3> NewR = FQuat::Slerp(OldAnimationTransforms[Index].GetRotation(), AnimationTransforms[Index].GetRotation(), Alpha);
+			Chaos::TRotation<float, 3> NewR = FQuat::Slerp(OldCollisionTransforms[Index].GetRotation(), CollisionTransforms[Index].GetRotation(), Alpha);
 			Chaos::TRotation<float, 3> Delta = NewR * ParticlesInput.R(Index).Inverse();
 			Chaos::TVector<float, 3> Axis;
 			float Angle;
@@ -132,8 +133,39 @@ void ClothingSimulation::Initialize()
     Time = 0.f;
 }
 
+void ClothingSimulation::Shutdown()
+{
+	Assets.Reset();
+	ExtractedCollisions.Reset();
+	ExternalCollisions.Reset();
+	CollisionBoneNames.Reset();
+	CollisionBoneIndices.Reset();
+	OldCollisionTransforms.Reset();
+	CollisionTransforms.Reset();
+	BoneIndices.Reset();
+	BaseTransforms.Reset();
+	OldAnimationPositions.Reset();
+	AnimationPositions.Reset();
+	AnimationNormals.Reset();
+	IndexToRangeMap.Reset();
+	Meshes.Reset();
+	FaceNormals.Reset();
+	PointNormals.Reset();
+	Evolution.Reset();
+	ExternalCollisionsOffset = 0;
+}
+
+void ClothingSimulation::DestroyActors()
+{
+	Shutdown();
+	Initialize();
+}
+
 void ClothingSimulation::CreateActor(USkeletalMeshComponent* InOwnerComponent, UClothingAssetBase* InAsset, int32 InSimDataIndex)
 {
+	UE_LOG(LogChaosCloth, Verbose, TEXT("Adding Cloth LOD asset to %s in sim slot %d"), InOwnerComponent->GetOwner() ? *InOwnerComponent->GetOwner()->GetName() : TEXT("None"), InSimDataIndex);
+
+	// TODO(Kriss.Gossart): Set the cloth LOD parameters per cloth rather than per component
 	MassMode = InOwnerComponent->MassMode; // uniform, total, density
 	UniformMass = InOwnerComponent->UniformMass;
 	TotalMass = InOwnerComponent->TotalMass;
@@ -440,13 +472,57 @@ void ClothingSimulation::CreateActor(USkeletalMeshComponent* InOwnerComponent, U
             }
         }
     }
-    // Add Collision Bodies
-    TGeometryClothParticles<float, 3>& CollisionParticles = Evolution->CollisionParticles();
+
+	// Collision particles
+	TGeometryClothParticles<float, 3>& CollisionParticles = Evolution->CollisionParticles();
+
+	// Pull collisions from the specified physics asset inside the clothing asset
+	ExtractPhysicsAssetCollisions(Asset);
+
+	// Extract the legacy Apex collision from the clothing asset
+	ExtractLegacyAssetCollisions(Asset, InOwnerComponent);
+
+	// Set the external collision starting index
+	checkf(ExternalCollisions.Spheres.Num() == 0 &&
+		ExternalCollisions.SphereConnections.Num() == 0 &&
+		ExternalCollisions.Convexes.Num() == 0 &&
+		ExternalCollisions.Boxes.Num() == 0, TEXT("There cannot be any external collisions added before all the cloth assets collisions are processed."));
+	ExternalCollisionsOffset = CollisionParticles.Size();
+
+	// Refresh collision bone mapping
+	RefreshBoneMapping(Asset);
+
+	// Set the initial conditions for the collision particles
+	check(CollisionParticles.Size() == BaseTransforms.Num());	
+	CollisionTransforms.SetNum(BaseTransforms.Num());
+	for (uint32 i = 0; i < CollisionParticles.Size(); ++i)
+	{
+		const int32 BoneIndex = BoneIndices[i];
+		const int32 MappedIndex = CollisionBoneIndices.IsValidIndex(BoneIndex) ? CollisionBoneIndices[BoneIndex] : INDEX_NONE;
+		if (Context.BoneTransforms.IsValidIndex(MappedIndex))
+		{
+			const FTransform& BoneTransform = Context.BoneTransforms[MappedIndex];
+			CollisionTransforms[i] = BaseTransforms[i] * BoneTransform * Context.ComponentToWorld;
+		}
+		else
+		{
+			CollisionTransforms[i] = BaseTransforms[i] * Context.ComponentToWorld;  // External colllisions often don't map to a bone
+		}
+		CollisionParticles.X(i) = CollisionTransforms[i].GetTranslation();
+		CollisionParticles.R(i) = CollisionTransforms[i].GetRotation();
+	}
+}
+
+void ClothingSimulation::ExtractPhysicsAssetCollisions(UClothingAssetCommon* Asset)
+{
+	ExtractedCollisions.Reset();
+
+	TGeometryClothParticles<float, 3>& CollisionParticles = Evolution->CollisionParticles();
     //USkeletalMesh* TargetMesh = InOwnerComponent->SkeletalMesh;
-	USkeletalMesh* TargetMesh = CastChecked<USkeletalMesh>(Asset->GetOuter());
+	const USkeletalMesh* const TargetMesh = CastChecked<USkeletalMesh>(Asset->GetOuter());
 
-		// TODO(mlentine): Support collision body activation on a per particle basis, preferably using a map but also can be a particle attribute
-    if (UPhysicsAsset* PhysAsset = Asset->PhysicsAsset)
+	// TODO(mlentine): Support collision body activation on a per particle basis, preferably using a map but also can be a particle attribute
+    if (const UPhysicsAsset* const PhysAsset = Asset->PhysicsAsset)
     {
         for (const USkeletalBodySetup* BodySetup : PhysAsset->SkeletalBodySetups)
         {
@@ -456,63 +532,107 @@ void ClothingSimulation::CreateActor(USkeletalMeshComponent* InOwnerComponent, U
 			const int32 MeshBoneIndex = TargetMesh->RefSkeleton.FindBoneIndex(BodySetup->BoneName);
 			const int32 MappedBoneIndex =
 				MeshBoneIndex != INDEX_NONE ?
-				Asset->UsedBoneNames.AddUnique(BodySetup->BoneName) : INDEX_NONE;
+				CollisionBoneNames.AddUnique(BodySetup->BoneName) : INDEX_NONE;
 
 			const FKAggregateGeom& AggGeom = BodySetup->AggGeom;
-			if(AggGeom.SphylElems.Num())
+			if (AggGeom.SphylElems.Num())
             {
                 uint32 OldSize = CollisionParticles.Size();
                 CollisionParticles.AddParticles(AggGeom.SphylElems.Num());
                 for (uint32 i = OldSize; i < CollisionParticles.Size(); ++i)
                 {
                     const auto& Capsule = AggGeom.SphylElems[i - OldSize];
-                    if (Capsule.Length == 0)
+                    if (Capsule.Length == 0.0f)
                     {
-                        CollisionParticles.SetDynamicGeometry(i, MakeUnique<Chaos::TSphere<float,3>>(Chaos::TVector<float, 3>(0), Capsule.Radius));
-                    }
-                    else
-                    {
-                        Chaos::TVector<float, 3> half_extents(0, 0, Capsule.Length / 2);
-                        CollisionParticles.SetDynamicGeometry(
+						// Set particle
+						CollisionParticles.SetDynamicGeometry(i, MakeUnique<Chaos::TSphere<float, 3>>(Chaos::TVector<float, 3>(0.0f), Capsule.Radius));
+
+						// Add extracted collision data
+						FClothCollisionPrim_Sphere NewSphere;
+						NewSphere.LocalPosition = Capsule.Center;
+						NewSphere.Radius = Capsule.Radius;
+						NewSphere.BoneIndex = MappedBoneIndex;
+						ExtractedCollisions.Spheres.Add(NewSphere);
+					}
+					else
+					{
+						// Set particle
+						const Chaos::TVector<float, 3> HalfExtents(0.0f, 0.0f, Capsule.Length / 2.0f);  // TODO(Kriss.Gossart): Is this code missing the capsule rotation???
+						CollisionParticles.SetDynamicGeometry(
 							i,
 							MakeUnique<Chaos::TCapsule<float>>(
-								-half_extents, half_extents, Capsule.Radius));
-                    }
-                    BaseTransforms[i] = Chaos::TRigidTransform<float, 3>(Capsule.Center, Capsule.Rotation.Quaternion());
-                    BoneIndices[i] = MappedBoneIndex;
-                }
-            }
-			if(AggGeom.SphereElems.Num())
+								-HalfExtents, HalfExtents, Capsule.Radius));
+
+						// Add extracted collision data
+						FClothCollisionPrim_Sphere Sphere0;
+						FClothCollisionPrim_Sphere Sphere1;
+						const FVector OrientedDirection = Capsule.Rotation.RotateVector(FVector::UpVector);
+						const FVector HalfDim = OrientedDirection * (Capsule.Length / 2.f);
+						Sphere0.LocalPosition = Capsule.Center - HalfDim;
+						Sphere1.LocalPosition = Capsule.Center + HalfDim;
+						Sphere0.Radius = Capsule.Radius;
+						Sphere1.Radius = Capsule.Radius;
+						Sphere0.BoneIndex = MappedBoneIndex;
+						Sphere1.BoneIndex = MappedBoneIndex;
+
+						FClothCollisionPrim_SphereConnection Connection;
+						Connection.SphereIndices[0] = ExtractedCollisions.Spheres.Add(Sphere0);
+						Connection.SphereIndices[1] = ExtractedCollisions.Spheres.Add(Sphere1);
+						ExtractedCollisions.SphereConnections.Add(Connection);
+					}
+					BaseTransforms[i] = Chaos::TRigidTransform<float, 3>(Capsule.Center, Capsule.Rotation.Quaternion());
+					BoneIndices[i] = MappedBoneIndex;
+				}
+			}
+			if (AggGeom.SphereElems.Num())
             {
                 uint32 OldSize = CollisionParticles.Size();
                 CollisionParticles.AddParticles(AggGeom.SphereElems.Num());
                 for (uint32 i = OldSize; i < CollisionParticles.Size(); ++i)
                 {
                     const auto& CollisionSphere = AggGeom.SphereElems[i - OldSize];
-                    CollisionParticles.SetDynamicGeometry(
+
+					// Set particle
+					CollisionParticles.SetDynamicGeometry(
 						i,
-						MakeUnique<Chaos::TSphere<float,3>>(
-							Chaos::TVector<float, 3>(0.f, 0.f, 0.f),
+						MakeUnique<Chaos::TSphere<float, 3>>(
+							Chaos::TVector<float, 3>(0.f),
 							CollisionSphere.Radius));
-                    BaseTransforms[i] = Chaos::TRigidTransform<float, 3>(CollisionSphere.Center, Chaos::TRotation<float, 3>::MakeFromEuler(Chaos::TVector<float, 3>(0.f, 0.f, 0.f)));
+                    BaseTransforms[i] = Chaos::TRigidTransform<float, 3>(CollisionSphere.Center, Chaos::TRotation<float, 3>::Identity);
                     BoneIndices[i] = MappedBoneIndex;
-                }
-            }
-			if(AggGeom.BoxElems.Num())
-            {
-                uint32 OldSize = CollisionParticles.Size();
-                CollisionParticles.AddParticles(AggGeom.BoxElems.Num());
-                for (uint32 i = OldSize; i < CollisionParticles.Size(); ++i)
-                {
-                    const auto& Box = AggGeom.BoxElems[i - OldSize];
-                    Chaos::TVector<float, 3> half_extents(Box.X / 2.f, Box.Y / 2.f, Box.Z / 2.f);
-                    CollisionParticles.SetDynamicGeometry(i, MakeUnique<Chaos::TBox<float,3>>(-half_extents, half_extents));
-                    BaseTransforms[i] = Chaos::TRigidTransform<float, 3>(Box.Center, Box.Rotation.Quaternion());
-                    BoneIndices[i] = MappedBoneIndex;
-                }
-            }
-            /*
-			if(AggGeom.TaperedCapsuleElems.Num())
+
+					// Add extracted collision data
+					FClothCollisionPrim_Sphere NewSphere;
+					NewSphere.LocalPosition = CollisionSphere.Center;
+					NewSphere.Radius = CollisionSphere.Radius;
+					NewSphere.BoneIndex = MappedBoneIndex;
+					ExtractedCollisions.Spheres.Add(NewSphere);
+				}
+			}
+			if (AggGeom.BoxElems.Num())
+			{
+				uint32 OldSize = CollisionParticles.Size();
+				CollisionParticles.AddParticles(AggGeom.BoxElems.Num());
+				for (uint32 i = OldSize; i < CollisionParticles.Size(); ++i)
+				{
+					// Set particle
+					const auto& Box = AggGeom.BoxElems[i - OldSize];
+					const Chaos::TVector<float, 3> HalfExtents(Box.X / 2.f, Box.Y / 2.f, Box.Z / 2.f);
+					CollisionParticles.SetDynamicGeometry(i, MakeUnique<Chaos::TBox<float, 3>>(-HalfExtents, HalfExtents));
+					BaseTransforms[i] = Chaos::TRigidTransform<float, 3>(Box.Center, Box.Rotation.Quaternion());
+					BoneIndices[i] = MappedBoneIndex;
+
+					// Add extracted collision data
+					FClothCollisionPrim_Box NewBox;
+					NewBox.LocalMin = -HalfExtents;
+					NewBox.LocalMax = HalfExtents;
+					NewBox.BoneIndex = MappedBoneIndex;
+					ExtractedCollisions.Boxes.Add(NewBox);
+				}
+			}
+			// TODO(Kriss.Gossart): Check if there is any plan to support tapered capsules and fix the following section
+			/*
+			if (AggGeom.TaperedCapsuleElems.Num())
 			{
                 int32 OldSize = CollisionParticles.Size();
                 CollisionParticles.AddParticles(AggGeom.TaperedCapsuleElems.Num());
@@ -539,30 +659,43 @@ void ClothingSimulation::CreateActor(USkeletalMeshComponent* InOwnerComponent, U
                     BoneIndices[i] = MappedBoneIndex;
                 }
 			}*/
-			if(AggGeom.ConvexElems.Num())
+			if (AggGeom.ConvexElems.Num())
 			{
-// Collision bodies are stored in PhysX specific data structures so they can only be imported if we enable PhysX.
+				// Collision bodies are stored in PhysX specific data structures so they can only be imported if we enable PhysX.
 #if WITH_PHYSX && !PLATFORM_LUMIN && !PLATFORM_ANDROID
 				uint32 OldSize = CollisionParticles.Size();
 				CollisionParticles.AddParticles(AggGeom.ConvexElems.Num());
 				for (uint32 i = OldSize; i < CollisionParticles.Size(); ++i)
 				{
-				const auto& CollisionBody = AggGeom.ConvexElems[i - OldSize];
-				TArray<Chaos::TVector<int32, 3>> CollisionMeshElements;
-				const auto PhysXMesh = CollisionBody.GetConvexMesh();
-				for (int32 j = 0; j < static_cast<int32>(PhysXMesh->getNbPolygons()); ++j)
-				{
-					physx::PxHullPolygon Poly;
-					PhysXMesh->getPolygonData(j, Poly);
-					check(Poly.mNbVerts == 3);
-					const auto Indices = PhysXMesh->getIndexBuffer() + Poly.mIndexBase;
-					CollisionMeshElements.Add(Chaos::TVector<int32, 3>(Indices[0], Indices[1], Indices[2]));
-				}
-				Chaos::TParticles<float, 3> CollisionMeshParticles;
-				CollisionMeshParticles.AddParticles(CollisionBody.VertexData.Num());
-				for (uint32 j = 0; j < CollisionMeshParticles.Size(); ++j)
-				{
-					CollisionMeshParticles.X(j) = CollisionBody.VertexData[j];
+					const auto& CollisionBody = AggGeom.ConvexElems[i - OldSize];
+					const auto PhysXMesh = CollisionBody.GetConvexMesh();
+					const int32 NumPolygons = int32(PhysXMesh->getNbPolygons());
+
+					// Add stub for extracted collision data
+					FClothCollisionPrim_Convex NewConvex;
+					NewConvex.Planes.SetNumUninitialized(NumPolygons);
+					NewConvex.BoneIndex = MappedBoneIndex;
+
+					// Setup new convex particle
+					TArray<Chaos::TVector<int32, 3>> CollisionMeshElements;
+					for (int32 j = 0; j < NumPolygons; ++j)
+					{
+						physx::PxHullPolygon Poly;
+						PhysXMesh->getPolygonData(j, Poly);
+						check(Poly.mNbVerts == 3);
+						const auto Indices = PhysXMesh->getIndexBuffer() + Poly.mIndexBase;
+						CollisionMeshElements.Add(Chaos::TVector<int32, 3>(Indices[0], Indices[1], Indices[2]));
+
+						NewConvex.Planes[j] = FPlane(
+							CollisionBody.VertexData[Indices[0]],
+							CollisionBody.VertexData[Indices[1]],
+							CollisionBody.VertexData[Indices[2]]);
+					}
+					Chaos::TParticles<float, 3> CollisionMeshParticles;
+					CollisionMeshParticles.AddParticles(CollisionBody.VertexData.Num());
+					for (uint32 j = 0; j < CollisionMeshParticles.Size(); ++j)
+					{
+						CollisionMeshParticles.X(j) = CollisionBody.VertexData[j];
 					}
 					Chaos::TBox<float, 3> BoundingBox(CollisionMeshParticles.X(0), CollisionMeshParticles.X(0));
 					for (uint32 j = 1; j < CollisionMeshParticles.Size(); ++j)
@@ -587,14 +720,27 @@ void ClothingSimulation::CreateActor(USkeletalMeshComponent* InOwnerComponent, U
 					Chaos::TUniformGrid<float, 3> Grid(BoundingBox.Min(), BoundingBox.Max(), Chaos::TVector<int32, 3>(100 * Extents[0] / Extents[MaxAxis], 100 * Extents[0] / Extents[MaxAxis], 100 * Extents[0] / Extents[MaxAxis]));
 					Chaos::TTriangleMesh<float> CollisionMesh(MoveTemp(CollisionMeshElements));
 					Chaos::FErrorReporter ErrorReporter;
-					CollisionParticles.SetDynamicGeometry(i, MakeUnique<Chaos::TLevelSet<float,3>>(ErrorReporter, Grid, CollisionMeshParticles, CollisionMesh));
-					BaseTransforms[i] = Chaos::TRigidTransform<float, 3>(Chaos::TVector<float, 3>(0.f, 0.f, 0.f), Chaos::TRotation<float, 3>::MakeFromEuler(Chaos::TVector<float, 3>(0.f, 0.f, 0.f)));
+					CollisionParticles.SetDynamicGeometry(i, MakeUnique<Chaos::TLevelSet<float, 3>>(ErrorReporter, Grid, CollisionMeshParticles, CollisionMesh));
+					BaseTransforms[i] = Chaos::TRigidTransform<float, 3>(Chaos::TVector<float, 3>(0.f), Chaos::TRotation<float, 3>::Identity);
 					BoneIndices[i] = MappedBoneIndex;
+
+					// Add extracted collision data
+					ExtractedCollisions.Convexes.Add(NewConvex);
 				}
 #endif
 			}
 		} // end for
 	} // end if PhysAsset
+	UE_LOG(LogChaosCloth, Verbose, TEXT("Added physics asset collisions: %d spheres, %d capsules, %d convexes, %d boxes."), ExtractedCollisions.Spheres.Num() - 2 * ExtractedCollisions.SphereConnections.Num(), ExtractedCollisions.SphereConnections.Num(), ExtractedCollisions.Convexes.Num(), ExtractedCollisions.Boxes.Num());
+}
+
+void ClothingSimulation::ExtractLegacyAssetCollisions(UClothingAssetCommon* Asset, const USkeletalMeshComponent* InOwnerComponent)
+{
+	check(Asset->GetNumLods() == 1);
+	const UClothLODDataBase* const AssetLodData = Asset->ClothLodData[0];
+	if (!AssetLodData) { return; }
+
+	TGeometryClothParticles<float, 3>& CollisionParticles = Evolution->CollisionParticles();
 
 	// We can't just use AddExternalCollisions() because we need to add entries for bone mappings and lookups.
 	const FClothCollisionData& LodCollData = AssetLodData->CollisionData;
@@ -627,37 +773,36 @@ void ClothingSimulation::CreateActor(USkeletalMeshComponent* InOwnerComponent, U
 				const int32 SphereIndex0 = Connection.SphereIndices[0];
 				const int32 SphereIndex1 = Connection.SphereIndices[1];
 				checkSlow(SphereIndex0 != SphereIndex1);
-				const float Radius = LodCollData.Spheres[SphereIndex0].Radius;
-				checkSlow(LodCollData.Spheres[SphereIndex0].Radius - LodCollData.Spheres[SphereIndex1].Radius < SMALL_NUMBER);
-				const Chaos::TVector<float, 3> X0 = LodCollData.Spheres[SphereIndex0].LocalPosition;
-				const Chaos::TVector<float, 3> X1 = LodCollData.Spheres[SphereIndex1].LocalPosition;
-
-				checkSlow(LodCollData.Spheres[SphereIndex0].BoneIndex == LodCollData.Spheres[SphereIndex1].BoneIndex);
-				const int32 BoneIndex = LodCollData.Spheres[SphereIndex0].BoneIndex;
-				const FName BoneName =
-					TargetMesh->RefSkeleton.IsValidIndex(BoneIndex) ?
-					TargetMesh->RefSkeleton.GetBoneName(BoneIndex) : NAME_None;
-				const int32 MappedBoneIndex = BoneName != NAME_None ?
-					Asset->UsedBoneNames.AddUnique(BoneName) : INDEX_NONE;
-				BoneIndices[i] = MappedBoneIndex;
-
-				const Chaos::TVector<float, 3> Center = (X0 + X1) * 0.5;
+				const FClothCollisionPrim_Sphere& CollisionSphere0 = LodCollData.Spheres[SphereIndex0];
+				const FClothCollisionPrim_Sphere& CollisionSphere1 = LodCollData.Spheres[SphereIndex1];
+
+				const float Radius = CollisionSphere0.Radius;
+				checkSlow(CollisionSphere0.Radius - CollisionSphere1.Radius < SMALL_NUMBER);
+				UE_CLOG(CollisionSphere0.Radius - CollisionSphere1.Radius >= SMALL_NUMBER,
+					LogChaosCloth, Warning, TEXT("Found a legacy Apex cloth asset with a collision capsule of two different radii."));
+
+				const Chaos::TVector<float, 3> X0 = CollisionSphere0.LocalPosition;
+				const Chaos::TVector<float, 3> X1 = CollisionSphere1.LocalPosition;
+
+				BoneIndices[i] = CollisionSphere0.BoneIndex;
+				checkSlow(CollisionSphere0.BoneIndex == CollisionSphere1.BoneIndex);
+				UE_CLOG(CollisionSphere0.BoneIndex != CollisionSphere1.BoneIndex,
+					LogChaosCloth, Warning, TEXT("Found a legacy Apex cloth asset with a collision capsule spanning across two bones."));
+
+				// We construct a capsule centered at the origin along the Z axis
+				const Chaos::TVector<float, 3> Center = (X0 + X1) * 0.5f;
 				const Chaos::TVector<float, 3> Axis = X1 - X0;
-				const float HalfHeight = Axis.Size() * 0.5;
 
-				// We construct a capsule centered at the origin along the Z axis, and
-				// then move it into place with X and R.
-				CollisionParticles.X(i) = Center;
-				CollisionParticles.R(i) = Chaos::TRotation<float, 3>::FromRotatedVector(Chaos::TVector<float, 3>::AxisVector(2), Axis.GetSafeNormal());
-				BaseTransforms[i] = Chaos::TRigidTransform<float, 3>(CollisionParticles.X(i), CollisionParticles.R(i));
+				const Chaos::TRotation<float, 3> Rotation = Chaos::TRotation<float, 3>::FromRotatedVector(Chaos::TVector<float, 3>::AxisVector(2), Axis.GetSafeNormal());
+				BaseTransforms[i] = Chaos::TRigidTransform<float, 3>(Center, Rotation);
 
+				const float HalfHeight = Axis.Size() * 0.5f;
 				CollisionParticles.SetDynamicGeometry(
 					i,
 					MakeUnique<Chaos::TCapsule<float>>(
 						Chaos::TVector<float, 3>(0.f, 0.f, -HalfHeight), // Min
 						Chaos::TVector<float, 3>(0.f, 0.f, HalfHeight), // Max
 						Radius));
-				IndexAndCapsuleCollisionMap.Add(MakePair(i, Connection));
 
 				// Skip spheres added as end caps for the capsule.
 				CapsuleEnds.Add(SphereIndex0);
@@ -677,41 +822,47 @@ void ClothingSimulation::CreateActor(USkeletalMeshComponent* InOwnerComponent, U
 
 				const FClothCollisionPrim_Sphere& CollisionSphere = LodCollData.Spheres[i];
 
-				const int32 BoneIndex = CollisionSphere.BoneIndex;
-				const FName BoneName =
-					TargetMesh->RefSkeleton.IsValidIndex(BoneIndex) ?
-					TargetMesh->RefSkeleton.GetBoneName(BoneIndex) : NAME_None;
-				const int32 MappedBoneIndex = BoneName != NAME_None ?
-					Asset->UsedBoneNames.AddUnique(BoneName) : INDEX_NONE;
-				BoneIndices[j] = MappedBoneIndex;
+				BoneIndices[i] = CollisionSphere.BoneIndex;
 
-				CollisionParticles.X(j) = CollisionSphere.LocalPosition;
-				CollisionParticles.R(j) = Chaos::TRotation<float, 3>::FromElements(1,0,0,0);
-				BaseTransforms[j] = Chaos::TRigidTransform<float, 3>(CollisionParticles.X(j), CollisionParticles.R(j));
+				BaseTransforms[j] = Chaos::TRigidTransform<float, 3>(CollisionSphere.LocalPosition, Chaos::TRotation<float, 3>::Identity);
 
 				CollisionParticles.SetDynamicGeometry(
 					j,
 					MakeUnique<Chaos::TSphere<float, 3>>(
-						Chaos::TVector<float, 3>(0.f, 0.f, 0.f),
+						Chaos::TVector<float, 3>(0.0f),
 						CollisionSphere.Radius));
-				IndexAndSphereCollisionMap.Add(MakePair(j, CollisionSphere));
-				j++;
+				++j;
 			}
 		}
+		// TODO(Kriss.Gossart): Convexes are missing (but boxes are a new addition, so they are not legacy)
 
+		UE_LOG(LogChaosCloth, Verbose, TEXT("Added legacy asset collisions: %d spheres, %d capsules, %d convexes."), LodCollData.Spheres.Num() - CapsuleEnds.Num(), LodCollData.SphereConnections.Num(), LodCollData.Convexes.Num());
 	} // end if LodCollData
+}
 
-	check(CollisionParticles.Size() == BaseTransforms.Num());	
-	AnimationTransforms.SetNum(BaseTransforms.Num());
-	for (uint32 i = 0; i < CollisionParticles.Size(); ++i)
+void ClothingSimulation::RefreshBoneMapping(UClothingAssetCommon* Asset)
+{
+	// No mesh, can't remap
+	if (const USkeletalMesh* const SkeletalMesh = CastChecked<USkeletalMesh>(Asset->GetOuter()))
 	{
-		const int32 MappedIndex = BoneIndices[i];
-		if (Context.BoneTransforms.IsValidIndex(MappedIndex))
+		// Add the asset known used bone names (will take care of the apex collision mapping)
+		for (const FName& Name : Asset->UsedBoneNames)
 		{
-			const FTransform& BoneTransform = Context.BoneTransforms[MappedIndex];
-			AnimationTransforms[i] = BaseTransforms[i] * BoneTransform * Context.ComponentToWorld;
-			CollisionParticles.X(i) = AnimationTransforms[i].GetTranslation();
-			CollisionParticles.R(i) = AnimationTransforms[i].GetRotation();
+			CollisionBoneNames.AddUnique(Name);
+		}
+
+		// New names added
+		if(CollisionBoneNames.Num() != CollisionBoneIndices.Num())
+		{
+			CollisionBoneIndices.Reset();
+			CollisionBoneIndices.AddDefaulted(CollisionBoneNames.Num());
+		}
+
+		// Repopulate the used indices
+		for (int32 BoneNameIndex = 0; BoneNameIndex < CollisionBoneNames.Num(); ++BoneNameIndex)
+		{
+			const FName& UsedBoneName = CollisionBoneNames[BoneNameIndex];
+			CollisionBoneIndices[BoneNameIndex] = SkeletalMesh->RefSkeleton.FindBoneIndex(UsedBoneName);
 		}
 	}
 }
@@ -783,7 +934,7 @@ void ClothingSimulation::Simulate(IClothingSimulationContext* InContext)
 		return;
 
 	// Get New Animation Positions and Normals
-	OldAnimationTransforms = AnimationTransforms;
+	OldCollisionTransforms = CollisionTransforms;
 	OldAnimationPositions = AnimationPositions;
 
 	for (int32 Index = 0; Index < IndexToRangeMap.Num(); ++Index)
@@ -825,17 +976,29 @@ void ClothingSimulation::Simulate(IClothingSimulationContext* InContext)
 		});
 	}
 
-	// Collision bodies
+	// Update collision tranforms
 	TGeometryClothParticles<float, 3>& CollisionParticles = Evolution->CollisionParticles();
 	for (uint32 i = 0; i < CollisionParticles.Size(); ++i)
 	{
-		const int32 MappedIndex = BoneIndices[i];
+		const int32 BoneIndex = BoneIndices[i];
+		const int32 MappedIndex = CollisionBoneIndices.IsValidIndex(BoneIndex) ? CollisionBoneIndices[BoneIndex] : INDEX_NONE;
 		if (Context->BoneTransforms.IsValidIndex(MappedIndex))
 		{
 			const FTransform& BoneTransform = Context->BoneTransforms[MappedIndex];
-			AnimationTransforms[i] = BaseTransforms[i] * BoneTransform * Context->ComponentToWorld;
+			CollisionTransforms[i] = BaseTransforms[i] * BoneTransform * Context->ComponentToWorld;
+		}
+		else
+		{
+			CollisionTransforms[i] = BaseTransforms[i] * Context->ComponentToWorld;  // External colllisions often don't map to a bone
 		}
 	}
+
+	// Make sure external collision have a previous transform  // TODO(Kriss.Gossart): This is a temporary fix and needs changing. With removing/re-adding external collision at every frame there's no transform history. 
+	for (uint32 i = ExternalCollisionsOffset; i < CollisionParticles.Size(); ++i)
+	{
+		OldCollisionTransforms[i] = CollisionTransforms[i];
+	}
+
 	// Advance Sim
 	DeltaTime = Context->DeltaTime;
 	while (Context->DeltaTime > MaxDeltaTime)
@@ -901,12 +1064,15 @@ void ClothingSimulation::GetSimulationData(
 
 void ClothingSimulation::AddExternalCollisions(const FClothCollisionData& InData)
 {
-    TGeometryClothParticles<float, 3>& CollisionParticles = Evolution->CollisionParticles();
-	TSet<int32> CapsuleEnds;
+	// Keep track of the external collisions added
+	ExternalCollisions.Append(InData);
+
+	// Add particles
+	TGeometryClothParticles<float, 3>& CollisionParticles = Evolution->CollisionParticles();
 
-	const uint32 NumParticles0 = CollisionParticles.Size();
+	TSet<int32> CapsuleEnds;
 
-	if(InData.SphereConnections.Num())
+	if (InData.SphereConnections.Num())
 	{
 		const uint32 Size = CollisionParticles.Size();
         CollisionParticles.AddParticles(InData.SphereConnections.Num());
@@ -944,7 +1110,6 @@ void ClothingSimulation::AddExternalCollisions(const FClothCollisionData& InData
 					Chaos::TVector<float, 3>(0.f, 0.f, -HalfHeight), // Min
 					Chaos::TVector<float, 3>(0.f, 0.f, HalfHeight), // Max
 					Radius));
-            IndexAndCapsuleCollisionMap.Add(MakePair(i, Connection));
 
 			// Skip spheres added as end caps for the capsule.
 			CapsuleEnds.Add(SphereIndex0);
@@ -952,7 +1117,7 @@ void ClothingSimulation::AddExternalCollisions(const FClothCollisionData& InData
 		}
 	}
 
-	if(InData.Spheres.Num() - CapsuleEnds.Num())
+	if (InData.Spheres.Num() - CapsuleEnds.Num())
 	{
 		const uint32 Size = CollisionParticles.Size();
 		CollisionParticles.AddParticles(InData.Spheres.Num() - CapsuleEnds.Num());
@@ -962,94 +1127,107 @@ void ClothingSimulation::AddExternalCollisions(const FClothCollisionData& InData
 			if (CapsuleEnds.Contains(i))
 				continue;
 
-			const auto& CollisionSphere = InData.Spheres[i];
+			const FClothCollisionPrim_Sphere& CollisionSphere = InData.Spheres[i];
 			CollisionParticles.X(j) = Chaos::TVector<float, 3>(0.f);
 			CollisionParticles.R(j) = Chaos::TRotation<float, 3>::FromIdentity(); 
-			BaseTransforms[i] = Chaos::TRigidTransform<float, 3>(CollisionParticles.X(i), CollisionParticles.R(i));
+			BaseTransforms[i] = Chaos::TRigidTransform<float, 3>(FTransform::Identity);
 			BoneIndices[i] = CollisionSphere.BoneIndex;
 			CollisionParticles.SetDynamicGeometry(
 				j,
-				MakeUnique<Chaos::TSphere<float,3>>(
+				MakeUnique<Chaos::TSphere<float, 3>>(
 					CollisionSphere.LocalPosition, //Chaos::TVector<float, 3>(0.f, 0.f, 0.f),
 					CollisionSphere.Radius));
-				IndexAndSphereCollisionMap.Add(MakePair(j, CollisionSphere));
-			j++;
+
+			++j;
 		}
 	}
-	if(InData.Convexes.Num())
+
+	if (InData.Convexes.Num())
 	{
 		const uint32 Size = CollisionParticles.Size();
 		CollisionParticles.AddParticles(InData.Convexes.Num());
 		for (uint32 i = Size; i < CollisionParticles.Size(); ++i)
 		{
-			const auto& convex = InData.Convexes[i - Size];
-			CollisionParticles.X(i) = Chaos::TVector<float, 3>(0.f, 0.f, 0.f);
+			const FClothCollisionPrim_Convex& Convex = InData.Convexes[i - Size];
+			CollisionParticles.X(i) = Chaos::TVector<float, 3>(0.f);
 			CollisionParticles.R(i) = Chaos::TRotation<float, 3>::FromIdentity();
-			BaseTransforms[i] = Chaos::TRigidTransform<float, 3>(CollisionParticles.X(i), CollisionParticles.R(i));
-			BoneIndices[i] = convex.BoneIndex;
+			BaseTransforms[i] = Chaos::TRigidTransform<float, 3>(FTransform::Identity);
+			BoneIndices[i] = Convex.BoneIndex;
 			TArray<TUniquePtr<TImplicitObject<float, 3>>> Planes;
-			for (int32 j = 0; j < convex.Planes.Num(); ++j)
+			for (int32 j = 0; j < Convex.Planes.Num(); ++j)
 			{
 				Planes.Add(TUniquePtr<Chaos::TImplicitObject<float, 3>>(
-					new Chaos::TPlane<float, 3>(Chaos::TVector<float, 3>(0.f, 0.f, convex.Planes[j].W / convex.Planes[j].Z),
-					Chaos::TVector<float, 3>(convex.Planes[j].X, convex.Planes[j].Y, convex.Planes[j].Z))));
+					new Chaos::TPlane<float, 3>(Chaos::TVector<float, 3>(0.f, 0.f, Convex.Planes[j].W / Convex.Planes[j].Z),
+					Chaos::TVector<float, 3>(Convex.Planes[j].X, Convex.Planes[j].Y, Convex.Planes[j].Z))));
 			}
-				CollisionParticles.SetDynamicGeometry(i, MakeUnique<Chaos::TImplicitObjectIntersection<float, 3>>(MoveTemp(Planes)));
-				IndexAndConvexCollisionMap.Add(MakePair(i, convex));
+			CollisionParticles.SetDynamicGeometry(i, MakeUnique<Chaos::TImplicitObjectIntersection<float, 3>>(MoveTemp(Planes)));
+		}
+	}
+
+	if (InData.Boxes.Num())
+	{
+		const uint32 Size = CollisionParticles.Size();
+		CollisionParticles.AddParticles(InData.Convexes.Num());
+		for (uint32 i = Size; i < CollisionParticles.Size(); ++i)
+		{
+			const FClothCollisionPrim_Box& Box = InData.Boxes[i - Size];
+			CollisionParticles.X(i) = Chaos::TVector<float, 3>(0.f);
+			CollisionParticles.R(i) = Chaos::TRotation<float, 3>::FromIdentity();
+			BaseTransforms[i] = Chaos::TRigidTransform<float, 3>(FTransform::Identity);
+			BoneIndices[i] = Box.BoneIndex;
+			CollisionParticles.SetDynamicGeometry(i, MakeUnique<Chaos::TBox<float, 3>>(Box.LocalMin, Box.LocalMax));
 		}
 	}
 
 	check(CollisionParticles.Size() == BaseTransforms.Num());
-	// TODO: We need the context to initialize this animation data correctly, so do it in another function perhaps
-	const uint32 PrevAnimationTransformsCount = AnimationTransforms.Num();
-	const uint32 NewAnimationTransformsCount = BaseTransforms.Num();
 
-	AnimationTransforms.SetNum(NewAnimationTransformsCount);
-	OldAnimationTransforms.SetNum(NewAnimationTransformsCount);
+	const uint32 PrevCollisionTransformsCount = CollisionTransforms.Num();
+	const uint32 NewCollisionTransformsCount = BaseTransforms.Num();
+
+	CollisionTransforms.SetNum(NewCollisionTransformsCount);
+	OldCollisionTransforms.SetNum(NewCollisionTransformsCount);
 	
-	for (uint32 i = PrevAnimationTransformsCount; i < NewAnimationTransformsCount; ++i)
-	{
-		AnimationTransforms[i] = BaseTransforms[i];		
-		OldAnimationTransforms[i] = BaseTransforms[i];
-	}
+	UE_LOG(LogChaosCloth, VeryVerbose, TEXT("Added external collisions: %d spheres, %d capsules, %d convexes, %d boxes."), InData.Spheres.Num() - CapsuleEnds.Num(), InData.SphereConnections.Num(), InData.Convexes.Num(), InData.Boxes.Num());
 }
 
 void ClothingSimulation::ClearExternalCollisions()
 {
+	// Remove all external collision particles, starting from the external collision offset
 	TGeometryClothParticles<float, 3>& CollisionParticles = Evolution->CollisionParticles();
-	CollisionParticles.Resize(0);
+	CollisionParticles.Resize(ExternalCollisionsOffset);
 
-	IndexAndSphereCollisionMap.Reset();
-	IndexAndCapsuleCollisionMap.Reset();
-	IndexAndConvexCollisionMap.Reset();
+	// Reset external collisions
+	ExternalCollisionsOffset = CollisionParticles.Size();
+	ExternalCollisions.Reset();
+
+	UE_LOG(LogChaosCloth, VeryVerbose, TEXT("Cleared all external collisions."));
 }
 
 void ClothingSimulation::GetCollisions(FClothCollisionData& OutCollisions, bool bIncludeExternal) const
 {
-	OutCollisions.Spheres.Reset();
-	OutCollisions.SphereConnections.Reset();
-	OutCollisions.Convexes.Reset();
-	for (const auto& IndexSphere : IndexAndSphereCollisionMap)
+	OutCollisions.Reset();
+
+	// Add internal asset collisions
+	for (const UClothingAssetCommon* Asset : Assets)
 	{
-		if (Evolution->Collided(IndexSphere.First))
+		if (const UClothLODDataBase* const ClothLodData = !Asset ? nullptr : Asset->ClothLodData[0])
 		{
-			OutCollisions.Spheres.Add(IndexSphere.Second);
+			OutCollisions.Append(ClothLodData->CollisionData);
 		}
 	}
-	for (const auto& IndexCapsule : IndexAndCapsuleCollisionMap)
+
+	// Add collisions extracted from the physics asset
+	// TODO: Including the following code seems to be the correct behaviour, but this did not appear
+	// in the NvCloth implementation, so best to leave it commented out for now.
+	//OutCollisions.Append(ExtractedCollisions);
+
+	// Add external asset collisions
+	if (bIncludeExternal)
 	{
-		if (Evolution->Collided(IndexCapsule.First))
-		{
-			OutCollisions.SphereConnections.Add(IndexCapsule.Second);
-		}
-	}
-	for (const auto& IndexConvex : IndexAndConvexCollisionMap)
-	{
-		if (Evolution->Collided(IndexConvex.First))
-		{
-			OutCollisions.Convexes.Add(IndexConvex.Second);
-		}
+		OutCollisions.Append(ExternalCollisions);
 	}
+
+	UE_LOG(LogChaosCloth, VeryVerbose, TEXT("Returned collisions: %d spheres, %d capsules, %d convexes, %d boxes."), OutCollisions.Spheres.Num() - 2 * OutCollisions.SphereConnections.Num(), OutCollisions.SphereConnections.Num(), OutCollisions.Convexes.Num(), OutCollisions.Boxes.Num());
 }
 
 #if WITH_EDITOR
@@ -1228,51 +1406,62 @@ void ClothingSimulation::DebugDrawInversedFaceNormals(USkeletalMeshComponent* Ow
 
 void ClothingSimulation::DebugDrawCollision(USkeletalMeshComponent* OwnerComponent, FPrimitiveDrawInterface* PDI) const
 {
-	static const FLinearColor ValidColor(FColor::Cyan);
-	static const FLinearColor InvalidColor(FColor::Red);
+	static const FLinearColor MappedColor(FColor::Cyan);
+	static const FLinearColor UnmappedColor(FColor::Red);
 
 	const TGeometryClothParticles<float, 3>& CollisionParticles = Evolution->CollisionParticles();
 
-	// Draw collision spheres
-	for (const auto& IndexAndSphereCollision : IndexAndSphereCollisionMap)
+	for (uint32 Index = 0; Index < CollisionParticles.Size(); ++Index)
 	{
-		const int32 Index = int32(IndexAndSphereCollision.First);
-
-		const Chaos::TSphere<float, 3>* const Sphere = CollisionParticles.DynamicGeometry(Index)->GetObject<Chaos::TSphere<float, 3>>();
-		if (Sphere)
-		{
-			const FLinearColor Color = (BoneIndices[Index] != INDEX_NONE) ? ValidColor : InvalidColor;
-
-			const Chaos::TVector<float, 3>& Center = CollisionParticles.X(Index);
-			const TRotation<float, 3>& Rotation = CollisionParticles.R(Index);
-
-			const float Radius = Sphere ? Sphere->GetRadius() : 0.0f;
-
-			DrawWireSphere(PDI, Center, Color, Radius, 12, SDPG_World, 0.0f, 0.001f, false);
-		}
-	}
+		const TUniquePtr<TImplicitObject<float, 3>>& DynamicGeometry = CollisionParticles.DynamicGeometry(Index);
+		const uint32 BoneIndex = BoneIndices[Index];
+		const int32 MappedIndex = CollisionBoneIndices.IsValidIndex(BoneIndex) ? CollisionBoneIndices[BoneIndex] : INDEX_NONE;
+		const FLinearColor Color = (MappedIndex != INDEX_NONE) ? MappedColor : UnmappedColor;
 
-	// Draw collision capsules
-	for (const auto& IndexAndCapsuleCollision : IndexAndCapsuleCollisionMap)
-	{
-		const int32 Index = int32(IndexAndCapsuleCollision.First);
+		const Chaos::TVector<float, 3>& Center = CollisionParticles.X(Index);
+		const TRotation<float, 3>& Rotation = CollisionParticles.R(Index);
 
-		const Chaos::TCapsule<float>* const Capsule = CollisionParticles.DynamicGeometry(Index)->GetObject<Chaos::TCapsule<float>>();
-		if (Capsule)
+		switch (DynamicGeometry->GetType())
 		{
-			const FLinearColor Color = (BoneIndices[Index] != INDEX_NONE) ? ValidColor : InvalidColor;
+		// Draw collision spheres
+		case Chaos::ImplicitObjectType::Sphere:
+			if (const Chaos::TSphere<float, 3>* const Sphere = DynamicGeometry->GetObject<Chaos::TSphere<float, 3>>())
+			{
+				const float Radius = Sphere->GetRadius();
+				const FTransform Transform(Rotation, Center);
+				DrawWireSphere(PDI, Transform, Color, Radius, 12, SDPG_World, 0.0f, 0.001f, false);
+			}
+			break;
 
-			const Chaos::TVector<float, 3>& Center = CollisionParticles.X(Index);
-			const TRotation<float, 3>& Rotation = CollisionParticles.R(Index);
+		// Draw collision boxes
+		case Chaos::ImplicitObjectType::Box:
+			if (const Chaos::TBox<float, 3>* const Box = DynamicGeometry->GetObject<Chaos::TBox<float, 3>>())
+			{
+				const FMatrix BoxToWorld(Rotation.ToMatrix());
+				const FVector Radii = Box->Extents() * 0.5f;
+				DrawWireBox(PDI, BoxToWorld, FBox(Box->Min(), Box->Max()), Color, SDPG_World, 0.0f, 0.001f, false);
+			}
+			break;
 
-			const float HalfHeight = Capsule ? Capsule->GetHeight() * 0.5f : 0.0f;
-			const float Radius = Capsule ? Capsule->GetRadius() : 0.0f;
+		// Draw collision capsules
+		case Chaos::ImplicitObjectType::Capsule:
+			if (const Chaos::TCapsule<float>* const Capsule = DynamicGeometry->GetObject<Chaos::TCapsule<float>>())
+			{
+				const float HalfHeight = Capsule->GetHeight() * 0.5f;
+				const float Radius = Capsule->GetRadius();
+
+				const FVector X = Rotation.RotateVector(FVector::ForwardVector);
+				const FVector Y = Rotation.RotateVector(FVector::RightVector);
+				const FVector Z = Rotation.RotateVector(FVector::UpVector);
 
-			const FVector X = Rotation.RotateVector(FVector::ForwardVector);
-			const FVector Y = Rotation.RotateVector(FVector::RightVector);
-			const FVector Z = Rotation.RotateVector(FVector::UpVector);
+				DrawWireCapsule(PDI, Center, X, Y, Z, Color, Radius, HalfHeight + Radius, 12, SDPG_World, 0.0f, 0.001f, false);
+			}
+			break;
 
-			DrawWireCapsule(PDI, Center, X, Y, Z, Color, Radius, HalfHeight, 12, SDPG_World, 0.0f, 0.001f, false);
+		// Draw everything else as a coordinate for now
+		default:
+			DrawCoordinateSystem(PDI, Center, FRotator(Rotation), 10.0f, SDPG_World, 0.1f);
+			break;
 		}
 	}
 }
@@ -1289,15 +1478,15 @@ void ClothingSimulation::DebugDrawMaxDistances(USkeletalMeshComponent* OwnerComp
 
 void ClothingSimulation::DebugDrawSelfCollision(USkeletalMeshComponent* OwnerComponent, FPrimitiveDrawInterface* PDI) const
 {
-	if(!bUseSelfCollisions)
+	if (!bUseSelfCollisions)
 	{
 		// No self collisions on this actor
 		return;
 	}
 	const TPBDParticles<float, 3>& Particles = Evolution->Particles();
-	for(int32 i = 0; i < IndexToRangeMap.Num(); ++i)
+	for (int32 i = 0; i < IndexToRangeMap.Num(); ++i)
 	{
-		if(const UClothingAssetCommon* const Asset = Assets[i])
+		if (const UClothingAssetCommon* const Asset = Assets[i])
 		{
 			const FTransform RootBoneTransform = OwnerComponent->GetComponentSpaceTransforms()[Asset->ReferenceBoneIndex];
 
diff --git a/Engine/Plugins/Experimental/ChaosCloth/Source/Chaos/Public/ChaosCloth/ChaosClothingSimulation.h b/Engine/Plugins/Experimental/ChaosCloth/Source/Chaos/Public/ChaosCloth/ChaosClothingSimulation.h
--- a/Engine/Plugins/Experimental/ChaosCloth/Source/Chaos/Public/ChaosCloth/ChaosClothingSimulation.h
+++ b/Engine/Plugins/Experimental/ChaosCloth/Source/Chaos/Public/ChaosCloth/ChaosClothingSimulation.h
@@ -36,7 +36,7 @@ namespace Chaos
 
 #if WITH_EDITOR
 		// FGCObject interface
-		virtual void AddReferencedObjects(FReferenceCollector& Collector) override;
+		void AddReferencedObjects(FReferenceCollector& Collector) override;
 		// End of FGCObject interface
 
 		CHAOSCLOTH_API void DebugDrawPhysMeshWired(USkeletalMeshComponent* OwnerComponent, FPrimitiveDrawInterface* PDI) const;
@@ -53,38 +53,53 @@ namespace Chaos
 #endif  // #if WITH_EDITOR
 
 	protected:
-		void Initialize();
-		void CreateActor(USkeletalMeshComponent* InOwnerComponent, UClothingAssetBase* InAsset, int32 SimDataIndex);
-		IClothingSimulationContext* CreateContext() { return new ClothingSimulationContext(); }
-		void FillContext(USkeletalMeshComponent* InComponent, float InDeltaTime, IClothingSimulationContext* InOutContext);
-		void Shutdown() {}
-		bool ShouldSimulate() const { return true; }
-		void Simulate(IClothingSimulationContext* InContext);
-		void DestroyActors() {}
-		void DestroyContext(IClothingSimulationContext* InContext) { delete InContext; }
-		void GetSimulationData(TMap<int32, FClothSimulData>& OutData, USkeletalMeshComponent* InOwnerComponent, USkinnedMeshComponent* InOverrideComponent) const;
-
-		FBoxSphereBounds GetBounds(const USkeletalMeshComponent* InOwnerComponent) const
+		// IClothingSimulation interface
+		void Initialize() override;
+		void CreateActor(USkeletalMeshComponent* InOwnerComponent, UClothingAssetBase* InAsset, int32 SimDataIndex) override;
+		IClothingSimulationContext* CreateContext() override { return new ClothingSimulationContext(); }
+		void FillContext(USkeletalMeshComponent* InComponent, float InDeltaTime, IClothingSimulationContext* InOutContext) override;
+		void Shutdown() override;
+		bool ShouldSimulate() const override { return true; }
+		void Simulate(IClothingSimulationContext* InContext) override;
+		void DestroyActors() override;
+		void DestroyContext(IClothingSimulationContext* InContext) override { delete InContext; }
+		void GetSimulationData(TMap<int32, FClothSimulData>& OutData, USkeletalMeshComponent* InOwnerComponent, USkinnedMeshComponent* InOverrideComponent) const override;
+
+		FBoxSphereBounds GetBounds(const USkeletalMeshComponent* InOwnerComponent) const override
 		{ return FBoxSphereBounds(Evolution->Particles().X().GetData(), Evolution->Particles().Size()); }
 
-		void AddExternalCollisions(const FClothCollisionData& InData);
-		void ClearExternalCollisions();
-		void GetCollisions(FClothCollisionData& OutCollisions, bool bIncludeExternal = true) const;
+		void AddExternalCollisions(const FClothCollisionData& InData) override;
+		void ClearExternalCollisions() override;
+		void GetCollisions(FClothCollisionData& OutCollisions, bool bIncludeExternal = true) const override;
+		// End of IClothingSimulation interface
 
 	private:
-		// UE Collision Data (Needed only for GetCollisions)
-		TArray<Pair<uint32, FClothCollisionPrim_Sphere>> IndexAndSphereCollisionMap;
-		TArray<Pair<uint32, FClothCollisionPrim_SphereConnection>> IndexAndCapsuleCollisionMap;
-		TArray<Pair<uint32, FClothCollisionPrim_Convex>> IndexAndConvexCollisionMap;
-		// Animation Data
+		// Extract the collisions from the physics asset referenced by the specified clothing asset
+		void ExtractPhysicsAssetCollisions(UClothingAssetCommon* Asset);
+		// Extract the collisions from the specified clothing asset 
+		void ExtractLegacyAssetCollisions(UClothingAssetCommon* Asset, const USkeletalMeshComponent* InOwnerComponent);
+		// Create the bone mapping for all the bones used by collision
+		void RefreshBoneMapping(UClothingAssetCommon* Asset);
+
+	private:
+		// Assets
 		TArray<UClothingAssetCommon*> Assets;
-		TArray<Chaos::TRigidTransform<float, 3>> OldAnimationTransforms;
-		TArray<Chaos::TRigidTransform<float, 3>> AnimationTransforms;
+
+		// Collision Data
+		FClothCollisionData ExtractedCollisions;  // Collisions extracted from the referenced physics asset
+		FClothCollisionData ExternalCollisions;  // External collisions
+		TArray<FName> CollisionBoneNames;  // List of the bone names used by collision
+		TArray<int32> CollisionBoneIndices;  // List of the indices for the bones in UsedBoneNames, used for remapping
+		TArray<Chaos::TRigidTransform<float, 3>> OldCollisionTransforms;  // Used for the kinematic collision transform update
+		TArray<Chaos::TRigidTransform<float, 3>> CollisionTransforms;  // Used for the kinematic collision transform update
+		Chaos::TArrayCollectionArray<int32> BoneIndices;
+		Chaos::TArrayCollectionArray<Chaos::TRigidTransform<float, 3>> BaseTransforms;
+
+		// Animation Data
 		TArray<Chaos::TVector<float, 3>> OldAnimationPositions;
 		TArray<Chaos::TVector<float, 3>> AnimationPositions;
 		TArray<Chaos::TVector<float, 3>> AnimationNormals;
-		Chaos::TArrayCollectionArray<int32> BoneIndices;
-		Chaos::TArrayCollectionArray<Chaos::TRigidTransform<float, 3>> BaseTransforms;
+
 		// Sim Data
 		TArray<Chaos::TVector<uint32, 2>> IndexToRangeMap;
 
@@ -93,6 +108,9 @@ namespace Chaos
 		mutable TArray<TArray<Chaos::TVector<float, 3>>> PointNormals;
 
 		TUniquePtr<Chaos::TPBDEvolution<float, 3>> Evolution;
+
+		uint32 ExternalCollisionsOffset;  // External collisions first particle index
+
 		float Time;
 		float DeltaTime;
 		float MaxDeltaTime;
@@ -122,7 +140,9 @@ namespace Chaos
 		bool bUseThinShellVolumeConstraints;
 		bool bUseSelfCollisions;
 		bool bUseContinuousCollisionDetection;
+
 #if WITH_EDITOR
+		// Visualization material
 		UMaterial* DebugClothMaterial;
 #endif  // #if WITH_EDITOR
 	};
diff --git a/Engine/Source/Runtime/ClothingSystemRuntimeInterface/Private/ClothCollisionData.cpp b/Engine/Source/Runtime/ClothingSystemRuntimeInterface/Private/ClothCollisionData.cpp
--- a/Engine/Source/Runtime/ClothingSystemRuntimeInterface/Private/ClothCollisionData.cpp
+++ b/Engine/Source/Runtime/ClothingSystemRuntimeInterface/Private/ClothCollisionData.cpp
@@ -7,6 +7,7 @@ void FClothCollisionData::Reset()
 	Spheres.Reset();
 	SphereConnections.Reset();
 	Convexes.Reset();
+	Boxes.Reset();
 }
 
 void FClothCollisionData::Append(const FClothCollisionData& InOther)
@@ -31,4 +32,6 @@ void FClothCollisionData::Append(const FClothCollisionData& InOther)
 	}
 
 	Convexes.Append(InOther.Convexes);
+
+	Boxes.Append(InOther.Boxes);
 }
\ No newline at end of file
diff --git a/Engine/Source/Runtime/ClothingSystemRuntimeInterface/Public/ClothCollisionData.h b/Engine/Source/Runtime/ClothingSystemRuntimeInterface/Public/ClothCollisionData.h
--- a/Engine/Source/Runtime/ClothingSystemRuntimeInterface/Public/ClothCollisionData.h
+++ b/Engine/Source/Runtime/ClothingSystemRuntimeInterface/Public/ClothCollisionData.h
@@ -27,4 +27,8 @@ struct CLOTHINGSYSTEMRUNTIMEINTERFACE_API FClothCollisionData
 	// Convex Data
 	UPROPERTY(EditAnywhere, Category = Collison)
 	TArray<FClothCollisionPrim_Convex> Convexes;
+
+	// Box data
+	UPROPERTY(EditAnywhere, Category = Collison)
+	TArray<FClothCollisionPrim_Box> Boxes;
 };
diff --git a/Engine/Source/Runtime/ClothingSystemRuntimeInterface/Public/ClothCollisionPrim.h b/Engine/Source/Runtime/ClothingSystemRuntimeInterface/Public/ClothCollisionPrim.h
--- a/Engine/Source/Runtime/ClothingSystemRuntimeInterface/Public/ClothCollisionPrim.h
+++ b/Engine/Source/Runtime/ClothingSystemRuntimeInterface/Public/ClothCollisionPrim.h
@@ -68,3 +68,26 @@ struct FClothCollisionPrim_Convex
 	UPROPERTY()
 	int32 BoneIndex;
 };
+
+/** Data for a single box primitive. */
+USTRUCT()
+struct FClothCollisionPrim_Box
+{
+	GENERATED_BODY()
+
+	FClothCollisionPrim_Box()
+		: BoneIndex(INDEX_NONE)
+		, LocalMin(FVector::ZeroVector)
+		, LocalMax(FVector::ZeroVector)
+	{}
+
+	UPROPERTY()
+	int32 BoneIndex;
+
+	UPROPERTY()
+	FVector LocalMin;
+
+	UPROPERTY()
+	FVector LocalMax;
+};
+
