Chaos CLoth:
- Animation Drives implemented as a PBD constraint
- Small Cloth Integrator fix for substepping

diff --git a/Engine/Plugins/Experimental/ChaosCloth/Source/Chaos/Private/ChaosCloth/ChaosClothingSimulation.cpp b/Engine/Plugins/Experimental/ChaosCloth/Source/Chaos/Private/ChaosCloth/ChaosClothingSimulation.cpp
--- a/Engine/Plugins/Experimental/ChaosCloth/Source/Chaos/Private/ChaosCloth/ChaosClothingSimulation.cpp
+++ b/Engine/Plugins/Experimental/ChaosCloth/Source/Chaos/Private/ChaosCloth/ChaosClothingSimulation.cpp
@@ -15,15 +15,16 @@
 #include "Chaos/ImplicitObjectIntersection.h"
 #include "Chaos/ImplicitObjectUnion.h"
 #include "Chaos/Levelset.h"
+#include "Chaos/PBDAnimDriveConstraint.h"
 #include "Chaos/PBDAxialSpringConstraints.h"
 #include "Chaos/PBDBendingConstraints.h"
 #include "Chaos/PBDParticles.h"
+#include "Chaos/PBDSphericalConstraint.h"
 #include "Chaos/PBDSpringConstraints.h"
 #include "Chaos/PBDVolumeConstraint.h"
 #include "Chaos/PerParticleGravity.h"
 #include "Chaos/PerParticlePBDLongRangeConstraints.h"
 #include "Chaos/PerParticlePBDShapeConstraints.h"
-#include "Chaos/PBDSphericalConstraint.h"
 #include "Chaos/Plane.h"
 #include "Chaos/Sphere.h"
 #include "Chaos/TaperedCylinder.h"
@@ -61,6 +62,7 @@ ClothingSimulation::ClothingSimulation()
 	, SelfCollisionThickness(2.f)
 	, CollisionThickness(1.2f)
 	, GravityMagnitude(490.f)
+	, AnimDriveSpringStiffness(0.001f)
 	, bUseBendingElements(false)
 	, bUseTetrahedralConstraints(false)
 	, bUseThinShellVolumeConstraints(false)
@@ -113,11 +115,11 @@ void ClothingSimulation::Initialize()
 //		[&](Chaos::TKinematicGeometryParticles<float, 3>& ParticlesInput, const float Dt, const float LocalTime, const int32 Index)
 		[&](Chaos::TKinematicGeometryClothParticles<float, 3>& ParticlesInput, const float Dt, const float LocalTime, const int32 Index)
 		{
-			checkSlow(DeltaTime > SMALL_NUMBER);
+			checkSlow(Dt > SMALL_NUMBER && DeltaTime > SMALL_NUMBER);
 			const float Alpha = (LocalTime - Time) / DeltaTime;
 			const Chaos::TVector<float, 3> NewX =
 				Alpha * CollisionTransforms[Index].GetTranslation() + (1.f - Alpha) * OldCollisionTransforms[Index].GetTranslation();
-			ParticlesInput.V(Index) = (NewX - ParticlesInput.X(Index)) / DeltaTime;
+			ParticlesInput.V(Index) = (NewX - ParticlesInput.X(Index)) / Dt;
 			ParticlesInput.X(Index) = NewX;
 			Chaos::TRotation<float, 3> NewR = FQuat::Slerp(OldCollisionTransforms[Index].GetRotation(), CollisionTransforms[Index].GetRotation(), Alpha);
 			Chaos::TRotation<float, 3> Delta = NewR * ParticlesInput.R(Index).Inverse();
@@ -406,7 +408,7 @@ void ClothingSimulation::CreateActor(USkeletalMeshComponent* InOwnerComponent, U
         {
             TArray<Chaos::TVector<int32, 3>> SurfaceConstraints = SurfaceElements;
 			Chaos::TPBDVolumeConstraint<float> PBDVolumeConstraint (Evolution->Particles(), MoveTemp(SurfaceConstraints));
-			Evolution->AddPBDConstraintFunction([PBDVolumeConstraint](TPBDParticles<float, 3>& InParticles, const float Dt)
+			Evolution->AddPBDConstraintFunction([PBDVolumeConstraint = MoveTemp(PBDVolumeConstraint)](TPBDParticles<float, 3>& InParticles, const float Dt)
 			{
 				PBDVolumeConstraint.Apply(InParticles, Dt);
 			});            
@@ -421,7 +423,7 @@ void ClothingSimulation::CreateActor(USkeletalMeshComponent* InOwnerComponent, U
 			10, // The max number of connected neighbors per particle.  ryan - What should this be?  Was k...
 			StrainLimitingStiffness);
 
-		Evolution->AddPBDConstraintFunction([PerParticlePBDLongRangeConstraints](TPBDParticles<float, 3>& InParticles, const float Dt)
+		Evolution->AddPBDConstraintFunction([PerParticlePBDLongRangeConstraints = MoveTemp(PerParticlePBDLongRangeConstraints)](TPBDParticles<float, 3>& InParticles, const float Dt)
 		{
 			PerParticlePBDLongRangeConstraints.Apply(InParticles, Dt);
 		});
@@ -434,7 +436,7 @@ void ClothingSimulation::CreateActor(USkeletalMeshComponent* InOwnerComponent, U
 	{
 		check(Mesh->GetNumElements() > 0);
 		Chaos::PBDSphericalConstraint<float, 3> SphericalContraint(Offset, PhysMesh->GetFloatArray(PhysMeshMaxDistanceIndex)->Num(), true, &AnimationPositions, PhysMesh->GetFloatArray(PhysMeshMaxDistanceIndex));
-		Evolution->AddPBDConstraintFunction([SphericalContraint](TPBDParticles<float, 3>& InParticles, const float Dt)
+		Evolution->AddPBDConstraintFunction([SphericalContraint = MoveTemp(SphericalContraint)](TPBDParticles<float, 3>& InParticles, const float Dt)
 		{
 			SphericalContraint.Apply(InParticles, Dt);
 		});
@@ -450,10 +452,23 @@ void ClothingSimulation::CreateActor(USkeletalMeshComponent* InOwnerComponent, U
 
 		Chaos::PBDSphericalConstraint<float, 3> SphericalContraint(Offset, PhysMesh->GetFloatArray(PhysMeshBackstopRadiusIndex)->Num(), false, &AnimationPositions, 
 			PhysMesh->GetFloatArray(PhysMeshBackstopRadiusIndex), PhysMesh->GetFloatArray(PhysMeshBackstopDistanceIndex), &AnimationNormals);
-		Evolution->AddPBDConstraintFunction([SphericalContraint](TPBDParticles<float, 3>& InParticles, const float Dt)
+		Evolution->AddPBDConstraintFunction([SphericalContraint = MoveTemp(SphericalContraint)](TPBDParticles<float, 3>& InParticles, const float Dt)
 		{
 			SphericalContraint.Apply(InParticles, Dt);
 		});		
+	}	
+
+	// Animation Drive Constraints
+	const uint32 PhysMeshAnimDriveIndex = MeshTargets->GetValueByName(TEXT("AnimDriveMultiplier"));
+	if (PhysMesh->GetFloatArray(PhysMeshAnimDriveIndex)->Num() > 0)
+	{
+		check(Mesh->GetNumElements() > 0);
+		TPBDAnimDriveConstraint<float, 3> PBDAnimDriveConstraint(Offset, &AnimationPositions, PhysMesh->GetFloatArray(PhysMeshAnimDriveIndex), AnimDriveSpringStiffness);
+		Evolution->AddPBDConstraintFunction(
+			[PBDAnimDriveConstraint = MoveTemp(PBDAnimDriveConstraint)](TPBDParticles<float, 3>& InParticles, const float Dt)
+		{
+			PBDAnimDriveConstraint.Apply(InParticles, Dt);
+		});
 	}
 
     // Add Self Collisions
diff --git a/Engine/Plugins/Experimental/ChaosCloth/Source/Chaos/Public/ChaosCloth/ChaosClothingSimulation.h b/Engine/Plugins/Experimental/ChaosCloth/Source/Chaos/Public/ChaosCloth/ChaosClothingSimulation.h
--- a/Engine/Plugins/Experimental/ChaosCloth/Source/Chaos/Public/ChaosCloth/ChaosClothingSimulation.h
+++ b/Engine/Plugins/Experimental/ChaosCloth/Source/Chaos/Public/ChaosCloth/ChaosClothingSimulation.h
@@ -135,6 +135,7 @@ namespace Chaos
 		float CoefficientOfFriction;
 		float Damping;
 		float GravityMagnitude;
+		float AnimDriveSpringStiffness;
 		bool bUseBendingElements;
 		bool bUseTetrahedralConstraints;
 		bool bUseThinShellVolumeConstraints;
diff --git a/Engine/Source/Runtime/ClothingSystemRuntimeInterface/Public/ClothingSimulation.h b/Engine/Source/Runtime/ClothingSystemRuntimeInterface/Public/ClothingSimulation.h
--- a/Engine/Source/Runtime/ClothingSystemRuntimeInterface/Public/ClothingSimulation.h
+++ b/Engine/Source/Runtime/ClothingSystemRuntimeInterface/Public/ClothingSimulation.h
@@ -17,7 +17,7 @@ public:
 
 	// The majority of the API for this class is protected. The required objects (skel meshes and the parallel task)
 	// are friends so they can use the functionality. For the most part the simulation is not designed to be used
-	// outside of the skeletal mesh component as it parallel simulation is tied to the skel mesh tick and
+	// outside of the skeletal mesh component as its parallel simulation is tied to the skel mesh tick and
 	// dependents.
 	// Any method that is available in the public section below should consider that it may be called while
 	// the simulation is running.
diff --git a/Engine/Source/Runtime/Experimental/Chaos/Public/Chaos/PBDAnimDriveConstraint.h b/Engine/Source/Runtime/Experimental/Chaos/Public/Chaos/PBDAnimDriveConstraint.h
new file mode 100644
--- /dev/null
+++ b/Engine/Source/Runtime/Experimental/Chaos/Public/Chaos/PBDAnimDriveConstraint.h
@@ -0,0 +1,57 @@
+// Copyright 1998-2019 Epic Games, Inc. All Rights Reserved.
+#pragma once
+
+#include "Chaos/ParticleRule.h"
+
+namespace Chaos
+{
+template<typename T, int d>
+class TPBDAnimDriveConstraint : public TParticleRule<T, d>
+{
+  public:
+	  TPBDAnimDriveConstraint(
+		  const int InParticleIndexOffset,
+		  const TArray<TVector<T, d>>* const InSpringNeutralPositions,
+		  const TArray<T>* const InSpringStiffnessMultiplier,
+		  float InSpringStiffness) :
+		  ParticleIndexOffset(InParticleIndexOffset),
+		  SpringNeutralPositions(InSpringNeutralPositions),
+		  SpringStiffnessMultiplier(InSpringStiffnessMultiplier),
+		  SpringStiffness(InSpringStiffness)
+	  {
+
+	  }
+	virtual ~TPBDAnimDriveConstraint() {}
+
+	inline virtual void Apply(TPBDParticles<T, d>& InParticles, const T Dt) const override
+	{
+		int32 ConstraintCount = SpringStiffnessMultiplier->Num();
+		PhysicsParallelFor(ConstraintCount, [&](int32 Index)
+		{
+			ApplyAnimDriveConstraint(InParticles, Dt, Index);
+		});
+	}
+
+	inline void ApplyAnimDriveConstraint(TPBDParticles<T, d>& InParticles, const T Dt, const int32 Index) const
+	{
+		
+		if (InParticles.InvM(Index) == 0)
+		{
+			return;
+		}
+
+		const TVector<T, d> NeutralPosition = (*SpringNeutralPositions)[Index + ParticleIndexOffset];
+		const TVector<T, d> ParticlePosition = InParticles.P(Index + ParticleIndexOffset);
+
+		const T RelaxationFactor = FMath::Min(SpringStiffness * (*SpringStiffnessMultiplier)[Index], 1.0f);
+		InParticles.P(Index + ParticleIndexOffset) = RelaxationFactor * NeutralPosition + (1 - RelaxationFactor) * ParticlePosition;
+	}
+
+private:
+	const int32 ParticleIndexOffset;
+	const TArray<TVector<T, d>>* const SpringNeutralPositions; // Size: Same as full particle array
+	const TArray<T>* const SpringStiffnessMultiplier; // Size: Number of Animation drive constraints to solve
+
+	const float SpringStiffness;
+};
+}
