Fix the Chaos Cloth external collision transform initial condition setup, and refactor the three Chaos Cloth cloth/asset/external assets collision particle setups to a single function.
Also:
-Fix and simplify bone index mapping.
-Make the new Box collision primitive to support transformed boxes.
-Add support fdor tapered capsule.
-Add missing support for convex collision.
-Fix debug draw collision sphere offset.

diff --git a/Engine/Plugins/Experimental/ChaosCloth/Source/Chaos/Private/ChaosCloth/ChaosClothingSimulation.cpp b/Engine/Plugins/Experimental/ChaosCloth/Source/Chaos/Private/ChaosCloth/ChaosClothingSimulation.cpp
--- a/Engine/Plugins/Experimental/ChaosCloth/Source/Chaos/Private/ChaosCloth/ChaosClothingSimulation.cpp
+++ b/Engine/Plugins/Experimental/ChaosCloth/Source/Chaos/Private/ChaosCloth/ChaosClothingSimulation.cpp
@@ -28,6 +28,7 @@
 #include "Chaos/Plane.h"
 #include "Chaos/Sphere.h"
 #include "Chaos/TaperedCylinder.h"
+#include "Chaos/Convex.h"
 #include "Chaos/Transform.h"
 #include "Chaos/Utilities.h"
 #include "Chaos/Vector.h"
@@ -140,8 +141,6 @@ void ClothingSimulation::Shutdown()
 	Assets.Reset();
 	ExtractedCollisions.Reset();
 	ExternalCollisions.Reset();
-	CollisionBoneNames.Reset();
-	CollisionBoneIndices.Reset();
 	OldCollisionTransforms.Reset();
 	CollisionTransforms.Reset();
 	BoneIndices.Reset();
@@ -503,28 +502,45 @@ void ClothingSimulation::CreateActor(USkeletalMeshComponent* InOwnerComponent, U
 		ExternalCollisions.Convexes.Num() == 0 &&
 		ExternalCollisions.Boxes.Num() == 0, TEXT("There cannot be any external collisions added before all the cloth assets collisions are processed."));
 	ExternalCollisionsOffset = CollisionParticles.Size();
+}
+
+void ClothingSimulation::UpdateCollisionTransforms(const ClothingSimulationContext& Context)
+{
+	TGeometryClothParticles<float, 3>& CollisionParticles = Evolution->CollisionParticles();
 
-	// Refresh collision bone mapping
-	RefreshBoneMapping(Asset);
+	// Resize the transform arrays
+	const int32 PrevNumCollisions = CollisionTransforms.Num();
+	const int32 NumCollisions = BaseTransforms.Num();
+	check(NumCollisions == int32(CollisionParticles.Size()));  // BaseTransforms should always automatically grow with the number of collision particles (collection array)
+
+	if (NumCollisions != PrevNumCollisions)
+	{
+		CollisionTransforms.SetNum(NumCollisions);
+		OldCollisionTransforms.SetNum(NumCollisions);
+	}
 
-	// Set the initial conditions for the collision particles
-	check(CollisionParticles.Size() == BaseTransforms.Num());	
-	CollisionTransforms.SetNum(BaseTransforms.Num());
-	for (uint32 i = 0; i < CollisionParticles.Size(); ++i)
+	// Update the collision transforms
+	for (int32 Index = 0; Index < NumCollisions; ++Index)
 	{
-		const int32 BoneIndex = BoneIndices[i];
-		const int32 MappedIndex = CollisionBoneIndices.IsValidIndex(BoneIndex) ? CollisionBoneIndices[BoneIndex] : INDEX_NONE;
-		if (Context.BoneTransforms.IsValidIndex(MappedIndex))
+		const int32 BoneIndex = BoneIndices[Index];
+		if (Context.BoneTransforms.IsValidIndex(BoneIndex))
 		{
-			const FTransform& BoneTransform = Context.BoneTransforms[MappedIndex];
-			CollisionTransforms[i] = BaseTransforms[i] * BoneTransform * Context.ComponentToWorld;
+			const FTransform& BoneTransform = Context.BoneTransforms[BoneIndex];
+			CollisionTransforms[Index] = BaseTransforms[Index] * BoneTransform * Context.ComponentToWorld;
 		}
 		else
 		{
-			CollisionTransforms[i] = BaseTransforms[i] * Context.ComponentToWorld;  // External colllisions often don't map to a bone
+			CollisionTransforms[Index] = BaseTransforms[Index] * Context.ComponentToWorld;  // External collisions often don't map to a bone
 		}
-		CollisionParticles.X(i) = CollisionTransforms[i].GetTranslation();
-		CollisionParticles.R(i) = CollisionTransforms[i].GetRotation();
+	}
+
+	// Set the new collision particles and transforms initial state
+	for (int32 Index = PrevNumCollisions; Index < NumCollisions; ++Index)
+	{
+		CollisionParticles.X(Index) = CollisionTransforms[Index].GetTranslation();
+		CollisionParticles.R(Index) = CollisionTransforms[Index].GetRotation();
+
+		OldCollisionTransforms[Index] = CollisionTransforms[Index];
 	}
 }
 
@@ -533,282 +549,240 @@ void ClothingSimulation::ExtractPhysicsAssetCollisions(UClothingAssetCommon* Ass
 	ExtractedCollisions.Reset();
 
 	TGeometryClothParticles<float, 3>& CollisionParticles = Evolution->CollisionParticles();
-    //USkeletalMesh* TargetMesh = InOwnerComponent->SkeletalMesh;
-	const USkeletalMesh* const TargetMesh = CastChecked<USkeletalMesh>(Asset->GetOuter());
 
 	// TODO(mlentine): Support collision body activation on a per particle basis, preferably using a map but also can be a particle attribute
-    if (const UPhysicsAsset* const PhysAsset = Asset->PhysicsAsset)
-    {
-        for (const USkeletalBodySetup* BodySetup : PhysAsset->SkeletalBodySetups)
-        {
+	if (const UPhysicsAsset* const PhysAsset = Asset->PhysicsAsset)
+	{
+		const USkeletalMesh* const TargetMesh = CastChecked<USkeletalMesh>(Asset->GetOuter());
+
+		TArray<int32> UsedBoneIndices;
+		UsedBoneIndices.Reserve(PhysAsset->SkeletalBodySetups.Num());
+
+		for (const USkeletalBodySetup* BodySetup : PhysAsset->SkeletalBodySetups)
+		{
 			if (!BodySetup)
 				continue;
 
 			const int32 MeshBoneIndex = TargetMesh->RefSkeleton.FindBoneIndex(BodySetup->BoneName);
-			const int32 MappedBoneIndex =
-				MeshBoneIndex != INDEX_NONE ?
-				CollisionBoneNames.AddUnique(BodySetup->BoneName) : INDEX_NONE;
-
+			const int32 MappedBoneIndex = UsedBoneIndices.Add(MeshBoneIndex);
+			
+			// Add capsules
 			const FKAggregateGeom& AggGeom = BodySetup->AggGeom;
 			if (AggGeom.SphylElems.Num())
-            {
-                uint32 OldSize = CollisionParticles.Size();
-                CollisionParticles.AddParticles(AggGeom.SphylElems.Num());
-                for (uint32 i = OldSize; i < CollisionParticles.Size(); ++i)
-                {
-                    const auto& Capsule = AggGeom.SphylElems[i - OldSize];
-                    if (Capsule.Length == 0.0f)
-                    {
-						// Set particle
-						CollisionParticles.SetDynamicGeometry(i, MakeUnique<Chaos::TSphere<float, 3>>(Chaos::TVector<float, 3>(0.0f), Capsule.Radius));
-
-						// Add extracted collision data
-						FClothCollisionPrim_Sphere NewSphere;
-						NewSphere.LocalPosition = Capsule.Center;
-						NewSphere.Radius = Capsule.Radius;
-						NewSphere.BoneIndex = MappedBoneIndex;
-						ExtractedCollisions.Spheres.Add(NewSphere);
+			{
+				for (const FKSphylElem& SphylElem : AggGeom.SphylElems)
+				{
+					if (SphylElem.Length == 0.0f)
+					{
+						// Add extracted sphere collision data
+						FClothCollisionPrim_Sphere Sphere;
+						Sphere.LocalPosition = SphylElem.Center;
+						Sphere.Radius = SphylElem.Radius;
+						Sphere.BoneIndex = MappedBoneIndex;
+						ExtractedCollisions.Spheres.Add(Sphere);
 					}
 					else
 					{
-						// Set particle
-						const Chaos::TVector<float, 3> HalfExtents(0.0f, 0.0f, Capsule.Length / 2.0f);  // TODO(Kriss.Gossart): Is this code missing the capsule rotation???
-						CollisionParticles.SetDynamicGeometry(
-							i,
-							MakeUnique<Chaos::TCapsule<float>>(
-								-HalfExtents, HalfExtents, Capsule.Radius));
-
-						// Add extracted collision data
+						// Add extracted spheres collision data
 						FClothCollisionPrim_Sphere Sphere0;
 						FClothCollisionPrim_Sphere Sphere1;
-						const FVector OrientedDirection = Capsule.Rotation.RotateVector(FVector::UpVector);
-						const FVector HalfDim = OrientedDirection * (Capsule.Length / 2.f);
-						Sphere0.LocalPosition = Capsule.Center - HalfDim;
-						Sphere1.LocalPosition = Capsule.Center + HalfDim;
-						Sphere0.Radius = Capsule.Radius;
-						Sphere1.Radius = Capsule.Radius;
+						const FVector OrientedDirection = SphylElem.Rotation.RotateVector(FVector::UpVector);
+						const FVector HalfDim = OrientedDirection * (SphylElem.Length / 2.f);
+						Sphere0.LocalPosition = SphylElem.Center - HalfDim;
+						Sphere1.LocalPosition = SphylElem.Center + HalfDim;
+						Sphere0.Radius = SphylElem.Radius;
+						Sphere1.Radius = SphylElem.Radius;
 						Sphere0.BoneIndex = MappedBoneIndex;
 						Sphere1.BoneIndex = MappedBoneIndex;
 
-						FClothCollisionPrim_SphereConnection Connection;
-						Connection.SphereIndices[0] = ExtractedCollisions.Spheres.Add(Sphere0);
-						Connection.SphereIndices[1] = ExtractedCollisions.Spheres.Add(Sphere1);
-						ExtractedCollisions.SphereConnections.Add(Connection);
+						// Add extracted sphere connection collision data
+						FClothCollisionPrim_SphereConnection SphereConnection;
+						SphereConnection.SphereIndices[0] = ExtractedCollisions.Spheres.Add(Sphere0);
+						SphereConnection.SphereIndices[1] = ExtractedCollisions.Spheres.Add(Sphere1);
+						ExtractedCollisions.SphereConnections.Add(SphereConnection);
 					}
-					BaseTransforms[i] = Chaos::TRigidTransform<float, 3>(Capsule.Center, Capsule.Rotation.Quaternion());
-					BoneIndices[i] = MappedBoneIndex;
 				}
 			}
-			if (AggGeom.SphereElems.Num())
-            {
-                uint32 OldSize = CollisionParticles.Size();
-                CollisionParticles.AddParticles(AggGeom.SphereElems.Num());
-                for (uint32 i = OldSize; i < CollisionParticles.Size(); ++i)
-                {
-                    const auto& CollisionSphere = AggGeom.SphereElems[i - OldSize];
-
-					// Set particle
-					CollisionParticles.SetDynamicGeometry(
-						i,
-						MakeUnique<Chaos::TSphere<float, 3>>(
-							Chaos::TVector<float, 3>(0.f),
-							CollisionSphere.Radius));
-                    BaseTransforms[i] = Chaos::TRigidTransform<float, 3>(CollisionSphere.Center, Chaos::TRotation<float, 3>::Identity);
-                    BoneIndices[i] = MappedBoneIndex;
-
-					// Add extracted collision data
-					FClothCollisionPrim_Sphere NewSphere;
-					NewSphere.LocalPosition = CollisionSphere.Center;
-					NewSphere.Radius = CollisionSphere.Radius;
-					NewSphere.BoneIndex = MappedBoneIndex;
-					ExtractedCollisions.Spheres.Add(NewSphere);
-				}
+
+			// Add spheres
+			for (const FKSphereElem& SphereElem : AggGeom.SphereElems)
+			{
+				// Add extracted sphere collision data
+				FClothCollisionPrim_Sphere Sphere;
+				Sphere.LocalPosition = SphereElem.Center;
+				Sphere.Radius = SphereElem.Radius;
+				Sphere.BoneIndex = MappedBoneIndex;
+				ExtractedCollisions.Spheres.Add(Sphere);
 			}
-			if (AggGeom.BoxElems.Num())
+
+			// Add boxes
+			for (const FKBoxElem& BoxElem : AggGeom.BoxElems)
 			{
-				uint32 OldSize = CollisionParticles.Size();
-				CollisionParticles.AddParticles(AggGeom.BoxElems.Num());
-				for (uint32 i = OldSize; i < CollisionParticles.Size(); ++i)
+				// Add extracted box collision data
+				FClothCollisionPrim_Box Box;
+				Box.LocalPosition = BoxElem.Center;
+				Box.LocalRotation = BoxElem.Rotation.Quaternion();
+				Box.HalfExtents = FVector(BoxElem.X, BoxElem.Y, BoxElem.Z) * 0.5f;
+				Box.BoneIndex = MappedBoneIndex;
+				ExtractedCollisions.Boxes.Add(Box);
+			}
+
+			// Add tapered capsules
+			for (const FKTaperedCapsuleElem& TaperedCapsuleElem : AggGeom.TaperedCapsuleElems)
+			{
+				if (TaperedCapsuleElem.Length == 0)
+				{
+					// Add extracted sphere collision data
+					FClothCollisionPrim_Sphere Sphere;
+					Sphere.LocalPosition = TaperedCapsuleElem.Center;
+					Sphere.Radius = FMath::Max(TaperedCapsuleElem.Radius0, TaperedCapsuleElem.Radius1);
+					Sphere.BoneIndex = MappedBoneIndex;
+					ExtractedCollisions.Spheres.Add(Sphere);
+				}
+				else
 				{
-					// Set particle
-					const auto& Box = AggGeom.BoxElems[i - OldSize];
-					const Chaos::TVector<float, 3> HalfExtents(Box.X / 2.f, Box.Y / 2.f, Box.Z / 2.f);
-					CollisionParticles.SetDynamicGeometry(i, MakeUnique<Chaos::TBox<float, 3>>(-HalfExtents, HalfExtents));
-					BaseTransforms[i] = Chaos::TRigidTransform<float, 3>(Box.Center, Box.Rotation.Quaternion());
-					BoneIndices[i] = MappedBoneIndex;
-
-					// Add extracted collision data
-					FClothCollisionPrim_Box NewBox;
-					NewBox.LocalMin = -HalfExtents;
-					NewBox.LocalMax = HalfExtents;
-					NewBox.BoneIndex = MappedBoneIndex;
-					ExtractedCollisions.Boxes.Add(NewBox);
+					// Add extracted spheres collision data
+					FClothCollisionPrim_Sphere Sphere0;
+					FClothCollisionPrim_Sphere Sphere1;
+					const FVector OrientedDirection = TaperedCapsuleElem.Rotation.RotateVector(FVector::UpVector);
+					const FVector HalfDim = OrientedDirection * (TaperedCapsuleElem.Length / 2.f);
+					Sphere0.LocalPosition = TaperedCapsuleElem.Center - HalfDim;
+					Sphere1.LocalPosition = TaperedCapsuleElem.Center + HalfDim;
+					Sphere0.Radius = TaperedCapsuleElem.Radius0;
+					Sphere1.Radius = TaperedCapsuleElem.Radius1;
+					Sphere0.BoneIndex = MappedBoneIndex;
+					Sphere1.BoneIndex = MappedBoneIndex;
+
+					// Add extracted sphere connection collision data
+					FClothCollisionPrim_SphereConnection SphereConnection;
+					SphereConnection.SphereIndices[0] = ExtractedCollisions.Spheres.Add(Sphere0);
+					SphereConnection.SphereIndices[1] = ExtractedCollisions.Spheres.Add(Sphere1);
+					ExtractedCollisions.SphereConnections.Add(SphereConnection);
 				}
 			}
-			// TODO(Kriss.Gossart): Check if there is any plan to support tapered capsules and fix the following section
-			/*
-			if (AggGeom.TaperedCapsuleElems.Num())
-			{
-                int32 OldSize = CollisionParticles.Size();
-                CollisionParticles.AddParticles(AggGeom.TaperedCapsuleElems.Num());
-                for (int32 i = OldSize; i < CollisionParticles.Size(); ++i)
-                {
-                    const auto& Capsule = AggGeom.TaperedCapsuleElems[i - OldSize];
-                    if (Capsule.Length == 0)
-                    {
-                        CollisionParticles.Geometry(i) = new Chaos::Sphere<float, 3>(Chaos::TVector<float, 3>(0), Capsule.Radius1 > Capsule.Radius0 ? Capsule.Radius1 : Capsule.Radius0);
-                    }
-                    else
-                    {
-                        TArray<TUniquePtr<TImplicitObject<float, 3>>> Objects;
-                        Chaos::TVector<float, 3> half_extents(0, 0, Capsule.Length / 2);
-                        Objects.Add(TUniquePtr<Chaos::TImplicitObject<float, 3>>(
-                            new Chaos::TTaperedCylinder<float>(-half_extents, half_extents, Capsule.Radius1, Capsule.Radius0)));
-                        Objects.Add(TUniquePtr<Chaos::TImplicitObject<float, 3>>(
-                            new Chaos::Sphere<float, 3>(-half_extents, Capsule.Radius1)));
-                        Objects.Add(TUniquePtr<Chaos::TImplicitObject<float, 3>>(
-                            new Chaos::Sphere<float, 3>(half_extents, Capsule.Radius0)));
-                        CollisionParticles.Geometry(i) = new Chaos::TImplicitObjectUnion<float, 3>(MoveTemp(Objects));
-                    }
-					BaseTransforms[i] = Chaos::TRigidTransform<float, 3>(Capsule.Center, Capsule.Rotation.Quaternion());
-                    BoneIndices[i] = MappedBoneIndex;
-                }
-			}*/
-			if (AggGeom.ConvexElems.Num())
+
+#if !PLATFORM_LUMIN && !PLATFORM_ANDROID
+			// Add convexes
+			for (const FKConvexElem& ConvexElem : AggGeom.ConvexElems)
 			{
+				// Add stub for extracted collision data
+				FClothCollisionPrim_Convex Convex;
+				Convex.BoneIndex = MappedBoneIndex;
+#if WITH_PHYSX
 				// Collision bodies are stored in PhysX specific data structures so they can only be imported if we enable PhysX.
-#if WITH_PHYSX && !PLATFORM_LUMIN && !PLATFORM_ANDROID
-				uint32 OldSize = CollisionParticles.Size();
-				CollisionParticles.AddParticles(AggGeom.ConvexElems.Num());
-				for (uint32 i = OldSize; i < CollisionParticles.Size(); ++i)
+				const physx::PxConvexMesh* const PhysXMesh = ConvexElem.GetConvexMesh();
+				const int32 NumPolygons = int32(PhysXMesh->getNbPolygons());
+				Convex.Planes.SetNumUninitialized(NumPolygons);
+				for (int32 i = 0; i < NumPolygons; ++i)
 				{
-					const auto& CollisionBody = AggGeom.ConvexElems[i - OldSize];
-					const auto PhysXMesh = CollisionBody.GetConvexMesh();
-					const int32 NumPolygons = int32(PhysXMesh->getNbPolygons());
-
-					// Add stub for extracted collision data
-					FClothCollisionPrim_Convex NewConvex;
-					NewConvex.Planes.SetNumUninitialized(NumPolygons);
-					NewConvex.BoneIndex = MappedBoneIndex;
-
-					// Setup new convex particle
-					TArray<Chaos::TVector<int32, 3>> CollisionMeshElements;
-					for (int32 j = 0; j < NumPolygons; ++j)
-					{
-						physx::PxHullPolygon Poly;
-						PhysXMesh->getPolygonData(j, Poly);
-						check(Poly.mNbVerts == 3);
-						const auto Indices = PhysXMesh->getIndexBuffer() + Poly.mIndexBase;
-						CollisionMeshElements.Add(Chaos::TVector<int32, 3>(Indices[0], Indices[1], Indices[2]));
-
-						NewConvex.Planes[j] = FPlane(
-							CollisionBody.VertexData[Indices[0]],
-							CollisionBody.VertexData[Indices[1]],
-							CollisionBody.VertexData[Indices[2]]);
-					}
-					Chaos::TParticles<float, 3> CollisionMeshParticles;
-					CollisionMeshParticles.AddParticles(CollisionBody.VertexData.Num());
-					for (uint32 j = 0; j < CollisionMeshParticles.Size(); ++j)
-					{
-						CollisionMeshParticles.X(j) = CollisionBody.VertexData[j];
-					}
-					Chaos::TBox<float, 3> BoundingBox(CollisionMeshParticles.X(0), CollisionMeshParticles.X(0));
-					for (uint32 j = 1; j < CollisionMeshParticles.Size(); ++j)
-					{
-						BoundingBox.GrowToInclude(CollisionMeshParticles.X(i));
-					}
-					int32 MaxAxisSize = 100;
-					int32 MaxAxis;
-					const auto Extents = BoundingBox.Extents();
-					if (Extents[0] > Extents[1] && Extents[0] > Extents[2])
-					{
-						MaxAxis = 0;
-					}
-					else if (Extents[1] > Extents[2])
-					{
-						MaxAxis = 1;
-					}
-					else
-					{
-						MaxAxis = 2;
-					}
-					Chaos::TUniformGrid<float, 3> Grid(BoundingBox.Min(), BoundingBox.Max(), Chaos::TVector<int32, 3>(100 * Extents[0] / Extents[MaxAxis], 100 * Extents[0] / Extents[MaxAxis], 100 * Extents[0] / Extents[MaxAxis]));
-					Chaos::TTriangleMesh<float> CollisionMesh(MoveTemp(CollisionMeshElements));
-					Chaos::FErrorReporter ErrorReporter;
-					CollisionParticles.SetDynamicGeometry(i, MakeUnique<Chaos::TLevelSet<float, 3>>(ErrorReporter, Grid, CollisionMeshParticles, CollisionMesh));
-					BaseTransforms[i] = Chaos::TRigidTransform<float, 3>(Chaos::TVector<float, 3>(0.f), Chaos::TRotation<float, 3>::Identity);
-					BoneIndices[i] = MappedBoneIndex;
-
-					// Add extracted collision data
-					ExtractedCollisions.Convexes.Add(NewConvex);
+					physx::PxHullPolygon Poly;
+					PhysXMesh->getPolygonData(i, Poly);
+					check(Poly.mNbVerts == 3);
+					const auto Indices = PhysXMesh->getIndexBuffer() + Poly.mIndexBase;
+
+					Convex.Planes[i] = FPlane(
+						ConvexElem.VertexData[Indices[0]],
+						ConvexElem.VertexData[Indices[1]],
+						ConvexElem.VertexData[Indices[2]]);
 				}
-#endif
+#elif WITH_CHAOS  // #if WITH_PHYSX
+				const Chaos::TImplicitObject<float, 3>& ChaosConvexMesh = *ConvexElem.GetChaosConvexMesh();
+				const Chaos::TConvex<float, 3>& ChaosConvex = ChaosConvexMesh.GetObjectChecked<Chaos::TConvex<float, 3>>();
+				const TArray<TPlane<float, 3>>& Planes = ChaosConvex.GetFaces();
+				Convex.Planes.Reserve(Planes.Num());
+				for (const TPlane<float, 3>& Plane : Planes)
+				{
+					Convex.Planes.Add(FPlane(Plane.X(), Plane.Normal()));
+				}
+#endif  // #if WITH_PHYSX #elif WITH_CHAOS
+
+				// Add extracted collision data
+				ExtractedCollisions.Convexes.Add(Convex);
+#endif  // #if !PLATFORM_LUMIN && !PLATFORM_ANDROID
 			}
 		} // end for
+
+		// Add collisions particles
+		UE_LOG(LogChaosCloth, VeryVerbose, TEXT("Adding physics asset collisions..."));
+		AddCollisions(ExtractedCollisions, UsedBoneIndices);
 	} // end if PhysAsset
-	UE_LOG(LogChaosCloth, Verbose, TEXT("Added physics asset collisions: %d spheres, %d capsules, %d convexes, %d boxes."), ExtractedCollisions.Spheres.Num() - 2 * ExtractedCollisions.SphereConnections.Num(), ExtractedCollisions.SphereConnections.Num(), ExtractedCollisions.Convexes.Num(), ExtractedCollisions.Boxes.Num());
 }
 
 void ClothingSimulation::ExtractLegacyAssetCollisions(UClothingAssetCommon* Asset, const USkeletalMeshComponent* InOwnerComponent)
 {
-	check(Asset->GetNumLods() == 1);
-	const UClothLODDataBase* const AssetLodData = Asset->ClothLodData[0];
-	if (!AssetLodData) { return; }
+	check(Asset->GetNumLods() > 0);
+	ensure(Asset->GetNumLods() == 1);
+	UE_CLOG(Asset->GetNumLods() != 1,
+		LogChaosCloth, Warning, TEXT("More than one LOD with the current cloth asset. Only LOD 0 is supported with the current system."));
 
+	if (const UClothLODDataBase* const AssetLodData = Asset->ClothLodData[0])
+	{
+		const FClothCollisionData& LodCollData = AssetLodData->CollisionData;
+		if (LodCollData.Spheres.Num() || LodCollData.SphereConnections.Num() || LodCollData.Convexes.Num())
+		{
+			UE_LOG(LogChaosCloth, Warning,
+				TEXT("Actor '%s' component '%s' has %d sphere, %d capsule, and %d "
+					"convex collision objects for physics authored as part of a LOD construct, "
+					"probably by the Apex cloth authoring system.  This is deprecated.  "
+					"Please update your asset!"),
+				InOwnerComponent->GetOwner() ? *InOwnerComponent->GetOwner()->GetName() : TEXT("None"),
+				*InOwnerComponent->GetName(),
+				LodCollData.Spheres.Num(),
+				LodCollData.SphereConnections.Num(),
+				LodCollData.Convexes.Num());
+
+			UE_LOG(LogChaosCloth, VeryVerbose, TEXT("Adding legacy cloth asset collisions..."));
+			AddCollisions(LodCollData, Asset->UsedBoneIndices);
+		}
+	}
+}
+
+void ClothingSimulation::AddCollisions(const FClothCollisionData& ClothCollisionData, const TArray<int32>& UsedBoneIndices)
+{
 	TGeometryClothParticles<float, 3>& CollisionParticles = Evolution->CollisionParticles();
 
-	// We can't just use AddExternalCollisions() because we need to add entries for bone mappings and lookups.
-	const FClothCollisionData& LodCollData = AssetLodData->CollisionData;
-	if (LodCollData.Spheres.Num() || LodCollData.SphereConnections.Num() || LodCollData.Convexes.Num())
+	// Capsules
+	TSet<int32> CapsuleEnds;
+	const int32 NumCapsules = ClothCollisionData.SphereConnections.Num();
+	if (NumCapsules)
 	{
-		UE_LOG(LogSkeletalMesh, Warning,
-			TEXT("Actor '%s' component '%s' has %d sphere, %d capsule, and %d "
-				"convex collision objects for physics authored as part of a LOD construct, "
-				"probably by the Apex cloth authoring system.  This is deprecated.  "
-				"Please update your asset!"),
-			InOwnerComponent->GetOwner() ? *InOwnerComponent->GetOwner()->GetName() : TEXT("None"),
-			*InOwnerComponent->GetName(),
-			LodCollData.Spheres.Num(),
-			LodCollData.SphereConnections.Num(),
-			LodCollData.Convexes.Num());
-
-		TSet<int32> CapsuleEnds;
-		if (LodCollData.SphereConnections.Num())
+		const uint32 Offset = CollisionParticles.Size();
+		CollisionParticles.AddParticles(NumCapsules);
+
+		CapsuleEnds.Reserve(NumCapsules * 2);
+		for (uint32 i = Offset; i < CollisionParticles.Size(); ++i)
 		{
-			const uint32 Size = CollisionParticles.Size();
-			CollisionParticles.AddParticles(LodCollData.SphereConnections.Num());
-			CapsuleEnds.Reserve(LodCollData.SphereConnections.Num() * 2);
-			for (uint32 i = Size; i < CollisionParticles.Size(); ++i)
+			const FClothCollisionPrim_SphereConnection& Connection = ClothCollisionData.SphereConnections[i - Offset];
+
+			const int32 SphereIndex0 = Connection.SphereIndices[0];
+			const int32 SphereIndex1 = Connection.SphereIndices[1];
+			checkSlow(SphereIndex0 != SphereIndex1);
+			const FClothCollisionPrim_Sphere& Sphere0 = ClothCollisionData.Spheres[SphereIndex0];
+			const FClothCollisionPrim_Sphere& Sphere1 = ClothCollisionData.Spheres[SphereIndex1];
+
+			const int32 MappedIndex = UsedBoneIndices.IsValidIndex(Sphere0.BoneIndex) ? UsedBoneIndices[Sphere0.BoneIndex] : INDEX_NONE;
+
+			BoneIndices[i] = GetMappedBoneIndex(UsedBoneIndices, Sphere0.BoneIndex);
+			checkSlow(Sphere0.BoneIndex == Sphere1.BoneIndex);
+			UE_CLOG(Sphere0.BoneIndex != Sphere1.BoneIndex,
+				LogChaosCloth, Warning, TEXT("Found a legacy Apex cloth asset with a collision capsule spanning across two bones. This is not supported with the current system."));
+			UE_LOG(LogChaosCloth, VeryVerbose, TEXT("Found collision capsule on bone index %d."), BoneIndices[i]);
+
+			const Chaos::TVector<float, 3> X0 = Sphere0.LocalPosition;
+			const Chaos::TVector<float, 3> X1 = Sphere1.LocalPosition;
+
+			const float Radius0 = Sphere0.Radius;
+			const float Radius1 = Sphere1.Radius;
+
+			if (FMath::Abs(Radius0 - Radius1) < SMALL_NUMBER)
 			{
-				// This data was pulled from a FKSphylElem, which is a capsule.  So
-				// it should only have 1 radius, and the BoneIndex for both spheres
-				// should be the same.
-				const FClothCollisionPrim_SphereConnection& Connection = LodCollData.SphereConnections[i - Size];
-
-				const int32 SphereIndex0 = Connection.SphereIndices[0];
-				const int32 SphereIndex1 = Connection.SphereIndices[1];
-				checkSlow(SphereIndex0 != SphereIndex1);
-				const FClothCollisionPrim_Sphere& CollisionSphere0 = LodCollData.Spheres[SphereIndex0];
-				const FClothCollisionPrim_Sphere& CollisionSphere1 = LodCollData.Spheres[SphereIndex1];
-
-				const float Radius = CollisionSphere0.Radius;
-				checkSlow(CollisionSphere0.Radius - CollisionSphere1.Radius < SMALL_NUMBER);
-				UE_CLOG(CollisionSphere0.Radius - CollisionSphere1.Radius >= SMALL_NUMBER,
-					LogChaosCloth, Warning, TEXT("Found a legacy Apex cloth asset with a collision capsule of two different radii."));
-
-				const Chaos::TVector<float, 3> X0 = CollisionSphere0.LocalPosition;
-				const Chaos::TVector<float, 3> X1 = CollisionSphere1.LocalPosition;
-
-				BoneIndices[i] = CollisionSphere0.BoneIndex;
-				checkSlow(CollisionSphere0.BoneIndex == CollisionSphere1.BoneIndex);
-				UE_CLOG(CollisionSphere0.BoneIndex != CollisionSphere1.BoneIndex,
-					LogChaosCloth, Warning, TEXT("Found a legacy Apex cloth asset with a collision capsule spanning across two bones."));
-
-				// We construct a capsule centered at the origin along the Z axis
-				const Chaos::TVector<float, 3> Center = (X0 + X1) * 0.5f;
+				// Capsule
+				const Chaos::TVector<float, 3> Center = (X0 + X1) * 0.5f;  // Construct a capsule centered at the origin along the Z axis
 				const Chaos::TVector<float, 3> Axis = X1 - X0;
+				const Chaos::TRotation<float, 3> Rotation = Chaos::TRotation<float, 3>::FromRotatedVector(
+					Chaos::TVector<float, 3>::AxisVector(2),
+					Axis.GetSafeNormal());
 
-				const Chaos::TRotation<float, 3> Rotation = Chaos::TRotation<float, 3>::FromRotatedVector(Chaos::TVector<float, 3>::AxisVector(2), Axis.GetSafeNormal());
 				BaseTransforms[i] = Chaos::TRigidTransform<float, 3>(Center, Rotation);
 
 				const float HalfHeight = Axis.Size() * 0.5f;
@@ -817,69 +791,132 @@ void ClothingSimulation::ExtractLegacyAssetCollisions(UClothingAssetCommon* Asse
 					MakeUnique<Chaos::TCapsule<float>>(
 						Chaos::TVector<float, 3>(0.f, 0.f, -HalfHeight), // Min
 						Chaos::TVector<float, 3>(0.f, 0.f, HalfHeight), // Max
-						Radius));
-
-				// Skip spheres added as end caps for the capsule.
-				CapsuleEnds.Add(SphereIndex0);
-				CapsuleEnds.Add(SphereIndex1);
+						Radius0));
 			}
+			else
+			{
+				// Tapered capsule
+				TArray<TUniquePtr<TImplicitObject<float, 3>>> Objects;
+				Objects.Reserve(3);
+				Objects.Add(TUniquePtr<Chaos::TImplicitObject<float, 3>>(
+					new Chaos::TTaperedCylinder<float>(X0, X1, Radius0, Radius1)));
+				Objects.Add(TUniquePtr<Chaos::TImplicitObject<float, 3>>(
+					new Chaos::TSphere<float, 3>(X0, Radius0)));
+				Objects.Add(TUniquePtr<Chaos::TImplicitObject<float, 3>>(
+					new Chaos::TSphere<float, 3>(X1, Radius1)));
+				CollisionParticles.SetDynamicGeometry(
+					i,
+					MakeUnique<Chaos::TImplicitObjectUnion<float, 3>>(MoveTemp(Objects)));  // TODO(Kriss.Gossart): Replace this once a TTaperedCapsule implicit type is implemented
+			}
+
+			// Skip spheres added as end caps for the capsule.
+			CapsuleEnds.Add(SphereIndex0);
+			CapsuleEnds.Add(SphereIndex1);
 		}
-		if (LodCollData.Spheres.Num() - CapsuleEnds.Num())
+	}
+
+	// Spheres
+	const int32 NumSpheres = ClothCollisionData.Spheres.Num() - CapsuleEnds.Num();
+	if (NumSpheres != 0)
+	{
+		const uint32 Offset = CollisionParticles.Size();
+		CollisionParticles.AddParticles(NumSpheres);
+		// i = Spheres index, j = CollisionParticles index
+		for (uint32 i = 0, j = Offset; i < (uint32)ClothCollisionData.Spheres.Num(); ++i)
 		{
-			const uint32 Size = CollisionParticles.Size();
-			CollisionParticles.AddParticles(LodCollData.Spheres.Num() - CapsuleEnds.Num());
-			// i = Spheres index, j = CollisionParticles index
-			for (uint32 i = 0, j = Size; i < (uint32)LodCollData.Spheres.Num(); ++i)
-			{
-				// Skip spheres that are the end caps of capsules.
-				if (CapsuleEnds.Contains(i))
-					continue;
+			// Skip spheres that are the end caps of capsules.
+			if (CapsuleEnds.Contains(i))
+				continue;
 
-				const FClothCollisionPrim_Sphere& CollisionSphere = LodCollData.Spheres[i];
+			const FClothCollisionPrim_Sphere& Sphere = ClothCollisionData.Spheres[i];
 
-				BoneIndices[i] = CollisionSphere.BoneIndex;
+			BoneIndices[j] = GetMappedBoneIndex(UsedBoneIndices, Sphere.BoneIndex);
+			UE_LOG(LogChaosCloth, VeryVerbose, TEXT("Found collision sphere on bone index %d."), BoneIndices[i]);
 
-				BaseTransforms[j] = Chaos::TRigidTransform<float, 3>(CollisionSphere.LocalPosition, Chaos::TRotation<float, 3>::Identity);
+			BaseTransforms[j] = Chaos::TRigidTransform<float, 3>(FTransform::Identity);
 
-				CollisionParticles.SetDynamicGeometry(
-					j,
-					MakeUnique<Chaos::TSphere<float, 3>>(
-						Chaos::TVector<float, 3>(0.0f),
-						CollisionSphere.Radius));
-				++j;
-			}
-		}
-		// TODO(Kriss.Gossart): Convexes are missing (but boxes are a new addition, so they are not legacy)
+			CollisionParticles.SetDynamicGeometry(
+				j,
+				MakeUnique<Chaos::TSphere<float, 3>>(
+					Sphere.LocalPosition,
+					Sphere.Radius));
 
-		UE_LOG(LogChaosCloth, Verbose, TEXT("Added legacy asset collisions: %d spheres, %d capsules, %d convexes."), LodCollData.Spheres.Num() - CapsuleEnds.Num(), LodCollData.SphereConnections.Num(), LodCollData.Convexes.Num());
-	} // end if LodCollData
-}
+			++j;
+		}
+	}
 
-void ClothingSimulation::RefreshBoneMapping(UClothingAssetCommon* Asset)
-{
-	// No mesh, can't remap
-	if (const USkeletalMesh* const SkeletalMesh = CastChecked<USkeletalMesh>(Asset->GetOuter()))
+	// Convexes
+	const uint32 NumConvexes = ClothCollisionData.Convexes.Num();
+	if (NumConvexes != 0)
 	{
-		// Add the asset known used bone names (will take care of the apex collision mapping)
-		for (const FName& Name : Asset->UsedBoneNames)
+		const uint32 Offset = CollisionParticles.Size();
+		CollisionParticles.AddParticles(NumConvexes);
+		for (uint32 i = Offset; i < CollisionParticles.Size(); ++i)
 		{
-			CollisionBoneNames.AddUnique(Name);
-		}
+			const FClothCollisionPrim_Convex& Convex = ClothCollisionData.Convexes[i - Offset];
 
-		// New names added
-		if(CollisionBoneNames.Num() != CollisionBoneIndices.Num())
-		{
-			CollisionBoneIndices.Reset();
-			CollisionBoneIndices.AddDefaulted(CollisionBoneNames.Num());
+			BaseTransforms[i] = Chaos::TRigidTransform<float, 3>(FTransform::Identity);
+
+			BoneIndices[i] = GetMappedBoneIndex(UsedBoneIndices, Convex.BoneIndex);
+			UE_LOG(LogChaosCloth, VeryVerbose, TEXT("Found collision convex on bone index %d."), BoneIndices[i]);
+
+			// Add convex planes
+			TArray<TUniquePtr<TImplicitObject<float, 3>>> Planes;
+			Planes.Reserve(Convex.Planes.Num());
+			for (const FPlane& Plane : Convex.Planes)
+			{
+				FPlane NormalizedPlane(Plane);
+				if (NormalizedPlane.Normalize())
+				{
+					const Chaos::TVector<float, 3> Normal(static_cast<FVector>(NormalizedPlane));
+					const Chaos::TVector<float, 3> Base = Normal * NormalizedPlane.W;
+
+					Planes.Add(TUniquePtr<Chaos::TImplicitObject<float, 3>>(
+						new Chaos::TPlane<float, 3>(Base, Normal)));
+				}
+				else
+				{
+					UE_LOG(LogChaosCloth, Warning, TEXT("Invalid convex collision plane."));
+				}
+			}
+
+			// Set the collision particle geometry
+			if (Planes.Num())
+			{
+				CollisionParticles.SetDynamicGeometry(i, MakeUnique<Chaos::TImplicitObjectIntersection<float, 3>>(MoveTemp(Planes)));  // TODO(Kriss.Gossart): Replace by TConvex once plane constructor has been added
+			}
+			else
+			{
+				UE_LOG(LogChaosCloth, Warning, TEXT("Invalid convex collision."));
+				CollisionParticles.SetDynamicGeometry(
+					i,
+					MakeUnique<Chaos::TSphere<float, 3>>(Chaos::TVector<float, 3>(0.0f), 1.0f));  // Default to a unit sphere to replace the faulty convex
+			}
 		}
+	}
 
-		// Repopulate the used indices
-		for (int32 BoneNameIndex = 0; BoneNameIndex < CollisionBoneNames.Num(); ++BoneNameIndex)
+	// Boxes
+	const uint32 NumBoxes = ClothCollisionData.Boxes.Num();
+	if (NumBoxes != 0)
+	{
+		const uint32 Offset = CollisionParticles.Size();
+		CollisionParticles.AddParticles(NumBoxes);
+		for (uint32 i = Offset; i < CollisionParticles.Size(); ++i)
 		{
-			const FName& UsedBoneName = CollisionBoneNames[BoneNameIndex];
-			CollisionBoneIndices[BoneNameIndex] = SkeletalMesh->RefSkeleton.FindBoneIndex(UsedBoneName);
+			const FClothCollisionPrim_Box& Box = ClothCollisionData.Boxes[i - Offset];
+			CollisionParticles.X(i) = Chaos::TVector<float, 3>(0.f);
+			CollisionParticles.R(i) = Chaos::TRotation<float, 3>::FromIdentity();
+			
+			BaseTransforms[i] = Chaos::TRigidTransform<float, 3>(Box.LocalPosition, Box.LocalRotation);
+			
+			BoneIndices[i] = GetMappedBoneIndex(UsedBoneIndices, Box.BoneIndex);
+			UE_LOG(LogChaosCloth, VeryVerbose, TEXT("Found collision box on bone index %d."), BoneIndices[i]);
+
+			CollisionParticles.SetDynamicGeometry(i, MakeUnique<Chaos::TBox<float, 3>>(-Box.HalfExtents, Box.HalfExtents));
 		}
 	}
+
+	UE_LOG(LogChaosCloth, VeryVerbose, TEXT("Added collisions: %d spheres, %d capsules, %d convexes, %d boxes."), NumSpheres, NumCapsules, NumConvexes, NumBoxes);
 }
 
 void ClothingSimulation::FillContext(USkeletalMeshComponent* InComponent, float InDeltaTime, IClothingSimulationContext* InOutContext)
@@ -992,27 +1029,7 @@ void ClothingSimulation::Simulate(IClothingSimulationContext* InContext)
 	}
 
 	// Update collision tranforms
-	TGeometryClothParticles<float, 3>& CollisionParticles = Evolution->CollisionParticles();
-	for (uint32 i = 0; i < CollisionParticles.Size(); ++i)
-	{
-		const int32 BoneIndex = BoneIndices[i];
-		const int32 MappedIndex = CollisionBoneIndices.IsValidIndex(BoneIndex) ? CollisionBoneIndices[BoneIndex] : INDEX_NONE;
-		if (Context->BoneTransforms.IsValidIndex(MappedIndex))
-		{
-			const FTransform& BoneTransform = Context->BoneTransforms[MappedIndex];
-			CollisionTransforms[i] = BaseTransforms[i] * BoneTransform * Context->ComponentToWorld;
-		}
-		else
-		{
-			CollisionTransforms[i] = BaseTransforms[i] * Context->ComponentToWorld;  // External colllisions often don't map to a bone
-		}
-	}
-
-	// Make sure external collision have a previous transform  // TODO(Kriss.Gossart): This is a temporary fix and needs changing. With removing/re-adding external collision at every frame there's no transform history. 
-	for (uint32 i = ExternalCollisionsOffset; i < CollisionParticles.Size(); ++i)
-	{
-		OldCollisionTransforms[i] = CollisionTransforms[i];
-	}
+	UpdateCollisionTransforms(*Context);
 
 	// Advance Sim
 	DeltaTime = Context->DeltaTime;
@@ -1079,137 +1096,21 @@ void ClothingSimulation::GetSimulationData(
 
 void ClothingSimulation::AddExternalCollisions(const FClothCollisionData& InData)
 {
-	// Keep track of the external collisions added
+	// Keep track of the external collisions data
 	ExternalCollisions.Append(InData);
 
-	// Add particles
-	TGeometryClothParticles<float, 3>& CollisionParticles = Evolution->CollisionParticles();
-
-	TSet<int32> CapsuleEnds;
-
-	if (InData.SphereConnections.Num())
-	{
-		const uint32 Size = CollisionParticles.Size();
-        CollisionParticles.AddParticles(InData.SphereConnections.Num());
-		CapsuleEnds.Reserve(InData.SphereConnections.Num() * 2);
-        for (uint32 i = Size; i < CollisionParticles.Size(); ++i)
-        {
-			// This data was pulled from a FKSphylElem, which is a capsule.  So
-			// it should only have 1 radius, and the BoneIndex for both spheres
-			// should be the same.
-			const FClothCollisionPrim_SphereConnection& Connection = InData.SphereConnections[i - Size];
-			const int32 SphereIndex0 = Connection.SphereIndices[0];
-			const int32 SphereIndex1 = Connection.SphereIndices[1];
-			checkSlow(SphereIndex0 != SphereIndex1);
-			const float Radius = InData.Spheres[SphereIndex0].Radius;
-			checkSlow(InData.Spheres[SphereIndex0].Radius - InData.Spheres[SphereIndex1].Radius < SMALL_NUMBER);
-			const Chaos::TVector<float, 3> X0 = InData.Spheres[SphereIndex0].LocalPosition;
-			const Chaos::TVector<float, 3> X1 = InData.Spheres[SphereIndex1].LocalPosition;
-			const int32 BoneIndex = InData.Spheres[SphereIndex0].BoneIndex;
-			checkSlow(InData.Spheres[SphereIndex0].BoneIndex == InData.Spheres[SphereIndex1].BoneIndex);
-
-			const Chaos::TVector<float, 3> Center = (X0 + X1) * 0.5;
-			const Chaos::TVector<float, 3> Axis = X1 - X0;
-			const float HalfHeight = Axis.Size() * 0.5;
-
-			// We construct a capsule centered at the origin along the Z axis, and
-			// then move it into place with X and R.
-			CollisionParticles.X(i) = Center;
-			CollisionParticles.R(i) = Chaos::TRotation<float, 3>::FromRotatedVector(Chaos::TVector<float,3>::AxisVector(2), Axis.GetSafeNormal());
-			BaseTransforms[i] = Chaos::TRigidTransform<float, 3>(CollisionParticles.X(i), CollisionParticles.R(i));
-			BoneIndices[i] = BoneIndex;
-
-			CollisionParticles.SetDynamicGeometry(
-				i,
-				MakeUnique<Chaos::TCapsule<float>>(
-					Chaos::TVector<float, 3>(0.f, 0.f, -HalfHeight), // Min
-					Chaos::TVector<float, 3>(0.f, 0.f, HalfHeight), // Max
-					Radius));
-
-			// Skip spheres added as end caps for the capsule.
-			CapsuleEnds.Add(SphereIndex0);
-			CapsuleEnds.Add(SphereIndex1);
-		}
-	}
-
-	if (InData.Spheres.Num() - CapsuleEnds.Num())
-	{
-		const uint32 Size = CollisionParticles.Size();
-		CollisionParticles.AddParticles(InData.Spheres.Num() - CapsuleEnds.Num());
-		for (uint32 i = 0, j = Size; i < (uint32)InData.Spheres.Num(); ++i)
-		{
-			// Skip spheres that are the end caps of capsules.
-			if (CapsuleEnds.Contains(i))
-				continue;
-
-			const FClothCollisionPrim_Sphere& CollisionSphere = InData.Spheres[i];
-			CollisionParticles.X(j) = Chaos::TVector<float, 3>(0.f);
-			CollisionParticles.R(j) = Chaos::TRotation<float, 3>::FromIdentity(); 
-			BaseTransforms[i] = Chaos::TRigidTransform<float, 3>(FTransform::Identity);
-			BoneIndices[i] = CollisionSphere.BoneIndex;
-			CollisionParticles.SetDynamicGeometry(
-				j,
-				MakeUnique<Chaos::TSphere<float, 3>>(
-					CollisionSphere.LocalPosition, //Chaos::TVector<float, 3>(0.f, 0.f, 0.f),
-					CollisionSphere.Radius));
-
-			++j;
-		}
-	}
-
-	if (InData.Convexes.Num())
-	{
-		const uint32 Size = CollisionParticles.Size();
-		CollisionParticles.AddParticles(InData.Convexes.Num());
-		for (uint32 i = Size; i < CollisionParticles.Size(); ++i)
-		{
-			const FClothCollisionPrim_Convex& Convex = InData.Convexes[i - Size];
-			CollisionParticles.X(i) = Chaos::TVector<float, 3>(0.f);
-			CollisionParticles.R(i) = Chaos::TRotation<float, 3>::FromIdentity();
-			BaseTransforms[i] = Chaos::TRigidTransform<float, 3>(FTransform::Identity);
-			BoneIndices[i] = Convex.BoneIndex;
-			TArray<TUniquePtr<TImplicitObject<float, 3>>> Planes;
-			for (int32 j = 0; j < Convex.Planes.Num(); ++j)
-			{
-				Planes.Add(TUniquePtr<Chaos::TImplicitObject<float, 3>>(
-					new Chaos::TPlane<float, 3>(Chaos::TVector<float, 3>(0.f, 0.f, Convex.Planes[j].W / Convex.Planes[j].Z),
-					Chaos::TVector<float, 3>(Convex.Planes[j].X, Convex.Planes[j].Y, Convex.Planes[j].Z))));
-			}
-			CollisionParticles.SetDynamicGeometry(i, MakeUnique<Chaos::TImplicitObjectIntersection<float, 3>>(MoveTemp(Planes)));
-		}
-	}
-
-	if (InData.Boxes.Num())
-	{
-		const uint32 Size = CollisionParticles.Size();
-		CollisionParticles.AddParticles(InData.Convexes.Num());
-		for (uint32 i = Size; i < CollisionParticles.Size(); ++i)
-		{
-			const FClothCollisionPrim_Box& Box = InData.Boxes[i - Size];
-			CollisionParticles.X(i) = Chaos::TVector<float, 3>(0.f);
-			CollisionParticles.R(i) = Chaos::TRotation<float, 3>::FromIdentity();
-			BaseTransforms[i] = Chaos::TRigidTransform<float, 3>(FTransform::Identity);
-			BoneIndices[i] = Box.BoneIndex;
-			CollisionParticles.SetDynamicGeometry(i, MakeUnique<Chaos::TBox<float, 3>>(Box.LocalMin, Box.LocalMax));
-		}
-	}
-
-	check(CollisionParticles.Size() == BaseTransforms.Num());
-
-	const uint32 PrevCollisionTransformsCount = CollisionTransforms.Num();
-	const uint32 NewCollisionTransformsCount = BaseTransforms.Num();
-
-	CollisionTransforms.SetNum(NewCollisionTransformsCount);
-	OldCollisionTransforms.SetNum(NewCollisionTransformsCount);
-	
-	UE_LOG(LogChaosCloth, VeryVerbose, TEXT("Added external collisions: %d spheres, %d capsules, %d convexes, %d boxes."), InData.Spheres.Num() - CapsuleEnds.Num(), InData.SphereConnections.Num(), InData.Convexes.Num(), InData.Boxes.Num());
+	// Setup the new collisions particles
+	UE_LOG(LogChaosCloth, VeryVerbose, TEXT("Adding external collisions..."));
+	const TArray<int32> UsedBoneIndices;  // There is no bone mapping available for external collisions
+	AddCollisions(InData, UsedBoneIndices);
 }
 
 void ClothingSimulation::ClearExternalCollisions()
 {
 	// Remove all external collision particles, starting from the external collision offset
+	// But do not resize CollisionTransforms as it is only resized in UpdateCollisionTransforms()
 	TGeometryClothParticles<float, 3>& CollisionParticles = Evolution->CollisionParticles();
-	CollisionParticles.Resize(ExternalCollisionsOffset);
+	CollisionParticles.Resize(ExternalCollisionsOffset);  // This will also resize BoneIndices and BaseTransforms
 
 	// Reset external collisions
 	ExternalCollisionsOffset = CollisionParticles.Size();
@@ -1232,7 +1133,7 @@ void ClothingSimulation::GetCollisions(FClothCollisionData& OutCollisions, bool
 	}
 
 	// Add collisions extracted from the physics asset
-	// TODO: Including the following code seems to be the correct behaviour, but this did not appear
+	// TODO(Kriss.Gossart): Including the following code seems to be the correct behaviour, but this did not appear
 	// in the NvCloth implementation, so best to leave it commented out for now.
 	//OutCollisions.Append(ExtractedCollisions);
 
@@ -1430,10 +1331,9 @@ void ClothingSimulation::DebugDrawCollision(USkeletalMeshComponent* OwnerCompone
 	{
 		const TUniquePtr<TImplicitObject<float, 3>>& DynamicGeometry = CollisionParticles.DynamicGeometry(Index);
 		const uint32 BoneIndex = BoneIndices[Index];
-		const int32 MappedIndex = CollisionBoneIndices.IsValidIndex(BoneIndex) ? CollisionBoneIndices[BoneIndex] : INDEX_NONE;
-		const FLinearColor Color = (MappedIndex != INDEX_NONE) ? MappedColor : UnmappedColor;
+		const FLinearColor Color = (BoneIndex != INDEX_NONE) ? MappedColor : UnmappedColor;
 
-		const Chaos::TVector<float, 3>& Center = CollisionParticles.X(Index);
+		const Chaos::TVector<float, 3>& Position = CollisionParticles.X(Index);
 		const TRotation<float, 3>& Rotation = CollisionParticles.R(Index);
 
 		switch (DynamicGeometry->GetType())
@@ -1443,7 +1343,8 @@ void ClothingSimulation::DebugDrawCollision(USkeletalMeshComponent* OwnerCompone
 			if (const Chaos::TSphere<float, 3>* const Sphere = DynamicGeometry->GetObject<Chaos::TSphere<float, 3>>())
 			{
 				const float Radius = Sphere->GetRadius();
-				const FTransform Transform(Rotation, Center);
+				const Chaos::TVector<float, 3> Center = Sphere->GetCenter();
+				const FTransform Transform(Rotation, Position + Rotation.RotateVector(Center));
 				DrawWireSphere(PDI, Transform, Color, Radius, 12, SDPG_World, 0.0f, 0.001f, false);
 			}
 			break;
@@ -1469,13 +1370,13 @@ void ClothingSimulation::DebugDrawCollision(USkeletalMeshComponent* OwnerCompone
 				const FVector Y = Rotation.RotateVector(FVector::RightVector);
 				const FVector Z = Rotation.RotateVector(FVector::UpVector);
 
-				DrawWireCapsule(PDI, Center, X, Y, Z, Color, Radius, HalfHeight + Radius, 12, SDPG_World, 0.0f, 0.001f, false);
+				DrawWireCapsule(PDI, Position, X, Y, Z, Color, Radius, HalfHeight + Radius, 12, SDPG_World, 0.0f, 0.001f, false);
 			}
 			break;
 
 		// Draw everything else as a coordinate for now
 		default:
-			DrawCoordinateSystem(PDI, Center, FRotator(Rotation), 10.0f, SDPG_World, 0.1f);
+			DrawCoordinateSystem(PDI, Position, FRotator(Rotation), 10.0f, SDPG_World, 0.1f);
 			break;
 		}
 	}
diff --git a/Engine/Plugins/Experimental/ChaosCloth/Source/Chaos/Public/ChaosCloth/ChaosClothingSimulation.h b/Engine/Plugins/Experimental/ChaosCloth/Source/Chaos/Public/ChaosCloth/ChaosClothingSimulation.h
--- a/Engine/Plugins/Experimental/ChaosCloth/Source/Chaos/Public/ChaosCloth/ChaosClothingSimulation.h
+++ b/Engine/Plugins/Experimental/ChaosCloth/Source/Chaos/Public/ChaosCloth/ChaosClothingSimulation.h
@@ -12,7 +12,7 @@
 
 namespace Chaos
 {
-	class ClothingSimulationContext : public IClothingSimulationContext
+	class ClothingSimulationContext : public IClothingSimulationContext  // TODO(Kriss.Gossart): Check whether this should inherit from FClothingSimulationContextBase
 	{
 	public:
 		ClothingSimulationContext() {}
@@ -78,8 +78,15 @@ namespace Chaos
 		void ExtractPhysicsAssetCollisions(UClothingAssetCommon* Asset);
 		// Extract the collisions from the specified clothing asset 
 		void ExtractLegacyAssetCollisions(UClothingAssetCommon* Asset, const USkeletalMeshComponent* InOwnerComponent);
-		// Create the bone mapping for all the bones used by collision
-		void RefreshBoneMapping(UClothingAssetCommon* Asset);
+		// Add collisions from a ClothCollisionData structure
+		void AddCollisions(const FClothCollisionData& ClothCollisionData, const TArray<int32>& UsedBoneIndices);
+		// Update the collision transforms using the specified context
+		void UpdateCollisionTransforms(const ClothingSimulationContext& Context);
+		// Return the correct bone index based on the asset used bone index array
+		FORCEINLINE int32 GetMappedBoneIndex(const TArray<int32>& UsedBoneIndices, int32 BoneIndex)
+		{
+			return UsedBoneIndices.IsValidIndex(BoneIndex) ? UsedBoneIndices[BoneIndex] : INDEX_NONE;
+		}
 
 	private:
 		// Assets
@@ -88,8 +95,6 @@ namespace Chaos
 		// Collision Data
 		FClothCollisionData ExtractedCollisions;  // Collisions extracted from the referenced physics asset
 		FClothCollisionData ExternalCollisions;  // External collisions
-		TArray<FName> CollisionBoneNames;  // List of the bone names used by collision
-		TArray<int32> CollisionBoneIndices;  // List of the indices for the bones in UsedBoneNames, used for remapping
 		TArray<Chaos::TRigidTransform<float, 3>> OldCollisionTransforms;  // Used for the kinematic collision transform update
 		TArray<Chaos::TRigidTransform<float, 3>> CollisionTransforms;  // Used for the kinematic collision transform update
 		Chaos::TArrayCollectionArray<int32> BoneIndices;
diff --git a/Engine/Source/Runtime/ClothingSystemRuntimeInterface/Public/ClothCollisionPrim.h b/Engine/Source/Runtime/ClothingSystemRuntimeInterface/Public/ClothCollisionPrim.h
--- a/Engine/Source/Runtime/ClothingSystemRuntimeInterface/Public/ClothCollisionPrim.h
+++ b/Engine/Source/Runtime/ClothingSystemRuntimeInterface/Public/ClothCollisionPrim.h
@@ -76,18 +76,21 @@ struct FClothCollisionPrim_Box
 	GENERATED_BODY()
 
 	FClothCollisionPrim_Box()
-		: BoneIndex(INDEX_NONE)
-		, LocalMin(FVector::ZeroVector)
-		, LocalMax(FVector::ZeroVector)
+		: LocalPosition(FVector::ZeroVector)
+		, LocalRotation(FQuat::Identity)
+		, HalfExtents(FVector::ZeroVector)
+		, BoneIndex(INDEX_NONE)
 	{}
 
 	UPROPERTY()
-	int32 BoneIndex;
+	FVector LocalPosition;
 
 	UPROPERTY()
-	FVector LocalMin;
+	FQuat LocalRotation;
 
 	UPROPERTY()
-	FVector LocalMax;
-};
+	FVector HalfExtents;
 
+	UPROPERTY()
+	int32 BoneIndex;
+};
