Adjust comments around the Chaos Cloth convex Android undefined code block.

diff --git a/Engine/Plugins/Experimental/ChaosCloth/Source/Chaos/Private/ChaosCloth/ChaosClothingSimulation.cpp b/Engine/Plugins/Experimental/ChaosCloth/Source/Chaos/Private/ChaosCloth/ChaosClothingSimulation.cpp
--- a/Engine/Plugins/Experimental/ChaosCloth/Source/Chaos/Private/ChaosCloth/ChaosClothingSimulation.cpp
+++ b/Engine/Plugins/Experimental/ChaosCloth/Source/Chaos/Private/ChaosCloth/ChaosClothingSimulation.cpp
@@ -661,7 +661,7 @@ void ClothingSimulation::ExtractPhysicsAssetCollisions(UClothingAssetCommon* Ass
 				}
 			}
 
-#if !PLATFORM_LUMIN && !PLATFORM_ANDROID
+#if !PLATFORM_LUMIN && !PLATFORM_ANDROID  // TODO(Kriss.Gossart): Compile on Android and fix whatever errors the following code is causing
 			// Add convexes
 			for (const FKConvexElem& ConvexElem : AggGeom.ConvexElems)
 			{
@@ -700,12 +700,14 @@ void ClothingSimulation::ExtractPhysicsAssetCollisions(UClothingAssetCommon* Ass
 				ExtractedCollisions.Convexes.Add(Convex);
 			}
 #endif  // #if !PLATFORM_LUMIN && !PLATFORM_ANDROID
-		} // end for
+
+		}  // End for PhysAsset->SkeletalBodySetups
 
 		// Add collisions particles
 		UE_LOG(LogChaosCloth, VeryVerbose, TEXT("Adding physics asset collisions..."));
 		AddCollisions(ExtractedCollisions, UsedBoneIndices);
-	} // end if PhysAsset
+
+	}  // End if Asset->PhysicsAsset
 }
 
 void ClothingSimulation::ExtractLegacyAssetCollisions(UClothingAssetCommon* Asset, const USkeletalMeshComponent* InOwnerComponent)
