Fix Chaos Cloth box collision visualisation transform.

diff --git a/Engine/Plugins/Experimental/ChaosCloth/Source/Chaos/Private/ChaosCloth/ChaosClothingSimulation.cpp b/Engine/Plugins/Experimental/ChaosCloth/Source/Chaos/Private/ChaosCloth/ChaosClothingSimulation.cpp
--- a/Engine/Plugins/Experimental/ChaosCloth/Source/Chaos/Private/ChaosCloth/ChaosClothingSimulation.cpp
+++ b/Engine/Plugins/Experimental/ChaosCloth/Source/Chaos/Private/ChaosCloth/ChaosClothingSimulation.cpp
@@ -1355,7 +1355,7 @@ void ClothingSimulation::DebugDrawCollision(USkeletalMeshComponent* OwnerCompone
 		case Chaos::ImplicitObjectType::Box:
 			if (const Chaos::TBox<float, 3>* const Box = DynamicGeometry->GetObject<Chaos::TBox<float, 3>>())
 			{
-				const FMatrix BoxToWorld(Rotation.ToMatrix());
+				const FMatrix BoxToWorld = FTransform(Rotation, Position).ToMatrixNoScale();
 				const FVector Radii = Box->Extents() * 0.5f;
 				DrawWireBox(PDI, BoxToWorld, FBox(Box->Min(), Box->Max()), Color, SDPG_World, 0.0f, 0.001f, false);
 			}
