Add Chaos Cloth tapered capsule collision debug visualization. Also fix tapered capsule pole orientation.

diff --git a/Engine/Plugins/Experimental/ChaosCloth/Source/Chaos/Private/ChaosCloth/ChaosClothingSimulation.cpp b/Engine/Plugins/Experimental/ChaosCloth/Source/Chaos/Private/ChaosCloth/ChaosClothingSimulation.cpp
--- a/Engine/Plugins/Experimental/ChaosCloth/Source/Chaos/Private/ChaosCloth/ChaosClothingSimulation.cpp
+++ b/Engine/Plugins/Experimental/ChaosCloth/Source/Chaos/Private/ChaosCloth/ChaosClothingSimulation.cpp
@@ -646,8 +646,8 @@ void ClothingSimulation::ExtractPhysicsAssetCollisions(UClothingAssetCommon* Ass
 					FClothCollisionPrim_Sphere Sphere1;
 					const FVector OrientedDirection = TaperedCapsuleElem.Rotation.RotateVector(FVector::UpVector);
 					const FVector HalfDim = OrientedDirection * (TaperedCapsuleElem.Length / 2.f);
-					Sphere0.LocalPosition = TaperedCapsuleElem.Center - HalfDim;
-					Sphere1.LocalPosition = TaperedCapsuleElem.Center + HalfDim;
+					Sphere0.LocalPosition = TaperedCapsuleElem.Center + HalfDim;
+					Sphere1.LocalPosition = TaperedCapsuleElem.Center - HalfDim;
 					Sphere0.Radius = TaperedCapsuleElem.Radius0;
 					Sphere1.Radius = TaperedCapsuleElem.Radius1;
 					Sphere0.BoneIndex = MappedBoneIndex;
@@ -1324,6 +1324,63 @@ void ClothingSimulation::DebugDrawInversedFaceNormals(USkeletalMeshComponent* Ow
 
 void ClothingSimulation::DebugDrawCollision(USkeletalMeshComponent* OwnerComponent, FPrimitiveDrawInterface* PDI) const
 {
+	auto DrawSphere = [&PDI](const Chaos::TSphere<float, 3>& Sphere, const TRotation<float, 3>& Rotation, const Chaos::TVector<float, 3>& Position, const FLinearColor& Color)
+	{
+		const float Radius = Sphere.GetRadius();
+		const Chaos::TVector<float, 3> Center = Sphere.GetCenter();
+		const FTransform Transform(Rotation, Position + Rotation.RotateVector(Center));
+		DrawWireSphere(PDI, Transform, Color, Radius, 12, SDPG_World, 0.0f, 0.001f, false);
+	};
+
+	auto DrawBox = [&PDI](const Chaos::TBox<float, 3>& Box, const TRotation<float, 3>& Rotation, const Chaos::TVector<float, 3>& Position, const FLinearColor& Color)
+	{
+		const FMatrix BoxToWorld = FTransform(Rotation, Position).ToMatrixNoScale();
+		const FVector Radii = Box.Extents() * 0.5f;
+		DrawWireBox(PDI, BoxToWorld, FBox(Box.Min(), Box.Max()), Color, SDPG_World, 0.0f, 0.001f, false);
+	};
+
+	auto DrawCapsule = [&PDI](const Chaos::TCapsule<float>& Capsule, const TRotation<float, 3>& Rotation, const Chaos::TVector<float, 3>& Position, const FLinearColor& Color)
+	{
+		const float HalfHeight = Capsule.GetHeight() * 0.5f;
+		const float Radius = Capsule.GetRadius();
+		const FVector X = Rotation.RotateVector(FVector::ForwardVector);
+		const FVector Y = Rotation.RotateVector(FVector::RightVector);
+		const FVector Z = Rotation.RotateVector(FVector::UpVector);
+		DrawWireCapsule(PDI, Position, X, Y, Z, Color, Radius, HalfHeight + Radius, 12, SDPG_World, 0.0f, 0.001f, false);
+	};
+
+	auto DrawTaperedCylinder = [&PDI](const Chaos::TTaperedCylinder<float>& TaperedCylinder, const TRotation<float, 3>& Rotation, const Chaos::TVector<float, 3>& Position, const FLinearColor& Color)
+	{
+		const float HalfHeight = TaperedCylinder.GetHeight() * 0.5f;
+		const float Radius1 = TaperedCylinder.GetRadius1();
+		const float Radius2 = TaperedCylinder.GetRadius2();
+		const FVector Position1 = Position + Rotation.RotateVector(TaperedCylinder.GetX1());
+		const FVector Position2 = Position + Rotation.RotateVector(TaperedCylinder.GetX2());
+		const FQuat Q = (Position2 - Position1).ToOrientationQuat();
+		const FVector I = Q.GetRightVector();
+		const FVector J = Q.GetUpVector();
+
+		static const int32 NumSides = 12;
+		static const float	AngleDelta = 2.0f * PI / NumSides;
+		FVector	LastVertex1 = Position1 + I * Radius1;
+		FVector	LastVertex2 = Position2 + I * Radius2;
+
+		for (int32 SideIndex = 1; SideIndex <= NumSides; ++SideIndex)
+		{
+			const float Angle = AngleDelta * float(SideIndex);
+			const FVector ArcPos = I * FMath::Cos(Angle) + J * FMath::Sin(Angle);
+			const FVector Vertex1 = Position1 + ArcPos * Radius1;
+			const FVector Vertex2 = Position2 + ArcPos * Radius2;
+
+			PDI->DrawLine(LastVertex1, Vertex1, Color, SDPG_World, 0.0f, 0.001f, false);
+			PDI->DrawLine(LastVertex2, Vertex2, Color, SDPG_World, 0.0f, 0.001f, false);
+			PDI->DrawLine(LastVertex1, LastVertex2, Color, SDPG_World, 0.0f, 0.001f, false);
+
+			LastVertex1 = Vertex1;
+			LastVertex2 = Vertex2;
+		}
+	};
+
 	static const FLinearColor MappedColor(FColor::Cyan);
 	static const FLinearColor UnmappedColor(FColor::Red);
 
@@ -1331,55 +1388,54 @@ void ClothingSimulation::DebugDrawCollision(USkeletalMeshComponent* OwnerCompone
 
 	for (uint32 Index = 0; Index < CollisionParticles.Size(); ++Index)
 	{
-		const TUniquePtr<TImplicitObject<float, 3>>& DynamicGeometry = CollisionParticles.DynamicGeometry(Index);
-		const uint32 BoneIndex = BoneIndices[Index];
-		const FLinearColor Color = (BoneIndex != INDEX_NONE) ? MappedColor : UnmappedColor;
-
-		const Chaos::TVector<float, 3>& Position = CollisionParticles.X(Index);
-		const TRotation<float, 3>& Rotation = CollisionParticles.R(Index);
-
-		switch (DynamicGeometry->GetType())
+		if (const TImplicitObject<float, 3>* const Object = CollisionParticles.DynamicGeometry(Index).Get())
 		{
-		// Draw collision spheres
-		case Chaos::ImplicitObjectType::Sphere:
-			if (const Chaos::TSphere<float, 3>* const Sphere = DynamicGeometry->GetObject<Chaos::TSphere<float, 3>>())
-			{
-				const float Radius = Sphere->GetRadius();
-				const Chaos::TVector<float, 3> Center = Sphere->GetCenter();
-				const FTransform Transform(Rotation, Position + Rotation.RotateVector(Center));
-				DrawWireSphere(PDI, Transform, Color, Radius, 12, SDPG_World, 0.0f, 0.001f, false);
-			}
-			break;
+			const uint32 BoneIndex = BoneIndices[Index];
+			const FLinearColor Color = (BoneIndex != INDEX_NONE) ? MappedColor : UnmappedColor;
 
-		// Draw collision boxes
-		case Chaos::ImplicitObjectType::Box:
-			if (const Chaos::TBox<float, 3>* const Box = DynamicGeometry->GetObject<Chaos::TBox<float, 3>>())
-			{
-				const FMatrix BoxToWorld = FTransform(Rotation, Position).ToMatrixNoScale();
-				const FVector Radii = Box->Extents() * 0.5f;
-				DrawWireBox(PDI, BoxToWorld, FBox(Box->Min(), Box->Max()), Color, SDPG_World, 0.0f, 0.001f, false);
-			}
-			break;
+			const Chaos::TVector<float, 3>& Position = CollisionParticles.X(Index);
+			const TRotation<float, 3>& Rotation = CollisionParticles.R(Index);
 
-		// Draw collision capsules
-		case Chaos::ImplicitObjectType::Capsule:
-			if (const Chaos::TCapsule<float>* const Capsule = DynamicGeometry->GetObject<Chaos::TCapsule<float>>())
+			switch (Object->GetType())
 			{
-				const float HalfHeight = Capsule->GetHeight() * 0.5f;
-				const float Radius = Capsule->GetRadius();
+			case Chaos::ImplicitObjectType::Sphere:
+				DrawSphere(Object->GetObjectChecked<Chaos::TSphere<float, 3>>(), Rotation, Position, Color);
+				break;
 
-				const FVector X = Rotation.RotateVector(FVector::ForwardVector);
-				const FVector Y = Rotation.RotateVector(FVector::RightVector);
-				const FVector Z = Rotation.RotateVector(FVector::UpVector);
+			case Chaos::ImplicitObjectType::Box:
+				DrawBox(Object->GetObjectChecked<Chaos::TBox<float, 3>>(), Rotation, Position, Color);
+				break;
 
-				DrawWireCapsule(PDI, Position, X, Y, Z, Color, Radius, HalfHeight + Radius, 12, SDPG_World, 0.0f, 0.001f, false);
-			}
-			break;
+			case Chaos::ImplicitObjectType::Capsule:
+				DrawCapsule(Object->GetObjectChecked<Chaos::TCapsule<float>>(), Rotation, Position, Color);
+				break;
 
-		// Draw everything else as a coordinate for now
-		default:
-			DrawCoordinateSystem(PDI, Position, FRotator(Rotation), 10.0f, SDPG_World, 0.1f);
-			break;
+			case Chaos::ImplicitObjectType::Union:  // Union only used as collision tappered capsules
+				for (const TUniquePtr<TImplicitObject<float, 3>>& SubObjectPtr : Object->GetObjectChecked<Chaos::TImplicitObjectUnion<float, 3>>().GetObjects())
+				{
+					if (const TImplicitObject<float, 3>* const SubObject = SubObjectPtr.Get())
+					{
+						switch (SubObject->GetType())
+						{
+						case Chaos::ImplicitObjectType::Sphere:
+							DrawSphere(SubObject->GetObjectChecked<Chaos::TSphere<float, 3>>(), Rotation, Position, Color);
+							break;
+
+						case Chaos::ImplicitObjectType::TaperedCylinder:
+							DrawTaperedCylinder(SubObject->GetObjectChecked<Chaos::TTaperedCylinder<float>>(), Rotation, Position, Color);
+							break;
+
+						default:
+							break;
+						}
+					}
+				}
+				break;
+	
+			default:
+				DrawCoordinateSystem(PDI, Position, FRotator(Rotation), 10.0f, SDPG_World, 0.1f);  // Draw everything else as a coordinate for now
+				break;
+			}
 		}
 	}
 }
