- Adding a ChaosClothConfig with chaos specific parameters.
- Deprecate chaos parameters on skeletal mesh.

diff --git a/Engine/Plugins/Experimental/ChaosCloth/Source/Chaos/Private/ChaosCloth/ChaosClothingSimulation.cpp b/Engine/Plugins/Experimental/ChaosCloth/Source/Chaos/Private/ChaosCloth/ChaosClothingSimulation.cpp
--- a/Engine/Plugins/Experimental/ChaosCloth/Source/Chaos/Private/ChaosCloth/ChaosClothingSimulation.cpp
+++ b/Engine/Plugins/Experimental/ChaosCloth/Source/Chaos/Private/ChaosCloth/ChaosClothingSimulation.cpp
@@ -32,6 +32,7 @@
 #include "Chaos/Transform.h"
 #include "Chaos/Utilities.h"
 #include "Chaos/Vector.h"
+#include "ChaosCloth/ChaosClothConfig.h"
 
 #if WITH_PHYSX && !PLATFORM_LUMIN && !PLATFORM_ANDROID
 #include "PhysXIncludes.h"
@@ -54,21 +55,9 @@ static TAutoConsoleVariable<float> CVarClothGravityMagnitude(TEXT("physics.Cloth
 ClothingSimulation::ClothingSimulation()
 	: ExternalCollisionsOffset(0)
 	, NumIterations(1)
-	, EdgeStiffness(1.f)
-	, BendingStiffness(1.f)
-	, AreaStiffness(1.f)
-	, VolumeStiffness(0.f)
-	, StrainLimitingStiffness(1.f)
-	, ShapeTargetStiffness(0.f)
 	, SelfCollisionThickness(2.f)
 	, CollisionThickness(1.2f)
 	, GravityMagnitude(490.f)
-	, AnimDriveSpringStiffness(0.001f)
-	, bUseBendingElements(false)
-	, bUseTetrahedralConstraints(false)
-	, bUseThinShellVolumeConstraints(false)
-	, bUseSelfCollisions(false)
-	, bUseContinuousCollisionDetection(false)
 {
 #if WITH_EDITOR
 	DebugClothMaterial = LoadObject<UMaterial>(nullptr, TEXT("/Engine/EditorMaterials/Cloth/CameraLitDoubleSided.CameraLitDoubleSided"), nullptr, LOAD_None, nullptr);  // LOAD_EditorOnly
@@ -166,32 +155,16 @@ void ClothingSimulation::CreateActor(USkeletalMeshComponent* InOwnerComponent, U
 {
 	UE_LOG(LogChaosCloth, Verbose, TEXT("Adding Cloth LOD asset to %s in sim slot %d"), InOwnerComponent->GetOwner() ? *InOwnerComponent->GetOwner()->GetName() : TEXT("None"), InSimDataIndex);
 
-	// TODO(Kriss.Gossart): Set the cloth LOD parameters per cloth rather than per component
-	MassMode = InOwnerComponent->MassMode; // uniform, total, density
-	UniformMass = InOwnerComponent->UniformMass;
-	TotalMass = InOwnerComponent->TotalMass;
-	Density = InOwnerComponent->Density;
-	MinMass = InOwnerComponent->MinPerParticleMass;
-
-	EdgeStiffness = InOwnerComponent->EdgeStiffness;
-	BendingStiffness = InOwnerComponent->BendingStiffness;
-	AreaStiffness = InOwnerComponent->AreaStiffness;
-	VolumeStiffness = InOwnerComponent->VolumeStiffness;
-	StrainLimitingStiffness = InOwnerComponent->StrainLimitingStiffness;
-	ShapeTargetStiffness = InOwnerComponent->ShapeTargetStiffness;
-	bUseBendingElements = InOwnerComponent->bUseBendingElements;
-	bUseTetrahedralConstraints = InOwnerComponent->bUseTetrahedralConstraints;
-	bUseThinShellVolumeConstraints = InOwnerComponent->bUseThinShellVolumeConstraints;
-	bUseSelfCollisions = InOwnerComponent->bUseSelfCollisions;
-	bUseContinuousCollisionDetection = InOwnerComponent->bUseContinuousCollisionDetection; // ccd
-
-	//Evolution->SetCCD(bUseContinuousCollisionDetection);
+	//Evolution->SetCCD(ChaosClothSimConfig->bUseContinuousCollisionDetection);
 	//Evolution->SetCCD(true); // ryan!!!
+	
+	UClothingAssetCommon* Asset = Cast<UClothingAssetCommon>(InAsset);
+	const UChaosClothConfig* const ChaosClothSimConfig = Cast<UChaosClothConfig>(Asset->ChaosClothSimConfig);
+	check(ChaosClothSimConfig);	
 
 	ClothingSimulationContext Context;
     FillContext(InOwnerComponent, 0, &Context);
-
-	UClothingAssetCommon* const Asset = Cast<UClothingAssetCommon>(InAsset);
+	
 	if (Assets.Num() <= InSimDataIndex)
         Assets.SetNumZeroed(InSimDataIndex + 1);
 	Assets[InSimDataIndex] = Asset;
@@ -289,17 +262,17 @@ void ClothingSimulation::CreateActor(USkeletalMeshComponent* InOwnerComponent, U
 		Particles.M(Tri[2]) += ThirdTriArea;
 	}
 	const TSet<int32> Vertices = Mesh->GetVertices();
-	switch (MassMode)
+	switch (ChaosClothSimConfig->MassMode)
 	{
 	case EClothMassMode::UniformMass:
 		for (const int32 Vertex : Vertices)
 		{
-			Particles.M(Vertex) = UniformMass;
+			Particles.M(Vertex) = ChaosClothSimConfig->UniformMass;
 		}
 		break;
 	case EClothMassMode::TotalMass:
 	{
-		const float MassPerUnitArea = TotalArea > 0.0 ? TotalMass / TotalArea : 1.0;
+		const float MassPerUnitArea = TotalArea > 0.0 ? ChaosClothSimConfig->TotalMass / TotalArea : 1.0;
 		for (const int32 Vertex : Vertices)
 		{
 			Particles.M(Vertex) *= MassPerUnitArea;
@@ -309,36 +282,36 @@ void ClothingSimulation::CreateActor(USkeletalMeshComponent* InOwnerComponent, U
 	case EClothMassMode::Density:
 		for (const int32 Vertex : Vertices)
 		{
-			Particles.M(Vertex) *= Density;
+			Particles.M(Vertex) *= ChaosClothSimConfig->Density;
 		}
 		break;
 	};
 	// Clamp and enslave
 	for (uint32 i = Offset; i < Particles.Size(); i++)
 	{
-		Particles.M(i) = FMath::Max(Particles.M(i), MinMass);
+		Particles.M(i) = FMath::Max(Particles.M(i), ChaosClothSimConfig->MinPerParticleMass);
 		Particles.InvM(i) = PhysMesh->IsFixed(i-Offset) ? 0.0 : 1.0/Particles.M(i);
 	}
 
     // Add Model
-    if (ShapeTargetStiffness)
+    if (ChaosClothSimConfig->ShapeTargetStiffness)
     {
-        check(ShapeTargetStiffness > 0.f && ShapeTargetStiffness <= 1.f);
-		Evolution->AddPBDConstraintFunction([ShapeConstraints = Chaos::TPerParticlePBDShapeConstraints<float, 3>(Evolution->Particles(), AnimationPositions, ShapeTargetStiffness)](TPBDParticles<float, 3>& InParticles, const float Dt) {
+        check(ChaosClothSimConfig->ShapeTargetStiffness > 0.f && ChaosClothSimConfig->ShapeTargetStiffness <= 1.f);
+		Evolution->AddPBDConstraintFunction([ShapeConstraints = Chaos::TPerParticlePBDShapeConstraints<float, 3>(Evolution->Particles(), AnimationPositions, ChaosClothSimConfig->ShapeTargetStiffness)](TPBDParticles<float, 3>& InParticles, const float Dt) {
 			ShapeConstraints.Apply(InParticles, Dt);
 		});
     }
-    if (EdgeStiffness)
+    if (ChaosClothSimConfig->EdgeStiffness)
     {
-        check(EdgeStiffness > 0.f && EdgeStiffness <= 1.f);
-        Evolution->AddPBDConstraintFunction([SpringConstraints = Chaos::TPBDSpringConstraints<float, 3>(Evolution->Particles(), SurfaceElements, EdgeStiffness)](TPBDParticles<float, 3>& InParticles, const float Dt) {
+        check(ChaosClothSimConfig->EdgeStiffness > 0.f && ChaosClothSimConfig->EdgeStiffness <= 1.f);
+        Evolution->AddPBDConstraintFunction([SpringConstraints = Chaos::TPBDSpringConstraints<float, 3>(Evolution->Particles(), SurfaceElements, ChaosClothSimConfig->EdgeStiffness)](TPBDParticles<float, 3>& InParticles, const float Dt) {
             SpringConstraints.Apply(InParticles, Dt);
         });
     }
-    if (BendingStiffness)
+    if (ChaosClothSimConfig->BendingStiffness)
     {
-        check(BendingStiffness > 0.f && BendingStiffness <= 1.f);
-        if (bUseBendingElements)
+        check(ChaosClothSimConfig->BendingStiffness > 0.f && ChaosClothSimConfig->BendingStiffness <= 1.f);
+        if (ChaosClothSimConfig->bUseBendingElements)
         {
             TArray<Chaos::TVector<int32, 4>> BendingConstraints = Mesh->GetUniqueAdjacentElements();
 			Evolution->AddPBDConstraintFunction([BendConstraints = Chaos::TPBDBendingConstraints<float>(Evolution->Particles(), MoveTemp(BendingConstraints))](TPBDParticles<float, 3>& InParticles, const float Dt) {
@@ -348,27 +321,27 @@ void ClothingSimulation::CreateActor(USkeletalMeshComponent* InOwnerComponent, U
         else
         {
             TArray<Chaos::TVector<int32, 2>> BendingConstraints = Mesh->GetUniqueAdjacentPoints();
-            Evolution->AddPBDConstraintFunction([SpringConstraints = Chaos::TPBDSpringConstraints<float, 3>(Evolution->Particles(), MoveTemp(BendingConstraints), BendingStiffness)](TPBDParticles<float, 3>& InParticles, const float Dt) {
+            Evolution->AddPBDConstraintFunction([SpringConstraints = Chaos::TPBDSpringConstraints<float, 3>(Evolution->Particles(), MoveTemp(BendingConstraints), ChaosClothSimConfig->BendingStiffness)](TPBDParticles<float, 3>& InParticles, const float Dt) {
                 SpringConstraints.Apply(InParticles, Dt);
             });
         }
     }
-    if (AreaStiffness)
+    if (ChaosClothSimConfig->AreaStiffness)
     {
         TArray<Chaos::TVector<int32, 3>> SurfaceConstraints = SurfaceElements;
-		Evolution->AddPBDConstraintFunction([SurfConstraints = Chaos::TPBDAxialSpringConstraints<float, 3>(Evolution->Particles(), MoveTemp(SurfaceConstraints), AreaStiffness)](TPBDParticles<float, 3>& InParticles, const float Dt) {
+		Evolution->AddPBDConstraintFunction([SurfConstraints = Chaos::TPBDAxialSpringConstraints<float, 3>(Evolution->Particles(), MoveTemp(SurfaceConstraints), ChaosClothSimConfig->AreaStiffness)](TPBDParticles<float, 3>& InParticles, const float Dt) {
 			SurfConstraints.Apply(InParticles, Dt);
 		});
     }
-    if (VolumeStiffness)
+    if (ChaosClothSimConfig->VolumeStiffness)
     {
-        check(VolumeStiffness > 0.f && VolumeStiffness <= 1.f);
-        if (bUseTetrahedralConstraints)
+        check(ChaosClothSimConfig->VolumeStiffness > 0.f && ChaosClothSimConfig->VolumeStiffness <= 1.f);
+        if (ChaosClothSimConfig->bUseTetrahedralConstraints)
         {
             // TODO(mlentine): Need to tetrahedralize surface to support this
             check(false);
         }
-        else if (bUseThinShellVolumeConstraints)
+        else if (ChaosClothSimConfig->bUseThinShellVolumeConstraints)
         {
             TArray<Chaos::TVector<int32, 2>> BendingConstraints = Mesh->GetUniqueAdjacentPoints();
             TArray<Chaos::TVector<int32, 2>> DoubleBendingConstraints;
@@ -399,7 +372,7 @@ void ClothingSimulation::CreateActor(USkeletalMeshComponent* InOwnerComponent, U
                     }
                 }
             }
-            Evolution->AddPBDConstraintFunction([SpringConstraints = Chaos::TPBDSpringConstraints<float, 3>(Evolution->Particles(), MoveTemp(DoubleBendingConstraints), VolumeStiffness)](TPBDParticles<float, 3>& InParticles, const float Dt) {
+            Evolution->AddPBDConstraintFunction([SpringConstraints = Chaos::TPBDSpringConstraints<float, 3>(Evolution->Particles(), MoveTemp(DoubleBendingConstraints), ChaosClothSimConfig->VolumeStiffness)](TPBDParticles<float, 3>& InParticles, const float Dt) {
                 SpringConstraints.Apply(InParticles, Dt);
             });
         }
@@ -413,14 +386,14 @@ void ClothingSimulation::CreateActor(USkeletalMeshComponent* InOwnerComponent, U
 			});            
         }
     }
-    if (StrainLimitingStiffness)
+    if (ChaosClothSimConfig->StrainLimitingStiffness)
     {
 		check(Mesh->GetNumElements() > 0);
 		Chaos::TPerParticlePBDLongRangeConstraints<float, 3> PerParticlePBDLongRangeConstraints(
 			Evolution->Particles(),
 			Mesh->GetPointToNeighborsMap(),
 			10, // The max number of connected neighbors per particle.  ryan - What should this be?  Was k...
-			StrainLimitingStiffness);
+			ChaosClothSimConfig->StrainLimitingStiffness);
 
 		Evolution->AddPBDConstraintFunction([PerParticlePBDLongRangeConstraints = MoveTemp(PerParticlePBDLongRangeConstraints)](TPBDParticles<float, 3>& InParticles, const float Dt)
 		{
@@ -462,7 +435,7 @@ void ClothingSimulation::CreateActor(USkeletalMeshComponent* InOwnerComponent, U
 	if (PhysMesh->GetFloatArray(PhysMeshAnimDriveIndex)->Num() > 0)
 	{
 		check(Mesh->GetNumElements() > 0);
-		TPBDAnimDriveConstraint<float, 3> PBDAnimDriveConstraint(Offset, &AnimationPositions, PhysMesh->GetFloatArray(PhysMeshAnimDriveIndex), AnimDriveSpringStiffness);
+		TPBDAnimDriveConstraint<float, 3> PBDAnimDriveConstraint(Offset, &AnimationPositions, PhysMesh->GetFloatArray(PhysMeshAnimDriveIndex), ChaosClothSimConfig->AnimDriveSpringStiffness);
 		Evolution->AddPBDConstraintFunction(
 			[PBDAnimDriveConstraint = MoveTemp(PBDAnimDriveConstraint)](TPBDParticles<float, 3>& InParticles, const float Dt)
 		{
@@ -471,7 +444,7 @@ void ClothingSimulation::CreateActor(USkeletalMeshComponent* InOwnerComponent, U
 	}
 
     // Add Self Collisions
-    if (bUseSelfCollisions)
+    if (ChaosClothSimConfig->bUseSelfCollisions)
     {
         // TODO(mlentine): Parallelize these for multiple meshes
         Evolution->CollisionTriangles().Append(SurfaceElements);
@@ -1452,27 +1425,26 @@ void ClothingSimulation::DebugDrawMaxDistances(USkeletalMeshComponent* OwnerComp
 
 void ClothingSimulation::DebugDrawSelfCollision(USkeletalMeshComponent* OwnerComponent, FPrimitiveDrawInterface* PDI) const
 {
-	if (!bUseSelfCollisions)
-	{
-		// No self collisions on this actor
-		return;
-	}
 	const TPBDParticles<float, 3>& Particles = Evolution->Particles();
 	for (int32 i = 0; i < IndexToRangeMap.Num(); ++i)
 	{
 		if (const UClothingAssetCommon* const Asset = Assets[i])
 		{
-			const FTransform RootBoneTransform = OwnerComponent->GetComponentSpaceTransforms()[Asset->ReferenceBoneIndex];
-
-			const UClothLODDataBase* LodData = Asset->ClothLodData[0];
-			const UClothPhysicalMeshDataBase* PhysMesh = LodData->PhysicalMeshData;
-			const TArray<uint32>& SelfCollisionIndices = PhysMesh->SelfCollisionIndices;
-			for (int32 SelfColIdx = 0; SelfColIdx < SelfCollisionIndices.Num(); ++SelfColIdx)
+			const UChaosClothConfig* const ChaosClothSimConfig = Cast<UChaosClothConfig>(Asset->ChaosClothSimConfig);
+			if (ChaosClothSimConfig && ChaosClothSimConfig->bUseSelfCollisions)
 			{
-				const FVector ParticlePosition =
-					RootBoneTransform.TransformPosition(
-						Particles.X(PhysMesh->SelfCollisionIndices[SelfColIdx]));
-				DrawWireSphere(PDI, ParticlePosition, FColor::White, SelfCollisionThickness, 8, SDPG_World, 0.0f, 0.001f);
+				const FTransform RootBoneTransform = OwnerComponent->GetComponentSpaceTransforms()[Asset->ReferenceBoneIndex];
+
+				const UClothLODDataBase* LodData = Asset->ClothLodData[0];
+				const UClothPhysicalMeshDataBase* PhysMesh = LodData->PhysicalMeshData;
+				const TArray<uint32>& SelfCollisionIndices = PhysMesh->SelfCollisionIndices;
+				for (int32 SelfColIdx = 0; SelfColIdx < SelfCollisionIndices.Num(); ++SelfColIdx)
+				{
+					const FVector ParticlePosition =
+						RootBoneTransform.TransformPosition(
+							Particles.X(PhysMesh->SelfCollisionIndices[SelfColIdx]));
+					DrawWireSphere(PDI, ParticlePosition, FColor::White, SelfCollisionThickness, 8, SDPG_World, 0.0f, 0.001f);
+				}
 			}
 		}
 	}
diff --git a/Engine/Plugins/Experimental/ChaosCloth/Source/Chaos/Private/ChaosCloth/ChaosClothingSimulationFactory.cpp b/Engine/Plugins/Experimental/ChaosCloth/Source/Chaos/Private/ChaosCloth/ChaosClothingSimulationFactory.cpp
--- a/Engine/Plugins/Experimental/ChaosCloth/Source/Chaos/Private/ChaosCloth/ChaosClothingSimulationFactory.cpp
+++ b/Engine/Plugins/Experimental/ChaosCloth/Source/Chaos/Private/ChaosCloth/ChaosClothingSimulationFactory.cpp
@@ -1,6 +1,7 @@
 // Copyright 1998-2019 Epic Games, Inc. All Rights Reserved.
 
 #include "ChaosCloth/ChaosClothingSimulationFactory.h"
+#include "ChaosCloth/ChaosClothConfig.h"
 #include "ChaosCloth/ChaosClothingSimulation.h"
 
 IClothingSimulation* UChaosClothingSimulationFactory::CreateSimulation()
@@ -27,4 +28,9 @@ bool UChaosClothingSimulationFactory::SupportsRuntimeInteraction()
 UClothingSimulationInteractor* UChaosClothingSimulationFactory::CreateInteractor()
 {
 	return nullptr;
+}
+
+UClothConfigBase* UChaosClothingSimulationFactory::CreateDefaultClothConfig(const FObjectInitializer& ObjectInitializer, UObject* Outer)
+{
+	return ObjectInitializer.CreateDefaultSubobject<UChaosClothConfig>(Outer, UChaosClothConfig::StaticClass()->GetFName());
 }
\ No newline at end of file
diff --git a/Engine/Plugins/Experimental/ChaosCloth/Source/Chaos/Public/ChaosCloth/ChaosClothConfig.h b/Engine/Plugins/Experimental/ChaosCloth/Source/Chaos/Public/ChaosCloth/ChaosClothConfig.h
new file mode 100644
--- /dev/null
+++ b/Engine/Plugins/Experimental/ChaosCloth/Source/Chaos/Public/ChaosCloth/ChaosClothConfig.h
@@ -0,0 +1,118 @@
+// Copyright 1998-2019 Epic Games, Inc. All Rights Reserved.
+#pragma once
+#include "Engine/SkeletalMesh.h"
+#include "ClothConfig.h"
+#include "CoreMinimal.h"
+#include "ChaosClothConfig.generated.h"
+
+/** Holds initial, asset level config for clothing actors. */
+// Hiding categories that will be used in the future
+UCLASS(HideCategories = (Solver, Collision))
+class CHAOSCLOTH_API UChaosClothConfig : public UClothConfigBase
+{
+	GENERATED_BODY()
+public:
+	UChaosClothConfig() {};
+
+	/**
+	 * How cloth particle mass is determined
+	 * -	Uniform Mass: Every particle's mass will be set to the value specified in the UniformMass setting
+	 * -	Total Mass: The total mass is distributed equally over all the particles
+	 * -	Density: A constant mass density is used	 	 
+	 */
+	UPROPERTY(EditAnywhere, Category = MassConfig)
+	EClothMassMode MassMode = EClothMassMode::Density;
+
+	// The value used when the Mass Mode is set to Uniform Mass
+	UPROPERTY(EditAnywhere, Category = MassConfig)
+	float UniformMass = 1.0f;
+
+	// The value used when Mass Mode is set to TotalMass
+	UPROPERTY(EditAnywhere, Category = MassConfig)
+	float TotalMass = 100.0f;
+
+	/**
+	 * The value used when Mass Mode is set to Density 
+	 * Water: 1.0
+	 * Cotton: 0.155
+	 * Wool: 0.13
+	 * Silk: 0.133
+	 */
+	UPROPERTY(EditAnywhere, Category = MassConfig)
+	float Density = 0.1f;
+
+	// This is a lower bound to cloth particle masses
+	UPROPERTY(EditAnywhere, Category = MassConfig)
+	float MinPerParticleMass = 0.0001f;
+
+	// The Number of iterations used in the cloth solver
+	// All cloth constraints will be more stiff when increased
+	UPROPERTY(EditAnywhere, Category = Solver)
+	int32 NumIterations = 1;
+
+	// The Stiffness of the Edge constraints
+	UPROPERTY(EditAnywhere, Category = Stiffness, meta = (UIMin = "0", UIMax = "1", ClampMin = "0", ClampMax = "1"))
+	float EdgeStiffness = 1.f;
+
+	// The Stiffness of the bending constraints
+	UPROPERTY(EditAnywhere, Category = Stiffness, meta = (UIMin = "0", UIMax = "1", ClampMin = "0", ClampMax = "1"))
+	float BendingStiffness = 1.f;
+
+	// The stiffness of the area preservation constraints
+	UPROPERTY(EditAnywhere, Category = Stiffness, meta = (UIMin = "0", UIMax = "1", ClampMin = "0", ClampMax = "1"))
+	float AreaStiffness = 1.f;
+
+	// The stiffness of the volume preservation constraints
+	UPROPERTY(EditAnywhere, Category = Stiffness, meta = (UIMin = "0", UIMax = "1", ClampMin = "0", ClampMax = "1"))
+	float VolumeStiffness = 0.f;
+
+	// The stiffness of the strain limiting constraints
+	UPROPERTY(EditAnywhere, Category = Stiffness, meta = (UIMin = "0", UIMax = "1", ClampMin = "0", ClampMax = "1"))
+	float StrainLimitingStiffness = 1.f;
+	
+	// The stiffness of the shape target constraints
+	UPROPERTY(EditAnywhere, Category = Stiffness, meta = (UIMin = "0", UIMax = "1", ClampMin = "0", ClampMax = "1"))
+	float ShapeTargetStiffness = 0.f;
+
+	// The radius of the self collisions spheres
+	UPROPERTY(EditAnywhere, Category = Collision)
+	float SelfCollisionThickness = 2.f;
+
+	// The thickness of the cloth point spheres when colliding against collisions primitives
+	UPROPERTY(EditAnywhere, Category = Collision)
+	float CollisionThickness = 1.2;
+
+	// Friction coefficient for cloth - collider interaction
+	UPROPERTY(EditAnywhere, Category = Collision, meta = (UIMin = "0", UIMax = "1", ClampMin = "0", ClampMax = "10"))
+	float CoefficientOfFriction = 0.0f;
+
+	UPROPERTY(EditAnywhere, Category = Solver, meta = (UIMin = "0", UIMax = "1", ClampMin = "0", ClampMax = "1"))
+	float Damping = 0.01f;
+
+	UPROPERTY(EditAnywhere, Category = Solver)
+	float GravityMagnitude = 490.f;
+
+	// Default spring stiffness for anim drive if an anim drive is in use
+	UPROPERTY(EditAnywhere, Category = Stiffness, meta = (UIMin = "0", UIMax = "1", ClampMin = "0", ClampMax = "1"))
+	float AnimDriveSpringStiffness = 0.001f;
+
+	// Enable bending elements
+	UPROPERTY(EditAnywhere, Category = ClothEnableFlags)
+	bool bUseBendingElements = false;
+
+	// Enable tetrahral constraints
+	UPROPERTY(EditAnywhere, Category = ClothEnableFlags)
+	bool bUseTetrahedralConstraints = false;
+
+	// Enable thin shell volume constraints 
+	UPROPERTY(EditAnywhere, Category = ClothEnableFlags)
+	bool bUseThinShellVolumeConstraints = false;
+
+	// Enable self collision
+	UPROPERTY(EditAnywhere, Category = ClothEnableFlags)
+	bool bUseSelfCollisions = false;
+
+	// Enable continuous collision detection
+	UPROPERTY(EditAnywhere, Category = ClothEnableFlags)
+	bool bUseContinuousCollisionDetection = false;
+};
diff --git a/Engine/Plugins/Experimental/ChaosCloth/Source/Chaos/Public/ChaosCloth/ChaosClothingSimulation.h b/Engine/Plugins/Experimental/ChaosCloth/Source/Chaos/Public/ChaosCloth/ChaosClothingSimulation.h
--- a/Engine/Plugins/Experimental/ChaosCloth/Source/Chaos/Public/ChaosCloth/ChaosClothingSimulation.h
+++ b/Engine/Plugins/Experimental/ChaosCloth/Source/Chaos/Public/ChaosCloth/ChaosClothingSimulation.h
@@ -122,30 +122,11 @@ namespace Chaos
 		float ClampDeltaTime;
 		// Parameters that should be set in the ui
 		int32 NumIterations;
-
-		EClothMassMode MassMode;
-		float UniformMass;
-		float TotalMass;
-		float Density;
-		float MinMass;
-
-		float EdgeStiffness;
-		float BendingStiffness;
-		float AreaStiffness;
-		float VolumeStiffness;
-		float StrainLimitingStiffness;
-		float ShapeTargetStiffness;
 		float SelfCollisionThickness;
 		float CollisionThickness;
 		float CoefficientOfFriction;
 		float Damping;
 		float GravityMagnitude;
-		float AnimDriveSpringStiffness;
-		bool bUseBendingElements;
-		bool bUseTetrahedralConstraints;
-		bool bUseThinShellVolumeConstraints;
-		bool bUseSelfCollisions;
-		bool bUseContinuousCollisionDetection;
 
 #if WITH_EDITOR
 		// Visualization material
diff --git a/Engine/Plugins/Experimental/ChaosCloth/Source/Chaos/Public/ChaosCloth/ChaosClothingSimulationFactory.h b/Engine/Plugins/Experimental/ChaosCloth/Source/Chaos/Public/ChaosCloth/ChaosClothingSimulationFactory.h
--- a/Engine/Plugins/Experimental/ChaosCloth/Source/Chaos/Public/ChaosCloth/ChaosClothingSimulationFactory.h
+++ b/Engine/Plugins/Experimental/ChaosCloth/Source/Chaos/Public/ChaosCloth/ChaosClothingSimulationFactory.h
@@ -16,4 +16,5 @@ class CHAOSCLOTH_API UChaosClothingSimulationFactory : public UClothingSimulatio
     virtual bool SupportsAsset(UClothingAssetBase* InAsset) override;
     virtual bool SupportsRuntimeInteraction() override;
 	virtual UClothingSimulationInteractor* CreateInteractor() override;
+	virtual UClothConfigBase* CreateDefaultClothConfig(const FObjectInitializer& ObjectInitializer, UObject* Outer) override;
 };
diff --git a/Engine/Source/Editor/DetailCustomizations/Private/SkeletalMeshComponentDetails.cpp b/Engine/Source/Editor/DetailCustomizations/Private/SkeletalMeshComponentDetails.cpp
--- a/Engine/Source/Editor/DetailCustomizations/Private/SkeletalMeshComponentDetails.cpp
+++ b/Engine/Source/Editor/DetailCustomizations/Private/SkeletalMeshComponentDetails.cpp
@@ -75,21 +75,7 @@ void FSkeletalMeshComponentDetails::CustomizeDetails(IDetailLayoutBuilder& Detai
 	DetailBuilder.EditCategory("Physics", FText::GetEmpty(), ECategoryPriority::TypeSpecific);
 	DetailBuilder.HideProperty("bCastStaticShadow", UPrimitiveComponent::StaticClass());
 	DetailBuilder.HideProperty("LightmapType", UPrimitiveComponent::StaticClass());
-	DetailBuilder.EditCategory("Animation", FText::GetEmpty(), ECategoryPriority::Important);
-	if (!FModuleManager::Get().IsModuleLoaded("UEPhysics"))
-	{
-		DetailBuilder.HideProperty("EdgeStiffness");
-		DetailBuilder.HideProperty("BendingStiffness");
-		DetailBuilder.HideProperty("AreaStiffness");
-		DetailBuilder.HideProperty("VolumeStiffness");
-		DetailBuilder.HideProperty("StrainLimitingStiffness");
-		DetailBuilder.HideProperty("ShapeTargetStiffness");
-		DetailBuilder.HideProperty("bUseBendingElements");
-		DetailBuilder.HideProperty("bUseTetrahedralConstraints");
-		DetailBuilder.HideProperty("bUseThinShellVolumeConstraints");
-		DetailBuilder.HideProperty("bUseSelfCollisions");
-		DetailBuilder.HideProperty("bUseContinuousCollisionDetection");
-	}
+	DetailBuilder.EditCategory("Animation", FText::GetEmpty(), ECategoryPriority::Important);	
 
 	PerformInitialRegistrationOfSkeletalMeshes(DetailBuilder);
 
diff --git a/Engine/Source/Runtime/ClothingSystemRuntimeCommon/Private/ClothingAsset.cpp b/Engine/Source/Runtime/ClothingSystemRuntimeCommon/Private/ClothingAsset.cpp
--- a/Engine/Source/Runtime/ClothingSystemRuntimeCommon/Private/ClothingAsset.cpp
+++ b/Engine/Source/Runtime/ClothingSystemRuntimeCommon/Private/ClothingAsset.cpp
@@ -173,6 +173,7 @@ UClothingAssetCommon::UClothingAssetCommon(const FObjectInitializer& ObjectIniti
 	: Super(ObjectInitializer)
 	, PhysicsAsset(nullptr)
 	, ClothSimConfig(nullptr)
+	, ChaosClothSimConfig(nullptr)
 	, ReferenceBoneIndex(0)
 	, CustomData(nullptr)
 {}
diff --git a/Engine/Source/Runtime/ClothingSystemRuntimeCommon/Public/ClothingAsset.h b/Engine/Source/Runtime/ClothingSystemRuntimeCommon/Public/ClothingAsset.h
--- a/Engine/Source/Runtime/ClothingSystemRuntimeCommon/Public/ClothingAsset.h
+++ b/Engine/Source/Runtime/ClothingSystemRuntimeCommon/Public/ClothingAsset.h
@@ -158,10 +158,18 @@ public:
 	UPROPERTY(EditAnywhere, Category = Config)
 	UPhysicsAsset* PhysicsAsset;
 
-	// Parameters for how the cloth behaves.
+	// Parameters for how the NVcloth behaves.
+	// These will have no effect on Chaos cloth
 	UPROPERTY(EditAnywhere, Category = Config, NoClear, meta = (NoResetToDefault))
 	UClothConfigBase* ClothSimConfig;
 
+	// Parameters for how Chaos cloth behaves 
+	// These will not affect NVcloth
+	// For now, we have two configuration parameters so that we can switch between chaos and
+	// non chaos at will without losing the original NVcloth data
+	UPROPERTY(EditAnywhere, Category = Config, NoClear, meta = (NoResetToDefault))
+	UClothConfigBase* ChaosClothSimConfig;
+
 	// The actual asset data, listed by LOD.
 	UPROPERTY()
 	TArray<UClothLODDataBase*> ClothLodData;
diff --git a/Engine/Source/Runtime/ClothingSystemRuntimeInterface/Public/ClothingSimulationFactory.h b/Engine/Source/Runtime/ClothingSystemRuntimeInterface/Public/ClothingSimulationFactory.h
--- a/Engine/Source/Runtime/ClothingSystemRuntimeInterface/Public/ClothingSimulationFactory.h
+++ b/Engine/Source/Runtime/ClothingSystemRuntimeInterface/Public/ClothingSimulationFactory.h
@@ -11,6 +11,7 @@
 
 class UClothingAssetBase;
 class UClothingSimulationInteractor;
+class UClothConfigBase;
 
 // An interface for a class that will provide default simulation factory classes
 // Used by modules wanting to override clothing simulation to provide their own implementation
@@ -59,4 +60,11 @@ public:
 	{
 		return nullptr;
 	}
+
+	// Creates a default cloth configuration during construction of the clothing Asset
+	// TBC: This may be removed after we have a specialized ChaosClothingAsset
+	virtual UClothConfigBase* CreateDefaultClothConfig(const FObjectInitializer& ObjectInitializer, UObject* Outer)
+	{
+		return nullptr;
+	}
 };
diff --git a/Engine/Source/Runtime/ClothingSystemRuntimeNv/Private/Assets/ClothingAssetNv.cpp b/Engine/Source/Runtime/ClothingSystemRuntimeNv/Private/Assets/ClothingAssetNv.cpp
--- a/Engine/Source/Runtime/ClothingSystemRuntimeNv/Private/Assets/ClothingAssetNv.cpp
+++ b/Engine/Source/Runtime/ClothingSystemRuntimeNv/Private/Assets/ClothingAssetNv.cpp
@@ -62,6 +62,27 @@ UClothingAssetNv::UClothingAssetNv(const FObjectInitializer& ObjectInitializer)
 	: Super(ObjectInitializer)
 {
 	ClothSimConfig = ObjectInitializer.CreateDefaultSubobject<UClothConfigNv>(this, UClothConfigNv::StaticClass()->GetFName());
+
+// TODO: Move this to Chaos cloth once we have a chaos specific UClothingAsset (In which case we don't need to get a factory from the feature system)
+#if WITH_CHAOS_CLOTHING
+	
+	TArray<IClothingSimulationFactoryClassProvider*> ClassProviders = IModularFeatures::Get().GetModularFeatureImplementations<IClothingSimulationFactoryClassProvider>(IClothingSimulationFactoryClassProvider::FeatureName);
+	TSubclassOf<class UClothingSimulationFactory> ClothingSimulationFactory = nullptr;
+	for (int32 Index = ClassProviders.Num() - 1; Index >= 0 && !ChaosClothSimConfig; --Index)  // Iterating backwards since the last providers in the list should override the first ones
+	{
+		IClothingSimulationFactoryClassProvider* const Provider = ClassProviders[Index];
+		check(Provider);
+		ClothingSimulationFactory = Provider->GetDefaultSimulationFactoryClass();
+
+		UClass* SimFactoryClass = *ClothingSimulationFactory;
+		if (SimFactoryClass)
+		{
+			ChaosClothSimConfig = SimFactoryClass->GetDefaultObject<UClothingSimulationFactory>()->CreateDefaultClothConfig(ObjectInitializer, this);
+		}
+	}
+	
+#endif
+
 }
 
 void UClothingAssetNv::PostLoad()
diff --git a/Engine/Source/Runtime/ClothingSystemRuntimeNv/Public/Assets/ClothingAssetNv.h b/Engine/Source/Runtime/ClothingSystemRuntimeNv/Public/Assets/ClothingAssetNv.h
--- a/Engine/Source/Runtime/ClothingSystemRuntimeNv/Public/Assets/ClothingAssetNv.h
+++ b/Engine/Source/Runtime/ClothingSystemRuntimeNv/Public/Assets/ClothingAssetNv.h
@@ -36,7 +36,7 @@ public:
 	FClothConfig ClothConfig_DEPRECATED;
 
 	/** 
-	 * Deprecasted property for transitioning \c FClothLODData struct to the 
+	 * Deprecated property for transitioning \c FClothLODData struct to the 
 	 * \c UClothLODDataNv class, in a new property called \c ClothLodData.
 	 */
 	UPROPERTY()
diff --git a/Engine/Source/Runtime/Engine/Classes/Components/SkeletalMeshComponent.h b/Engine/Source/Runtime/Engine/Classes/Components/SkeletalMeshComponent.h
--- a/Engine/Source/Runtime/Engine/Classes/Components/SkeletalMeshComponent.h
+++ b/Engine/Source/Runtime/Engine/Classes/Components/SkeletalMeshComponent.h
@@ -587,20 +587,20 @@ public:
 	UPROPERTY()
 	uint8 bEnableLineCheckWithBounds:1;
 
-	UPROPERTY(EditAnywhere, BlueprintReadWrite, Category = Clothing)
-	uint8 bUseBendingElements:1;
+	UPROPERTY(meta = (DeprecatedProperty, DeprecationMessage = "This property is deprecated, please set it on the Clothing Asset / ClothConfig instead."))
+	uint8 bUseBendingElements_DEPRECATED :1;
 	
-	UPROPERTY(EditAnywhere, BlueprintReadWrite, Category = Clothing)
-	uint8 bUseTetrahedralConstraints:1;
+	UPROPERTY(meta = (DeprecatedProperty, DeprecationMessage = "This property is deprecated, please set it on the Clothing Asset / ClothConfig instead."))
+	uint8 bUseTetrahedralConstraints_DEPRECATED :1;
 	
-	UPROPERTY(EditAnywhere, BlueprintReadWrite, Category = Clothing)
-	uint8 bUseThinShellVolumeConstraints:1;
+	UPROPERTY(meta = (DeprecatedProperty, DeprecationMessage = "This property is deprecated, please set it on the Clothing Asset / ClothConfig instead."))
+	uint8 bUseThinShellVolumeConstraints_DEPRECATED :1;
 	
-	UPROPERTY(EditAnywhere, BlueprintReadWrite, Category = Clothing)
-	uint8 bUseSelfCollisions:1;
+	UPROPERTY(meta = (DeprecatedProperty, DeprecationMessage = "This property is deprecated, please set it on the Clothing Asset / ClothConfig instead."))
+	uint8 bUseSelfCollisions_DEPRECATED :1;
 	
-	UPROPERTY(EditAnywhere, BlueprintReadWrite, Category = Clothing)
-	uint8 bUseContinuousCollisionDetection:1;
+	UPROPERTY(meta = (DeprecatedProperty, DeprecationMessage = "This property is deprecated, please set it on the Clothing Asset / ClothConfig instead."))
+	uint8 bUseContinuousCollisionDetection_DEPRECATED :1;
 
 	/** If true, propagates calls to ApplyAnimationCurvesToComponent for slave components, only needed if slave components do not tick themselves */
 	UPROPERTY(EditAnywhere, BlueprintReadWrite, Category = MasterPoseComponent)
@@ -638,14 +638,14 @@ public:
 	UPROPERTY(transient)
 	uint16 CachedAnimCurveUidVersion;
 
-	UPROPERTY(EditAnywhere, BlueprintReadWrite, Category = Clothing)
-	EClothMassMode MassMode;
+	UPROPERTY(meta = (DeprecatedProperty, DeprecationMessage = "This property is deprecated, please set it on the Clothing Asset / ClothConfig instead."))
+	EClothMassMode MassMode_DEPRECATED;
 
-	UPROPERTY(EditAnywhere, BlueprintReadWrite, Category = Clothing)
-	float UniformMass;
+	UPROPERTY(meta = (DeprecatedProperty, DeprecationMessage = "This property is deprecated, please set it on the Clothing Asset / ClothConfig instead."))
+	float UniformMass_DEPRECATED;
 
-	UPROPERTY(EditAnywhere, BlueprintReadWrite, Category = Clothing)
-	float TotalMass;
+	UPROPERTY(meta = (DeprecatedProperty, DeprecationMessage = "This property is deprecated, please set it on the Clothing Asset / ClothConfig instead."))
+	float TotalMass_DEPRECATED;
 	
 	/**
 	 * Water: 1.0
@@ -653,11 +653,11 @@ public:
 	 * Wool: 0.13
 	 * Silk: 0.133
 	 */
-	UPROPERTY(EditAnywhere, BlueprintReadWrite, Category = Clothing)
-	float Density;
+	UPROPERTY(meta = (DeprecatedProperty, DeprecationMessage = "This property is deprecated, please set it on the Clothing Asset / ClothConfig instead."))
+	float Density_DEPRECATED;
 
-	UPROPERTY(EditAnywhere, BlueprintReadWrite, Category = Clothing)
-	float MinPerParticleMass;
+	UPROPERTY(meta = (DeprecatedProperty, DeprecationMessage = "This property is deprecated, please set it on the Clothing Asset / ClothConfig instead."))
+	float MinPerParticleMass_DEPRECATED;
 
 
 	/**
@@ -667,23 +667,23 @@ public:
 	UPROPERTY(EditAnywhere, BlueprintReadWrite, Category = Clothing)
 	float ClothBlendWeight;
 
-	UPROPERTY(EditAnywhere, BlueprintReadWrite, Category = Clothing)
-	float EdgeStiffness;
+	UPROPERTY(meta = (DeprecatedProperty, DeprecationMessage = "This property is deprecated, please set it on the Clothing Asset / ClothConfig instead."))
+	float EdgeStiffness_DEPRECATED;
 	
-	UPROPERTY(EditAnywhere, BlueprintReadWrite, Category = Clothing)
-	float BendingStiffness;
+	UPROPERTY(meta = (DeprecatedProperty, DeprecationMessage = "This property is deprecated, please set it on the Clothing Asset / ClothConfig instead."))
+	float BendingStiffness_DEPRECATED;
 	
-	UPROPERTY(EditAnywhere, BlueprintReadWrite, Category = Clothing)
-	float AreaStiffness;
+	UPROPERTY(meta = (DeprecatedProperty, DeprecationMessage = "This property is deprecated, please set it on the Clothing Asset / ClothConfig instead."))
+	float AreaStiffness_DEPRECATED;
 	
-	UPROPERTY(EditAnywhere, BlueprintReadWrite, Category = Clothing)
-	float VolumeStiffness;
+	UPROPERTY(meta = (DeprecatedProperty, DeprecationMessage = "This property is deprecated, please set it on the Clothing Asset / ClothConfig instead."))
+	float VolumeStiffness_DEPRECATED;
 	
-	UPROPERTY(EditAnywhere, BlueprintReadWrite, Category = Clothing)
-	float StrainLimitingStiffness;
+	UPROPERTY(meta = (DeprecatedProperty, DeprecationMessage = "This property is deprecated, please set it on the Clothing Asset / ClothConfig instead."))
+	float StrainLimitingStiffness_DEPRECATED;
 	
-	UPROPERTY(EditAnywhere, BlueprintReadWrite, Category = Clothing)
-	float ShapeTargetStiffness;
+	UPROPERTY(meta = (DeprecatedProperty, DeprecationMessage = "This property is deprecated, please set it on the Clothing Asset / ClothConfig instead."))
+	float ShapeTargetStiffness_DEPRECATED;
 
 private:
 
diff --git a/Engine/Source/Runtime/Engine/Private/Components/SkeletalMeshComponent.cpp b/Engine/Source/Runtime/Engine/Private/Components/SkeletalMeshComponent.cpp
--- a/Engine/Source/Runtime/Engine/Private/Components/SkeletalMeshComponent.cpp
+++ b/Engine/Source/Runtime/Engine/Private/Components/SkeletalMeshComponent.cpp
@@ -198,22 +198,22 @@ USkeletalMeshComponent::USkeletalMeshComponent(const FObjectInitializer& ObjectI
 
 #endif//#if WITH_APEX_CLOTHING || WITH_CHAOS_CLOTHING
 
-	MassMode = EClothMassMode::Density;
-	UniformMass = 1.f;
-	TotalMass = 100.0f;
-	Density = 0.1f;
-	MinPerParticleMass = 0.0001f;
-	EdgeStiffness = 1.f;
-	BendingStiffness = 1.f;
-	AreaStiffness = 1.f;
-	VolumeStiffness = 0.f;
-	StrainLimitingStiffness = 1.f;
-	ShapeTargetStiffness = 0.f;
-	bUseBendingElements = false;
-	bUseTetrahedralConstraints = false;
-	bUseThinShellVolumeConstraints = false;
-	bUseSelfCollisions = false;
-	bUseContinuousCollisionDetection = false;
+	MassMode_DEPRECATED = EClothMassMode::Density;
+	UniformMass_DEPRECATED = 1.f;
+	TotalMass_DEPRECATED = 100.0f;
+	Density_DEPRECATED = 0.1f;
+	MinPerParticleMass_DEPRECATED = 0.0001f;
+	EdgeStiffness_DEPRECATED = 1.f;
+	BendingStiffness_DEPRECATED = 1.f;
+	AreaStiffness_DEPRECATED = 1.f;
+	VolumeStiffness_DEPRECATED = 0.f;
+	StrainLimitingStiffness_DEPRECATED = 1.f;
+	ShapeTargetStiffness_DEPRECATED = 0.f;
+	bUseBendingElements_DEPRECATED = false;
+	bUseTetrahedralConstraints_DEPRECATED = false;
+	bUseThinShellVolumeConstraints_DEPRECATED = false;
+	bUseSelfCollisions_DEPRECATED = false;
+	bUseContinuousCollisionDetection_DEPRECATED = false;
 
 #if WITH_EDITORONLY_DATA
 	DefaultPlayRate_DEPRECATED = 1.0f;
