ChaosCloth: ios compile error fix

diff --git a/Engine/Plugins/Experimental/ChaosCloth/Source/Chaos/Private/ChaosCloth/ChaosClothModule.cpp b/Engine/Plugins/Experimental/ChaosCloth/Source/Chaos/Private/ChaosCloth/ChaosClothModule.cpp
--- a/Engine/Plugins/Experimental/ChaosCloth/Source/Chaos/Private/ChaosCloth/ChaosClothModule.cpp
+++ b/Engine/Plugins/Experimental/ChaosCloth/Source/Chaos/Private/ChaosCloth/ChaosClothModule.cpp
@@ -3,6 +3,7 @@
 #include "ChaosCloth/ChaosClothModule.h"
 #include "ChaosCloth/ChaosClothPrivate.h"
 #include "ChaosCloth/ChaosClothingSimulationFactory.h"
+#include "Features/IModularFeatures.h"
 #include "Modules/ModuleManager.h"
 
 //////////////////////////////////////////////////////////////////////////
