- Adding Cloth configuration parameters that are common to all instances in a skeletal mesh
- Remove the related cvars

diff --git a/Engine/Plugins/Experimental/ChaosCloth/Source/Chaos/Private/ChaosCloth/ChaosClothingSimulation.cpp b/Engine/Plugins/Experimental/ChaosCloth/Source/Chaos/Private/ChaosCloth/ChaosClothingSimulation.cpp
--- a/Engine/Plugins/Experimental/ChaosCloth/Source/Chaos/Private/ChaosCloth/ChaosClothingSimulation.cpp
+++ b/Engine/Plugins/Experimental/ChaosCloth/Source/Chaos/Private/ChaosCloth/ChaosClothingSimulation.cpp
@@ -45,19 +45,9 @@
 
 using namespace Chaos;
 
-static TAutoConsoleVariable<int32> CVarClothNumIterations(TEXT("physics.ClothNumIterations"), 1, TEXT(""));
-static TAutoConsoleVariable<float> CVarClothSelfCollisionThickness(TEXT("physics.ClothSelfCollisionThickness"), 2.f, TEXT(""));
-static TAutoConsoleVariable<float> CVarClothCollisionThickness(TEXT("physics.ClothCollisionThickness"), 1.2f, TEXT(""));
-static TAutoConsoleVariable<float> CVarClothCoefficientOfFriction(TEXT("physics.ClothCoefficientOfFriction"), 0.f, TEXT(""));
-static TAutoConsoleVariable<float> CVarClothDamping(TEXT("physics.ClothDamping"), 0.01f, TEXT(""));
-static TAutoConsoleVariable<float> CVarClothGravityMagnitude(TEXT("physics.ClothGravityMagnitude"), 490.f, TEXT(""));
-
 ClothingSimulation::ClothingSimulation()
-	: ExternalCollisionsOffset(0)
-	, NumIterations(1)
-	, SelfCollisionThickness(2.f)
-	, CollisionThickness(1.2f)
-	, GravityMagnitude(490.f)
+	: ClothSharedSimConfig(nullptr)
+	, ExternalCollisionsOffset(0)
 {
 #if WITH_EDITOR
 	DebugClothMaterial = LoadObject<UMaterial>(nullptr, TEXT("/Engine/EditorMaterials/Cloth/CameraLitDoubleSided.CameraLitDoubleSided"), nullptr, LOAD_None, nullptr);  // LOAD_EditorOnly
@@ -69,12 +59,14 @@ ClothingSimulation::~ClothingSimulation()
 
 void ClothingSimulation::Initialize()
 {
-    NumIterations = CVarClothNumIterations.GetValueOnGameThread();
-    SelfCollisionThickness = CVarClothSelfCollisionThickness.GetValueOnGameThread();
-    CollisionThickness = CVarClothCollisionThickness.GetValueOnGameThread();
-    CoefficientOfFriction = CVarClothCoefficientOfFriction.GetValueOnGameThread();
-    Damping = CVarClothDamping.GetValueOnGameThread();
-    GravityMagnitude = CVarClothGravityMagnitude.GetValueOnGameThread();
+
+	// Default parameters. Will be overwritten when cloth assets are loaded
+	int32 NumIterations = 1;
+	float SelfCollisionThickness = 2.0f;
+	float CollisionThickness = 1.2f;
+	float CoefficientOfFriction = 0.0f;
+	float Damping = 0.01f;
+	float GravityMagnitude = 490.0f;
 
     Chaos::TPBDParticles<float, 3> LocalParticles;
     Chaos::TKinematicGeometryClothParticles<float, 3> TRigidParticles;
@@ -143,6 +135,7 @@ void ClothingSimulation::Shutdown()
 	PointNormals.Reset();
 	Evolution.Reset();
 	ExternalCollisionsOffset = 0;
+	ClothSharedSimConfig = nullptr;
 }
 
 void ClothingSimulation::DestroyActors()
@@ -160,7 +153,7 @@ void ClothingSimulation::CreateActor(USkeletalMeshComponent* InOwnerComponent, U
 	
 	UClothingAssetCommon* Asset = Cast<UClothingAssetCommon>(InAsset);
 	const UChaosClothConfig* const ChaosClothSimConfig = Cast<UChaosClothConfig>(Asset->ChaosClothSimConfig);
-	check(ChaosClothSimConfig);	
+	check(ChaosClothSimConfig);
 
 	ClothingSimulationContext Context;
     FillContext(InOwnerComponent, 0, &Context);
@@ -205,6 +198,12 @@ void ClothingSimulation::CreateActor(USkeletalMeshComponent* InOwnerComponent, U
 	const uint32 Offset = Particles.Size();
 	Particles.AddParticles(PhysMesh->Vertices.Num());
 
+	// ClothSharedSimConfig should either be a nullptr, or point to an object common to the whole skeletal mesh
+	if (ClothSharedSimConfig == nullptr)
+	{
+		ClothSharedSimConfig = Cast<UChaosClothSharedSimConfig>(Asset->ClothSharedSimConfig);
+	}
+
 	AnimationPositions.SetNum(Particles.Size());
 	AnimationNormals.SetNum(Particles.Size());
 
@@ -477,6 +476,33 @@ void ClothingSimulation::CreateActor(USkeletalMeshComponent* InOwnerComponent, U
 	ExternalCollisionsOffset = CollisionParticles.Size();
 }
 
+void ClothingSimulation::PostActorCreationInitialize()
+{
+	if (ClothSharedSimConfig == nullptr)
+	{
+		check(Assets.Num())
+		// None of the cloth assets had a clothSharedSimConfig, so we will create it
+		ClothSharedSimConfig = NewObject<UChaosClothSharedSimConfig>(Assets[0], UChaosClothSharedSimConfig::StaticClass()->GetFName());
+	}
+	check(ClothSharedSimConfig);
+
+	// Let all assets point to the same shared configuration
+	for (UClothingAssetCommon* Asset : Assets)
+	{
+		if (Asset)
+		{
+			Asset->ClothSharedSimConfig = ClothSharedSimConfig;
+		}
+	}
+
+	// Now set all the common parameters on the simulation
+	Evolution->SetIterations(ClothSharedSimConfig->IterationCount);
+	Evolution->SetSelfCollisionThickness(ClothSharedSimConfig->SelfCollisionThickness);
+	Evolution->SetCollisionThickness(ClothSharedSimConfig->CollisionThickness);
+	Evolution->SetDamping(ClothSharedSimConfig->Damping);
+	Evolution->GetGravityForces().SetAcceleration(Chaos::TVector<float, 3>(ClothSharedSimConfig->Gravity));
+}
+
 void ClothingSimulation::UpdateCollisionTransforms(const ClothingSimulationContext& Context)
 {
 	TGeometryClothParticles<float, 3>& CollisionParticles = Evolution->CollisionParticles();
@@ -1443,7 +1469,7 @@ void ClothingSimulation::DebugDrawSelfCollision(USkeletalMeshComponent* OwnerCom
 					const FVector ParticlePosition =
 						RootBoneTransform.TransformPosition(
 							Particles.X(PhysMesh->SelfCollisionIndices[SelfColIdx]));
-					DrawWireSphere(PDI, ParticlePosition, FColor::White, SelfCollisionThickness, 8, SDPG_World, 0.0f, 0.001f);
+					DrawWireSphere(PDI, ParticlePosition, FColor::White, Evolution->GetSelfCollisionThickness(), 8, SDPG_World, 0.0f, 0.001f);
 				}
 			}
 		}
diff --git a/Engine/Plugins/Experimental/ChaosCloth/Source/Chaos/Public/ChaosCloth/ChaosClothConfig.h b/Engine/Plugins/Experimental/ChaosCloth/Source/Chaos/Public/ChaosCloth/ChaosClothConfig.h
--- a/Engine/Plugins/Experimental/ChaosCloth/Source/Chaos/Public/ChaosCloth/ChaosClothConfig.h
+++ b/Engine/Plugins/Experimental/ChaosCloth/Source/Chaos/Public/ChaosCloth/ChaosClothConfig.h
@@ -8,7 +8,7 @@
 
 /** Holds initial, asset level config for clothing actors. */
 // Hiding categories that will be used in the future
-UCLASS(HideCategories = (Solver, Collision))
+UCLASS(HideCategories = (Collision))
 class CHAOSCLOTH_API UChaosClothConfig : public UClothConfigBase
 {
 	GENERATED_BODY()
@@ -44,12 +44,7 @@ public:
 
 	// This is a lower bound to cloth particle masses
 	UPROPERTY(EditAnywhere, Category = MassConfig)
-	float MinPerParticleMass = 0.0001f;
-
-	// The Number of iterations used in the cloth solver
-	// All cloth constraints will be more stiff when increased
-	UPROPERTY(EditAnywhere, Category = Solver)
-	int32 NumIterations = 1;
+	float MinPerParticleMass = 0.0001f;	
 
 	// The Stiffness of the Edge constraints
 	UPROPERTY(EditAnywhere, Category = Stiffness, meta = (UIMin = "0", UIMax = "1", ClampMin = "0", ClampMax = "1"))
@@ -73,25 +68,11 @@ public:
 	
 	// The stiffness of the shape target constraints
 	UPROPERTY(EditAnywhere, Category = Stiffness, meta = (UIMin = "0", UIMax = "1", ClampMin = "0", ClampMax = "1"))
-	float ShapeTargetStiffness = 0.f;
-
-	// The radius of the self collisions spheres
-	UPROPERTY(EditAnywhere, Category = Collision)
-	float SelfCollisionThickness = 2.f;
-
-	// The thickness of the cloth point spheres when colliding against collisions primitives
-	UPROPERTY(EditAnywhere, Category = Collision)
-	float CollisionThickness = 1.2;
+	float ShapeTargetStiffness = 0.f;	
 
 	// Friction coefficient for cloth - collider interaction
 	UPROPERTY(EditAnywhere, Category = Collision, meta = (UIMin = "0", UIMax = "1", ClampMin = "0", ClampMax = "10"))
-	float CoefficientOfFriction = 0.0f;
-
-	UPROPERTY(EditAnywhere, Category = Solver, meta = (UIMin = "0", UIMax = "1", ClampMin = "0", ClampMax = "1"))
-	float Damping = 0.01f;
-
-	UPROPERTY(EditAnywhere, Category = Solver)
-	float GravityMagnitude = 490.f;
+	float CoefficientOfFriction = 0.0f;	
 
 	// Default spring stiffness for anim drive if an anim drive is in use
 	UPROPERTY(EditAnywhere, Category = Stiffness, meta = (UIMin = "0", UIMax = "1", ClampMin = "0", ClampMax = "1"))
@@ -117,3 +98,37 @@ public:
 	UPROPERTY(EditAnywhere, Category = ClothEnableFlags)
 	bool bUseContinuousCollisionDetection = false;
 };
+
+/*
+These settings are shared between all instances on a skeletal mesh
+*/
+UCLASS()
+class CHAOSCLOTH_API UChaosClothSharedSimConfig : public UClothSharedSimConfigBase
+{
+	GENERATED_BODY()
+public:
+	UChaosClothSharedSimConfig() : Gravity(FVector(0.0f, 0.0f, -490)) {};
+	virtual ~UChaosClothSharedSimConfig() {};
+
+	// The number of solver iterations
+	// This will increase the stiffness of all constraints but will increase the CPU cost
+	UPROPERTY(EditAnywhere, Category = Simulation, meta = (UIMin = "0", UIMax = "20", ClampMin = "0", ClampMax = "100"))
+	int32 IterationCount = 1;
+
+	// The radius of the spheres used in self collision 
+	UPROPERTY(EditAnywhere, Category = Collision, meta = (UIMin = "0", UIMax = "100", ClampMin = "0", ClampMax = "1000"))
+	float SelfCollisionThickness = 2.0f;
+
+	// The radius of cloth points when considering collisions against collider shapes
+	UPROPERTY(EditAnywhere, Category = Collision, meta = (UIMin = "0", UIMax = "100", ClampMin = "0", ClampMax = "1000"))
+	float CollisionThickness = 1.0f;
+
+	//The amount of cloth damping
+	UPROPERTY(EditAnywhere, Category = Simulation, meta = (UIMin = "0", UIMax = "1", ClampMin = "0", ClampMax = "1"))
+	float Damping = 0.01f;
+
+	// The gravitational acceleration vector [cm/s^2]
+	UPROPERTY(EditAnywhere, Category = Simulation)
+	FVector Gravity;
+};
+
diff --git a/Engine/Plugins/Experimental/ChaosCloth/Source/Chaos/Public/ChaosCloth/ChaosClothingSimulation.h b/Engine/Plugins/Experimental/ChaosCloth/Source/Chaos/Public/ChaosCloth/ChaosClothingSimulation.h
--- a/Engine/Plugins/Experimental/ChaosCloth/Source/Chaos/Public/ChaosCloth/ChaosClothingSimulation.h
+++ b/Engine/Plugins/Experimental/ChaosCloth/Source/Chaos/Public/ChaosCloth/ChaosClothingSimulation.h
@@ -8,6 +8,7 @@
 #include "Chaos/PBDEvolution.h"
 #include "Chaos/Transform.h"
 #include "Chaos/TriangleMesh.h"
+#include "ChaosCloth/ChaosClothConfig.h"
 #include "Components/SkeletalMeshComponent.h"
 
 namespace Chaos
@@ -24,8 +25,7 @@ namespace Chaos
 		FTransform ComponentToWorld;
 	};
 
-	class ClothingSimulation
-		: public IClothingSimulation
+	class ClothingSimulation : public IClothingSimulation
 #if WITH_EDITOR
 		, public FGCObject  // Add garbage collection for debug cloth material
 #endif  // #if WITH_EDITOR
@@ -56,6 +56,7 @@ namespace Chaos
 		// IClothingSimulation interface
 		void Initialize() override;
 		void CreateActor(USkeletalMeshComponent* InOwnerComponent, UClothingAssetBase* InAsset, int32 SimDataIndex) override;
+		void PostActorCreationInitialize() override;
 		IClothingSimulationContext* CreateContext() override { return new ClothingSimulationContext(); }
 		void FillContext(USkeletalMeshComponent* InComponent, float InDeltaTime, IClothingSimulationContext* InOutContext) override;
 		void Shutdown() override;
@@ -91,6 +92,7 @@ namespace Chaos
 	private:
 		// Assets
 		TArray<UClothingAssetCommon*> Assets;
+		UChaosClothSharedSimConfig* ClothSharedSimConfig;
 
 		// Collision Data
 		FClothCollisionData ExtractedCollisions;  // Collisions extracted from the referenced physics asset
@@ -120,13 +122,6 @@ namespace Chaos
 		float DeltaTime;
 		float MaxDeltaTime;
 		float ClampDeltaTime;
-		// Parameters that should be set in the ui
-		int32 NumIterations;
-		float SelfCollisionThickness;
-		float CollisionThickness;
-		float CoefficientOfFriction;
-		float Damping;
-		float GravityMagnitude;
 
 #if WITH_EDITOR
 		// Visualization material
diff --git a/Engine/Source/Runtime/ClothingSystemRuntimeCommon/Public/ClothingAsset.h b/Engine/Source/Runtime/ClothingSystemRuntimeCommon/Public/ClothingAsset.h
--- a/Engine/Source/Runtime/ClothingSystemRuntimeCommon/Public/ClothingAsset.h
+++ b/Engine/Source/Runtime/ClothingSystemRuntimeCommon/Public/ClothingAsset.h
@@ -158,6 +158,12 @@ public:
 	UPROPERTY(EditAnywhere, Category = Config)
 	UPhysicsAsset* PhysicsAsset;
 
+	// Shared by all cloth instances in a skeletal mesh
+	// Only supported with Chaos Cloth for now
+	// This may not be editable on unused cloth assets
+	UPROPERTY(EditAnywhere, Category = Config, NoClear, meta = (NoResetToDefault))
+	UClothSharedSimConfigBase* ClothSharedSimConfig;
+
 	// Parameters for how the NVcloth behaves.
 	// These will have no effect on Chaos cloth
 	UPROPERTY(EditAnywhere, Category = Config, NoClear, meta = (NoResetToDefault))
diff --git a/Engine/Source/Runtime/ClothingSystemRuntimeInterface/Public/ClothConfig.h b/Engine/Source/Runtime/ClothingSystemRuntimeInterface/Public/ClothConfig.h
--- a/Engine/Source/Runtime/ClothingSystemRuntimeInterface/Public/ClothConfig.h
+++ b/Engine/Source/Runtime/ClothingSystemRuntimeInterface/Public/ClothConfig.h
@@ -9,6 +9,7 @@ class UClothingAssetBase;
 
 /**
  * Base class for simulator specific simulation controls.
+ * Each cloth instance on a skeletal mesh can have a unique cloth config
  */
 UCLASS(Abstract, DefaultToInstanced)
 class CLOTHINGSYSTEMRUNTIMEINTERFACE_API UClothConfigBase : public UObject
@@ -19,29 +20,21 @@ public:
 	virtual ~UClothConfigBase();
 
 	virtual bool HasSelfCollision() const
-	{ unimplemented(); return false; }
-
-#if WITH_EDITOR
-	/** 
-	 * Callback invoked by \c UClothingAssetBase::PostEditChangeChainProperty(). 
-	 * Override to add custom handling of editor events.
-	 */
-	virtual bool PostEditChangeChainPropertyCallback(
-		UClothingAssetBase* Asset,
-		FPropertyChangedChainEvent& Event)
 	{
-		return false; 
+		unimplemented(); return false;
 	}
+};
 
-	/**
-	 * Callback invoked by \c UClothingAssetBase::OnFinishedChangingClothPropertiesCallback().
-	 * Override to add custom handling of editor events.
-	 */
-	virtual bool OnFinishedChangingClothingPropertiesCallback(
-		UClothingAssetBase* Asset,
-		const FPropertyChangedEvent& Event)
-	{
-		return false;
-	}
-#endif
+
+/*
+These settings are shared between all instances on a skeletal mesh
+*/
+UCLASS(Abstract)
+class CLOTHINGSYSTEMRUNTIMEINTERFACE_API UClothSharedSimConfigBase : public UObject
+{
+	GENERATED_BODY()
+public:
+	UClothSharedSimConfigBase() {};
+	virtual ~UClothSharedSimConfigBase() {};
+	
 };
diff --git a/Engine/Source/Runtime/ClothingSystemRuntimeInterface/Public/ClothingSimulation.h b/Engine/Source/Runtime/ClothingSystemRuntimeInterface/Public/ClothingSimulation.h
--- a/Engine/Source/Runtime/ClothingSystemRuntimeInterface/Public/ClothingSimulation.h
+++ b/Engine/Source/Runtime/ClothingSystemRuntimeInterface/Public/ClothingSimulation.h
@@ -37,6 +37,9 @@ protected:
 	 * @param SimDataIndex - the sim data index to use when doing the writeback in GetSimulationData
 	 */
 	virtual void CreateActor(USkeletalMeshComponent* InOwnerComponent, UClothingAssetBase* InAsset, int32 SimDataIndex) = 0;
+
+	/** Do any required initialization after all the actors have been loaded */
+	virtual void PostActorCreationInitialize() = 0;
 	
 	/** Create a new context, will not be filled, call FillContext before simulating with this context */
 	virtual IClothingSimulationContext* CreateContext() = 0;
diff --git a/Engine/Source/Runtime/ClothingSystemRuntimeNv/Public/ClothingSimulationNv.h b/Engine/Source/Runtime/ClothingSystemRuntimeNv/Public/ClothingSimulationNv.h
--- a/Engine/Source/Runtime/ClothingSystemRuntimeNv/Public/ClothingSimulationNv.h
+++ b/Engine/Source/Runtime/ClothingSystemRuntimeNv/Public/ClothingSimulationNv.h
@@ -198,6 +198,7 @@ public:
 
 	// IClothingSimulation Interface
 	virtual void CreateActor(USkeletalMeshComponent* InOwnerComponent, UClothingAssetBase* InAsset, int32 InSimDataIndex) override;
+	virtual void PostActorCreationInitialize() override {};
 	virtual IClothingSimulationContext* CreateContext() override;
 	virtual void FillContext(USkeletalMeshComponent* InComponent, float InDeltaTime, IClothingSimulationContext* InOutContext) override;
 	virtual void Initialize() override;
diff --git a/Engine/Source/Runtime/Engine/Private/SkeletalMeshComponentPhysics.cpp b/Engine/Source/Runtime/Engine/Private/SkeletalMeshComponentPhysics.cpp
--- a/Engine/Source/Runtime/Engine/Private/SkeletalMeshComponentPhysics.cpp
+++ b/Engine/Source/Runtime/Engine/Private/SkeletalMeshComponentPhysics.cpp
@@ -2372,7 +2372,7 @@ void USkeletalMeshComponent::RecreateClothingActors()
 				}
 				ClothingSimulation->CreateActor(this, Asset, BaseAssetIndex);
 			}
-
+			ClothingSimulation->PostActorCreationInitialize();
 			WritebackClothingSimulationData();
 		}
 	}
diff --git a/Engine/Source/Runtime/Experimental/Chaos/Public/Chaos/PBDEvolution.h b/Engine/Source/Runtime/Experimental/Chaos/Public/Chaos/PBDEvolution.h
--- a/Engine/Source/Runtime/Experimental/Chaos/Public/Chaos/PBDEvolution.h
+++ b/Engine/Source/Runtime/Experimental/Chaos/Public/Chaos/PBDEvolution.h
@@ -42,6 +42,18 @@ class CHAOS_API TPBDEvolution
 	int32 GetIterations() const { return MNumIterations; }
 	void SetIterations(const int32 Iterations) { MNumIterations = Iterations; }
 
+	T GetSelfCollisionThickness() const { return MSelfCollisionThickness; }
+	void SetSelfCollisionThickness(const T SelfCollisionThickness) { MSelfCollisionThickness = SelfCollisionThickness; }
+
+	T GetCollisionThickness() const { return MCollisionThickness; }
+	void SetCollisionThickness(const T CollisionThickness) { MCollisionThickness = CollisionThickness; }
+
+	T GetCoefficientOfFriction() const { return MCoefficientOfFriction; }
+	void SetCoefficientOfFriction(const T CoefficientOfFriction) { MCoefficientOfFriction = CoefficientOfFriction; }
+
+	T GetDamping() const { return MDamping; }
+	void SetDamping(const T Damping) { MDamping = Damping; }
+
 	T GetTime() const { return MTime; }
 
   private:
