Reworked the Chaos Cloth Convex collisions to use a Chaos TConvex implicit.
Also:
-Regenerate the convexes' surface points for the legacy apex cloth assets (apparently the only use of convexes in apex import was to emulate box collisions).
-Add debug draw visualization for the Chaos Cloth editor.

diff --git a/Engine/Plugins/Experimental/ChaosCloth/Source/Chaos/Private/ChaosCloth/ChaosClothingSimulation.cpp b/Engine/Plugins/Experimental/ChaosCloth/Source/Chaos/Private/ChaosCloth/ChaosClothingSimulation.cpp
--- a/Engine/Plugins/Experimental/ChaosCloth/Source/Chaos/Private/ChaosCloth/ChaosClothingSimulation.cpp
+++ b/Engine/Plugins/Experimental/ChaosCloth/Source/Chaos/Private/ChaosCloth/ChaosClothingSimulation.cpp
@@ -669,7 +669,7 @@ void ClothingSimulation::ExtractPhysicsAssetCollisions(UClothingAssetCommon* Ass
 				Convex.BoneIndex = MappedBoneIndex;
 #if WITH_PHYSX
 				// Collision bodies are stored in PhysX specific data structures so they can only be imported if we enable PhysX.
-				const physx::PxConvexMesh* const PhysXMesh = ConvexElem.GetConvexMesh();
+				const physx::PxConvexMesh* const PhysXMesh = ConvexElem.GetConvexMesh();  // TODO(Kriss.Gossart): Deal with this legacy structure in a different place, so that there's only TConvex
 				const int32 NumPolygons = int32(PhysXMesh->getNbPolygons());
 				Convex.Planes.SetNumUninitialized(NumPolygons);
 				for (int32 i = 0; i < NumPolygons; ++i)
@@ -684,15 +684,29 @@ void ClothingSimulation::ExtractPhysicsAssetCollisions(UClothingAssetCommon* Ass
 						ConvexElem.VertexData[Indices[1]],
 						ConvexElem.VertexData[Indices[2]]);
 				}
+
+				// Rebuild surface points
+				Convex.RebuildSurfacePoints();
+
 #elif WITH_CHAOS  // #if WITH_PHYSX
 				const Chaos::TImplicitObject<float, 3>& ChaosConvexMesh = *ConvexElem.GetChaosConvexMesh();
 				const Chaos::TConvex<float, 3>& ChaosConvex = ChaosConvexMesh.GetObjectChecked<Chaos::TConvex<float, 3>>();
+
+				// Copy planes
 				const TArray<TPlane<float, 3>>& Planes = ChaosConvex.GetFaces();
 				Convex.Planes.Reserve(Planes.Num());
 				for (const TPlane<float, 3>& Plane : Planes)
 				{
 					Convex.Planes.Add(FPlane(Plane.X(), Plane.Normal()));
 				}
+
+				// Copy surface points
+				const uint32 NumSurfacePoints = ChaosConvex.GetSurfaceParticles().Size();
+				Convex.SurfacePoints.Reserve(NumSurfacePoints);
+				for (uint32 ParticleIndex = 0; ParticleIndex < NumSurfacePoints; ++ParticleIndex)
+				{
+					Convex.SurfacePoints.Add(ChaosConvex.GetSurfaceParticles().X(ParticleIndex));
+				}
 #endif  // #if WITH_PHYSX #elif WITH_CHAOS
 
 				// Add extracted collision data
@@ -861,37 +875,58 @@ void ClothingSimulation::AddCollisions(const FClothCollisionData& ClothCollision
 			BoneIndices[i] = GetMappedBoneIndex(UsedBoneIndices, Convex.BoneIndex);
 			UE_LOG(LogChaosCloth, VeryVerbose, TEXT("Found collision convex on bone index %d."), BoneIndices[i]);
 
-			// Add convex planes
-			TArray<TUniquePtr<TImplicitObject<float, 3>>> Planes;
-			Planes.Reserve(Convex.Planes.Num());
-			for (const FPlane& Plane : Convex.Planes)
+			const int32 NumSurfacePoints = Convex.SurfacePoints.Num();
+			const int32 NumPlanes = Convex.Planes.Num();
+
+			if (NumSurfacePoints < 4)
+			{
+				UE_LOG(LogChaosCloth, Warning, TEXT("Invalid convex collision: not enough surface points."));
+			}
+			else if (NumPlanes < 4)
+			{
+				UE_LOG(LogChaosCloth, Warning, TEXT("Invalid convex collision: not enough planes."));
+			}
+			else
 			{
-				FPlane NormalizedPlane(Plane);
-				if (NormalizedPlane.Normalize())
+				// Retrieve convex planes
+				TArray<TPlane<float, 3>> Planes;
+				Planes.Reserve(Convex.Planes.Num());
+				for (const FPlane& Plane : Convex.Planes)
 				{
-					const Chaos::TVector<float, 3> Normal(static_cast<FVector>(NormalizedPlane));
-					const Chaos::TVector<float, 3> Base = Normal * NormalizedPlane.W;
+					FPlane NormalizedPlane(Plane);
+					if (NormalizedPlane.Normalize())
+					{
+						const Chaos::TVector<float, 3> Normal(static_cast<FVector>(NormalizedPlane));
+						const Chaos::TVector<float, 3> Base = Normal * NormalizedPlane.W;
 
-					Planes.Add(TUniquePtr<Chaos::TImplicitObject<float, 3>>(
-						new Chaos::TPlane<float, 3>(Base, Normal)));
+						Planes.Add(Chaos::TPlane<float, 3>(Base, Normal));
+					}
+					else
+					{
+						UE_LOG(LogChaosCloth, Warning, TEXT("Invalid convex collision: bad plane normal."));
+						break;
+					}
 				}
-				else
+
+				if (Planes.Num() == Convex.Planes.Num())
 				{
-					UE_LOG(LogChaosCloth, Warning, TEXT("Invalid convex collision plane."));
+					// Retrieve particles
+					TParticles<float, 3> SurfaceParticles;
+					SurfaceParticles.Resize(NumSurfacePoints);
+					for (int32 ParticleIndex = 0; ParticleIndex < NumSurfacePoints; ++ParticleIndex)
+					{
+						SurfaceParticles.X(ParticleIndex) = Convex.SurfacePoints[ParticleIndex];
+					}
+
+					// Setup the collision particle geometry
+					CollisionParticles.SetDynamicGeometry(i, MakeUnique<Chaos::TConvex<float, 3>>(MoveTemp(Planes), MoveTemp(SurfaceParticles)));
 				}
 			}
 
-			// Set the collision particle geometry
-			if (Planes.Num())
+			if (!CollisionParticles.DynamicGeometry(i))
 			{
-				CollisionParticles.SetDynamicGeometry(i, MakeUnique<Chaos::TImplicitObjectIntersection<float, 3>>(MoveTemp(Planes)));  // TODO(Kriss.Gossart): Replace by TConvex once plane constructor has been added
-			}
-			else
-			{
-				UE_LOG(LogChaosCloth, Warning, TEXT("Invalid convex collision."));
-				CollisionParticles.SetDynamicGeometry(
-					i,
-					MakeUnique<Chaos::TSphere<float, 3>>(Chaos::TVector<float, 3>(0.0f), 1.0f));  // Default to a unit sphere to replace the faulty convex
+				UE_LOG(LogChaosCloth, Warning, TEXT("Replacing invalid convex collision by a default unit sphere."));
+				CollisionParticles.SetDynamicGeometry(i, MakeUnique<Chaos::TSphere<float, 3>>(Chaos::TVector<float, 3>(0.0f), 1.0f));  // Default to a unit sphere to replace the faulty convex
 			}
 		}
 	}
@@ -1380,6 +1415,43 @@ void ClothingSimulation::DebugDrawCollision(USkeletalMeshComponent* OwnerCompone
 		}
 	};
 
+	auto DrawConvex = [&PDI](const Chaos::TConvex<float, 3>& Convex, const TRotation<float, 3>& Rotation, const Chaos::TVector<float, 3>& Position, const FLinearColor& Color)
+	{
+		const TArray<Chaos::TPlane<float, 3>>& Planes = Convex.GetFaces();
+		for (int32 PlaneIndex1 = 0; PlaneIndex1 < Planes.Num(); ++PlaneIndex1)
+		{
+			const Chaos::TPlane<float, 3>& Plane1 = Planes[PlaneIndex1];
+
+			for (int32 PlaneIndex2 = PlaneIndex1 + 1; PlaneIndex2 < Planes.Num(); ++PlaneIndex2)
+			{
+				const Chaos::TPlane<float, 3>& Plane2 = Planes[PlaneIndex2];
+
+				// Find the two surface points that belong to both Plane1 and Plane2
+				uint32 ParticleIndex1 = INDEX_NONE;
+
+				const Chaos::TParticles<float, 3>& SurfaceParticles = Convex.GetSurfaceParticles();
+				for (uint32 ParticleIndex = 0; ParticleIndex < SurfaceParticles.Size(); ++ParticleIndex)
+				{
+					const Chaos::TVector<float, 3>& X = SurfaceParticles.X(ParticleIndex);
+
+					if (FMath::Square(Plane1.SignedDistance(X)) < KINDA_SMALL_NUMBER && 
+						FMath::Square(Plane2.SignedDistance(X)) < KINDA_SMALL_NUMBER)
+					{
+						if (ParticleIndex1 != INDEX_NONE)
+						{
+							const Chaos::TVector<float, 3>& X1 = SurfaceParticles.X(ParticleIndex1);
+							const FVector Position1 = Position + Rotation.RotateVector(X1);
+							const FVector Position2 = Position + Rotation.RotateVector(X);
+							PDI->DrawLine(Position1, Position2, Color, SDPG_World, 0.0f, 0.001f, false);
+							break;
+						}
+						ParticleIndex1 = ParticleIndex;
+					}
+				}
+			}
+		}
+	};
+
 	static const FLinearColor MappedColor(FColor::Cyan);
 	static const FLinearColor UnmappedColor(FColor::Red);
 
@@ -1431,6 +1503,10 @@ void ClothingSimulation::DebugDrawCollision(USkeletalMeshComponent* OwnerCompone
 				}
 				break;
 	
+			case Chaos::ImplicitObjectType::Convex:
+				DrawConvex(Object->GetObjectChecked<Chaos::TConvex<float, 3>>(), Rotation, Position, Color);
+				break;
+
 			default:
 				DrawCoordinateSystem(PDI, Position, FRotator(Rotation), 10.0f, SDPG_World, 0.1f);  // Draw everything else as a coordinate for now
 				break;
diff --git a/Engine/Source/Runtime/ClothingSystemRuntimeInterface/Private/ClothCollisionPrim.cpp b/Engine/Source/Runtime/ClothingSystemRuntimeInterface/Private/ClothCollisionPrim.cpp
new file mode 100644
--- /dev/null
+++ b/Engine/Source/Runtime/ClothingSystemRuntimeInterface/Private/ClothCollisionPrim.cpp
@@ -0,0 +1,44 @@
+// Copyright 1998-2019 Epic Games, Inc. All Rights Reserved.
+
+#include "ClothCollisionPrim.h"
+
+void FClothCollisionPrim_Convex::RebuildSurfacePoints()
+{
+	const int32 NumPlanes = Planes.Num();
+	if (NumPlanes >= 3)
+	{
+		SurfacePoints.Reset(NumPlanes * (NumPlanes - 1) * (NumPlanes - 2) / 6);  // Maximum number of intersections
+
+		auto PointInHull = [this](const FVector& Point) -> bool
+		{
+			for (const FPlane& Plane : Planes)
+			{
+				if (Plane.PlaneDot(Point) > KINDA_SMALL_NUMBER)
+				{
+					return false;
+				}
+			}
+			return true;
+		};
+
+		for (int32 Index0 = 0; Index0 < NumPlanes; ++Index0)
+		{
+			for (int32 Index1 = Index0 + 1; Index1 < NumPlanes; ++Index1)
+			{
+				for (int32 Index2 = Index1 + 1; Index2 < NumPlanes; ++Index2)
+				{
+					FVector Intersection;
+					if (FMath::IntersectPlanes3(Intersection, Planes[Index0], Planes[Index1], Planes[Index2]) &&
+						PointInHull(Intersection))
+					{
+						SurfacePoints.Add(Intersection);
+					}
+				}
+			}
+		}
+	}
+	else
+	{
+		SurfacePoints.Reset();
+	}
+}
diff --git a/Engine/Source/Runtime/ClothingSystemRuntimeInterface/Public/ClothCollisionPrim.h b/Engine/Source/Runtime/ClothingSystemRuntimeInterface/Public/ClothCollisionPrim.h
--- a/Engine/Source/Runtime/ClothingSystemRuntimeInterface/Public/ClothCollisionPrim.h
+++ b/Engine/Source/Runtime/ClothingSystemRuntimeInterface/Public/ClothCollisionPrim.h
@@ -54,7 +54,7 @@ struct FClothCollisionPrim_SphereConnection
  *	shape created by the planes combined.
  */
 USTRUCT()
-struct FClothCollisionPrim_Convex
+struct CLOTHINGSYSTEMRUNTIMEINTERFACE_API FClothCollisionPrim_Convex
 {
 	GENERATED_BODY()
 
@@ -62,9 +62,17 @@ struct FClothCollisionPrim_Convex
 		: BoneIndex(INDEX_NONE)
 	{}
 
+	/** Rebuild the surface point array from the existing planes.
+	 *  This is an expensive function (O(n^4) per number of planes).
+	 */
+	void RebuildSurfacePoints();
+
 	UPROPERTY()
 	TArray<FPlane> Planes;
 
+	UPROPERTY()
+	TArray<FVector> SurfacePoints;  // Surface points, used by Chaos and also for visualization
+
 	UPROPERTY()
 	int32 BoneIndex;
 };
diff --git a/Engine/Source/Runtime/ClothingSystemRuntimeNv/Private/ClothLODDataNv.cpp b/Engine/Source/Runtime/ClothingSystemRuntimeNv/Private/ClothLODDataNv.cpp
--- a/Engine/Source/Runtime/ClothingSystemRuntimeNv/Private/ClothLODDataNv.cpp
+++ b/Engine/Source/Runtime/ClothingSystemRuntimeNv/Private/ClothLODDataNv.cpp
@@ -9,3 +9,28 @@ UClothLODDataNv::UClothLODDataNv(const FObjectInitializer& Init)
 
 UClothLODDataNv::~UClothLODDataNv()
 {}
+
+void FClothLODData::MigrateTo(UClothLODDataNv* LodData) const
+{
+	PhysicalMeshData.MigrateTo(LodData->PhysicalMeshData);
+	LodData->CollisionData = CollisionData;
+#if WITH_CHAOS
+	// Rebuild surface points so that the legacy Apex convex collision data can also be used with Chaos
+	for (auto& Convex : LodData->CollisionData.Convexes)
+	{
+		if (!Convex.SurfacePoints.Num())
+		{
+			Convex.RebuildSurfacePoints();
+		}
+	}
+#endif // #if WITH_CHAOS
+#if WITH_EDITORONLY_DATA
+	LodData->ParameterMasks.SetNum(ParameterMasks.Num());
+	for (int i = 0; i < ParameterMasks.Num(); i++)
+	{
+		ParameterMasks[i].MigrateTo(&LodData->ParameterMasks[i]);
+	}
+#endif // WITH_EDITORONLY_DATA
+	LodData->TransitionUpSkinData = TransitionUpSkinData;
+	LodData->TransitionDownSkinData = TransitionDownSkinData;
+}
diff --git a/Engine/Source/Runtime/ClothingSystemRuntimeNv/Public/ClothLODDataNv.h b/Engine/Source/Runtime/ClothingSystemRuntimeNv/Public/ClothLODDataNv.h
--- a/Engine/Source/Runtime/ClothingSystemRuntimeNv/Public/ClothLODDataNv.h
+++ b/Engine/Source/Runtime/ClothingSystemRuntimeNv/Public/ClothLODDataNv.h
@@ -64,21 +64,10 @@ struct CLOTHINGSYSTEMRUNTIMENV_API FClothLODData
 		return true;
 	}
 
-	void MigrateTo(UClothLODDataNv* LodData) const
-	{
-		PhysicalMeshData.MigrateTo(LodData->PhysicalMeshData);
-		LodData->CollisionData = CollisionData;
-#if WITH_EDITORONLY_DATA
-		LodData->ParameterMasks.SetNum(ParameterMasks.Num());
-		for (int i = 0; i < ParameterMasks.Num(); i++)
-		{
-			ParameterMasks[i].MigrateTo(&LodData->ParameterMasks[i]);
-		}
-#endif // WITH_EDITORONLY_DATA
-		LodData->TransitionUpSkinData = TransitionUpSkinData;
-		LodData->TransitionDownSkinData = TransitionDownSkinData;
-	}
+	// Migrate this legacy Nv struct to the new class format (called by UClothingAssetNv::Postload())
+	void MigrateTo(UClothLODDataNv* LodData) const;
 };
+
 template<>
 struct TStructOpsTypeTraits<FClothLODData> : public TStructOpsTypeTraitsBase2<FClothLODData>
 {
