Bug Fix: ChaosCloth Make sure that new shared cloth configs have valid owners

diff --git a/Engine/Plugins/Experimental/ChaosCloth/Source/Chaos/Private/ChaosCloth/ChaosClothingSimulation.cpp b/Engine/Plugins/Experimental/ChaosCloth/Source/Chaos/Private/ChaosCloth/ChaosClothingSimulation.cpp
--- a/Engine/Plugins/Experimental/ChaosCloth/Source/Chaos/Private/ChaosCloth/ChaosClothingSimulation.cpp
+++ b/Engine/Plugins/Experimental/ChaosCloth/Source/Chaos/Private/ChaosCloth/ChaosClothingSimulation.cpp
@@ -478,23 +478,23 @@ void ClothingSimulation::CreateActor(USkeletalMeshComponent* InOwnerComponent, U
 
 void ClothingSimulation::PostActorCreationInitialize()
 {
-	if (ClothSharedSimConfig == nullptr)
-	{
-		check(Assets.Num())
-		// None of the cloth assets had a clothSharedSimConfig, so we will create it
-		ClothSharedSimConfig = NewObject<UChaosClothSharedSimConfig>(Assets[0], UChaosClothSharedSimConfig::StaticClass()->GetFName());
-	}
-	check(ClothSharedSimConfig);
-
 	// Let all assets point to the same shared configuration
 	for (UClothingAssetCommon* Asset : Assets)
 	{
 		if (Asset)
 		{
+			// If we don't have a sim config, create one
+			if (ClothSharedSimConfig == nullptr)
+			{
+					// None of the cloth assets had a clothSharedSimConfig, so we will create it
+					ClothSharedSimConfig = NewObject<UChaosClothSharedSimConfig>(Asset, UChaosClothSharedSimConfig::StaticClass()->GetFName());
+			}
 			Asset->ClothSharedSimConfig = ClothSharedSimConfig;
 		}
 	}
 
+	check(ClothSharedSimConfig);
+
 	// Now set all the common parameters on the simulation
 	Evolution->SetIterations(ClothSharedSimConfig->IterationCount);
 	Evolution->SetSelfCollisionThickness(ClothSharedSimConfig->SelfCollisionThickness);
