Replace TImplicitObjects with FImplicitObjects
#rnx

diff --git a/Engine/Plugins/Experimental/ChaosCloth/Source/Chaos/Private/ChaosCloth/ChaosClothingSimulation.cpp b/Engine/Plugins/Experimental/ChaosCloth/Source/Chaos/Private/ChaosCloth/ChaosClothingSimulation.cpp
--- a/Engine/Plugins/Experimental/ChaosCloth/Source/Chaos/Private/ChaosCloth/ChaosClothingSimulation.cpp
+++ b/Engine/Plugins/Experimental/ChaosCloth/Source/Chaos/Private/ChaosCloth/ChaosClothingSimulation.cpp
@@ -689,7 +689,7 @@ void ClothingSimulation::ExtractPhysicsAssetCollisions(UClothingAssetCommon* Ass
 				Convex.RebuildSurfacePoints();
 
 #elif WITH_CHAOS  // #if WITH_PHYSX
-				const Chaos::TImplicitObject<float, 3>& ChaosConvexMesh = *ConvexElem.GetChaosConvexMesh();
+				const Chaos::FImplicitObject& ChaosConvexMesh = *ConvexElem.GetChaosConvexMesh();
 				const Chaos::TConvex<float, 3>& ChaosConvex = ChaosConvexMesh.GetObjectChecked<Chaos::TConvex<float, 3>>();
 
 				// Copy planes
@@ -811,13 +811,13 @@ void ClothingSimulation::AddCollisions(const FClothCollisionData& ClothCollision
 			else
 			{
 				// Tapered capsule
-				TArray<TUniquePtr<TImplicitObject<float, 3>>> Objects;
+				TArray<TUniquePtr<FImplicitObject>> Objects;
 				Objects.Reserve(3);
-				Objects.Add(TUniquePtr<Chaos::TImplicitObject<float, 3>>(
+				Objects.Add(TUniquePtr<Chaos::FImplicitObject>(
 					new Chaos::TTaperedCylinder<float>(X0, X1, Radius0, Radius1)));
-				Objects.Add(TUniquePtr<Chaos::TImplicitObject<float, 3>>(
+				Objects.Add(TUniquePtr<Chaos::FImplicitObject>(
 					new Chaos::TSphere<float, 3>(X0, Radius0)));
-				Objects.Add(TUniquePtr<Chaos::TImplicitObject<float, 3>>(
+				Objects.Add(TUniquePtr<Chaos::FImplicitObject>(
 					new Chaos::TSphere<float, 3>(X1, Radius1)));
 				CollisionParticles.SetDynamicGeometry(
 					i,
@@ -1459,7 +1459,7 @@ void ClothingSimulation::DebugDrawCollision(USkeletalMeshComponent* OwnerCompone
 
 	for (uint32 Index = 0; Index < CollisionParticles.Size(); ++Index)
 	{
-		if (const TImplicitObject<float, 3>* const Object = CollisionParticles.DynamicGeometry(Index).Get())
+		if (const FImplicitObject* const Object = CollisionParticles.DynamicGeometry(Index).Get())
 		{
 			const uint32 BoneIndex = BoneIndices[Index];
 			const FLinearColor Color = (BoneIndex != INDEX_NONE) ? MappedColor : UnmappedColor;
@@ -1482,9 +1482,9 @@ void ClothingSimulation::DebugDrawCollision(USkeletalMeshComponent* OwnerCompone
 				break;
 
 			case Chaos::ImplicitObjectType::Union:  // Union only used as collision tappered capsules
-				for (const TUniquePtr<TImplicitObject<float, 3>>& SubObjectPtr : Object->GetObjectChecked<Chaos::TImplicitObjectUnion<float, 3>>().GetObjects())
+				for (const TUniquePtr<FImplicitObject>& SubObjectPtr : Object->GetObjectChecked<Chaos::TImplicitObjectUnion<float, 3>>().GetObjects())
 				{
-					if (const TImplicitObject<float, 3>* const SubObject = SubObjectPtr.Get())
+					if (const FImplicitObject* const SubObject = SubObjectPtr.Get())
 					{
 						switch (SubObject->GetType())
 						{
diff --git a/Engine/Source/Runtime/Experimental/Chaos/Public/Chaos/GeometryParticles.h b/Engine/Source/Runtime/Experimental/Chaos/Public/Chaos/GeometryParticles.h
--- a/Engine/Source/Runtime/Experimental/Chaos/Public/Chaos/GeometryParticles.h
+++ b/Engine/Source/Runtime/Experimental/Chaos/Public/Chaos/GeometryParticles.h
@@ -176,7 +176,7 @@ namespace Chaos
 		CHAOS_API const TRotation<T, d>& R(const int32 Index) const { return MR[Index]; }
 		CHAOS_API TRotation<T, d>& R(const int32 Index) { return MR[Index]; }
 
-		CHAOS_API TSerializablePtr<TImplicitObject<T, d>> Geometry(const int32 Index) const { return MGeometry[Index]; }
+		CHAOS_API TSerializablePtr<FImplicitObject> Geometry(const int32 Index) const { return MGeometry[Index]; }
 
 		CHAOS_API const TUniquePtr<FImplicitObject>& DynamicGeometry(const int32 Index) const { return MDynamicGeometry[Index]; }
 
diff --git a/Engine/Source/Runtime/Experimental/ChaosSolvers/Private/PhysicsProxy/GeometryCollectionPhysicsProxy.cpp b/Engine/Source/Runtime/Experimental/ChaosSolvers/Private/PhysicsProxy/GeometryCollectionPhysicsProxy.cpp
--- a/Engine/Source/Runtime/Experimental/ChaosSolvers/Private/PhysicsProxy/GeometryCollectionPhysicsProxy.cpp
+++ b/Engine/Source/Runtime/Experimental/ChaosSolvers/Private/PhysicsProxy/GeometryCollectionPhysicsProxy.cpp
@@ -2483,7 +2483,7 @@ void FGeometryCollectionPhysicsProxy::InitializeSharedCollisionStructures(Chaos:
 				else if (SizeSpecificData.ImplicitType == EImplicitTypeEnum::Chaos_Implicit_Box)
 				{
 					ErrorReporter.SetPrefix(BaseErrorPrefix + " | Cluster Transform Index: " + FString::FromInt(ClusterTransformIdx));
-					CollectionImplicits[ClusterTransformIdx] = TUniquePtr<TImplicitObject<float, 3>>(
+					CollectionImplicits[ClusterTransformIdx] = TUniquePtr<FImplicitObject>(
 						FCollisionStructureManager::NewImplicit(ErrorReporter, MassSpaceParticles, *UnionMesh,
 							InstanceBoundingBox, 0, 0, 0,
 							SizeSpecificData.CollisionObjectReductionPercentage, SizeSpecificData.CollisionType,
@@ -2498,7 +2498,7 @@ void FGeometryCollectionPhysicsProxy::InitializeSharedCollisionStructures(Chaos:
 				else if (SizeSpecificData.ImplicitType == EImplicitTypeEnum::Chaos_Implicit_Sphere)
 				{
 					ErrorReporter.SetPrefix(BaseErrorPrefix + " | Cluster Transform Index: " + FString::FromInt(ClusterTransformIdx));
-					CollectionImplicits[ClusterTransformIdx] = TUniquePtr<TImplicitObject<float, 3>>(
+					CollectionImplicits[ClusterTransformIdx] = TUniquePtr<FImplicitObject>(
 						FCollisionStructureManager::NewImplicit(ErrorReporter, MassSpaceParticles, *UnionMesh,
 							InstanceBoundingBox, InstanceBoundingBox.GetExtent().GetAbsMin() / 2, 0, 0,
 							SizeSpecificData.CollisionObjectReductionPercentage, SizeSpecificData.CollisionType,
