Fix CIS FImplicitObject

diff --git a/Engine/Plugins/Experimental/ChaosCloth/Source/Chaos/Private/ChaosCloth/ChaosClothingSimulation.cpp b/Engine/Plugins/Experimental/ChaosCloth/Source/Chaos/Private/ChaosCloth/ChaosClothingSimulation.cpp
--- a/Engine/Plugins/Experimental/ChaosCloth/Source/Chaos/Private/ChaosCloth/ChaosClothingSimulation.cpp
+++ b/Engine/Plugins/Experimental/ChaosCloth/Source/Chaos/Private/ChaosCloth/ChaosClothingSimulation.cpp
@@ -811,7 +811,7 @@ void ClothingSimulation::AddCollisions(const FClothCollisionData& ClothCollision
 			else
 			{
 				// Tapered capsule
-				TArray<TUniquePtr<FImplicitObject>> Objects;
+				TArray<TUniquePtr<Chaos::FImplicitObject>> Objects;
 				Objects.Reserve(3);
 				Objects.Add(TUniquePtr<Chaos::FImplicitObject>(
 					new Chaos::TTaperedCylinder<float>(X0, X1, Radius0, Radius1)));
@@ -1459,7 +1459,7 @@ void ClothingSimulation::DebugDrawCollision(USkeletalMeshComponent* OwnerCompone
 
 	for (uint32 Index = 0; Index < CollisionParticles.Size(); ++Index)
 	{
-		if (const FImplicitObject* const Object = CollisionParticles.DynamicGeometry(Index).Get())
+		if (const Chaos::FImplicitObject* const Object = CollisionParticles.DynamicGeometry(Index).Get())
 		{
 			const uint32 BoneIndex = BoneIndices[Index];
 			const FLinearColor Color = (BoneIndex != INDEX_NONE) ? MappedColor : UnmappedColor;
@@ -1482,9 +1482,9 @@ void ClothingSimulation::DebugDrawCollision(USkeletalMeshComponent* OwnerCompone
 				break;
 
 			case Chaos::ImplicitObjectType::Union:  // Union only used as collision tappered capsules
-				for (const TUniquePtr<FImplicitObject>& SubObjectPtr : Object->GetObjectChecked<Chaos::TImplicitObjectUnion<float, 3>>().GetObjects())
+				for (const TUniquePtr<Chaos::FImplicitObject>& SubObjectPtr : Object->GetObjectChecked<Chaos::TImplicitObjectUnion<float, 3>>().GetObjects())
 				{
-					if (const FImplicitObject* const SubObject = SubObjectPtr.Get())
+					if (const Chaos::FImplicitObject* const SubObject = SubObjectPtr.Get())
 					{
 						switch (SubObject->GetType())
 						{
