- Animation Drive stiffness can now be controlled through Blueprints
- Chaos cloth interactors implemented

diff --git a/Engine/Plugins/Experimental/ChaosCloth/Source/Chaos/Private/ChaosCloth/ChaosClothingSimulation.cpp b/Engine/Plugins/Experimental/ChaosCloth/Source/Chaos/Private/ChaosCloth/ChaosClothingSimulation.cpp
--- a/Engine/Plugins/Experimental/ChaosCloth/Source/Chaos/Private/ChaosCloth/ChaosClothingSimulation.cpp
+++ b/Engine/Plugins/Experimental/ChaosCloth/Source/Chaos/Private/ChaosCloth/ChaosClothingSimulation.cpp
@@ -120,6 +120,7 @@ void ClothingSimulation::Initialize()
 void ClothingSimulation::Shutdown()
 {
 	Assets.Reset();
+	AnimDriveSpringStiffness.Reset();
 	ExtractedCollisions.Reset();
 	ExternalCollisions.Reset();
 	OldCollisionTransforms.Reset();
@@ -150,22 +151,26 @@ void ClothingSimulation::CreateActor(USkeletalMeshComponent* InOwnerComponent, U
 
 	//Evolution->SetCCD(ChaosClothSimConfig->bUseContinuousCollisionDetection);
 	//Evolution->SetCCD(true); // ryan!!!
-	
+
 	UClothingAssetCommon* Asset = Cast<UClothingAssetCommon>(InAsset);
 	const UChaosClothConfig* const ChaosClothSimConfig = Cast<UChaosClothConfig>(Asset->ChaosClothSimConfig);
 	check(ChaosClothSimConfig);
 
 	ClothingSimulationContext Context;
-    FillContext(InOwnerComponent, 0, &Context);
-	
+	FillContext(InOwnerComponent, 0, &Context);
+
 	if (Assets.Num() <= InSimDataIndex)
-        Assets.SetNumZeroed(InSimDataIndex + 1);
+	{
+		Assets.SetNumZeroed(InSimDataIndex + 1);
+		AnimDriveSpringStiffness.SetNumZeroed(InSimDataIndex + 1);
+	}
 	Assets[InSimDataIndex] = Asset;
+	AnimDriveSpringStiffness[InSimDataIndex] = ChaosClothSimConfig->AnimDriveSpringStiffness;
 
-    check(Asset->GetNumLods() == 1);
+	check(Asset->GetNumLods() == 1);
 	UClothLODDataBase* AssetLodData = Asset->ClothLodData[0];
 	check(AssetLodData->PhysicalMeshData);
-    UClothPhysicalMeshDataBase* PhysMesh = AssetLodData->PhysicalMeshData;
+	UClothPhysicalMeshDataBase* PhysMesh = AssetLodData->PhysicalMeshData;
 
 	// SkinPhysicsMesh() strips scale from RootBoneTransform ("Ignore any user scale.
 	// It's already accounted for in our skinning matrices."), and returns all points
@@ -180,18 +185,18 @@ void ClothingSimulation::CreateActor(USkeletalMeshComponent* InOwnerComponent, U
 		RootBoneTransform,
 		Context.RefToLocals.GetData(),
 		Context.RefToLocals.Num(),
-		reinterpret_cast<TArray<FVector>&>(TempAnimationPositions), 
+		reinterpret_cast<TArray<FVector>&>(TempAnimationPositions),
 		reinterpret_cast<TArray<FVector>&>(TempAnimationNormals));
 
 	// Transform points & normals to world space
 	RootBoneTransform.SetScale3D(FVector(1.0f));
 	const FTransform RootBoneWorldTransform = RootBoneTransform * Context.ComponentToWorld;
-	ParallelFor(TempAnimationPositions.Num(), 
+	ParallelFor(TempAnimationPositions.Num(),
 		[&](int32 Index)
-		{
-			TempAnimationPositions[Index] = RootBoneWorldTransform.TransformPosition(TempAnimationPositions[Index]);
-			TempAnimationNormals[Index] = RootBoneWorldTransform.TransformVector(TempAnimationNormals[Index]);
-		});
+	{
+		TempAnimationPositions[Index] = RootBoneWorldTransform.TransformPosition(TempAnimationPositions[Index]);
+		TempAnimationNormals[Index] = RootBoneWorldTransform.TransformVector(TempAnimationNormals[Index]);
+	});
 
 	// Add particles
 	TPBDParticles<float, 3>& Particles = Evolution->Particles();
@@ -208,7 +213,7 @@ void ClothingSimulation::CreateActor(USkeletalMeshComponent* InOwnerComponent, U
 	AnimationNormals.SetNum(Particles.Size());
 
 	if (IndexToRangeMap.Num() <= InSimDataIndex)
-		 IndexToRangeMap.SetNum(InSimDataIndex + 1);
+		IndexToRangeMap.SetNum(InSimDataIndex + 1);
 	IndexToRangeMap[InSimDataIndex] = Chaos::TVector<uint32, 2>(Offset, Particles.Size());
 
 	for (uint32 i = Offset; i < Particles.Size(); ++i)
@@ -230,9 +235,9 @@ void ClothingSimulation::CreateActor(USkeletalMeshComponent* InOwnerComponent, U
 	{
 		const int32 Index = 3 * i;
 		InputSurfaceElements.Add(
-			{static_cast<int32>(Offset + PhysMesh->Indices[Index]),
+			{ static_cast<int32>(Offset + PhysMesh->Indices[Index]),
 			 static_cast<int32>(Offset + PhysMesh->Indices[Index + 1]),
-			 static_cast<int32>(Offset + PhysMesh->Indices[Index + 2])});
+			 static_cast<int32>(Offset + PhysMesh->Indices[Index + 2]) });
 	}
 	check(InputSurfaceElements.Num() == NumTriangles);
 	if (Meshes.Num() <= InSimDataIndex)
@@ -242,9 +247,9 @@ void ClothingSimulation::CreateActor(USkeletalMeshComponent* InOwnerComponent, U
 		PointNormals.SetNum(InSimDataIndex + 1);
 	}
 	TUniquePtr<Chaos::TTriangleMesh<float>>& Mesh = Meshes[InSimDataIndex];
-    Mesh.Reset(new Chaos::TTriangleMesh<float>(MoveTemp(InputSurfaceElements)));
+	Mesh.Reset(new Chaos::TTriangleMesh<float>(MoveTemp(InputSurfaceElements)));
 	check(Mesh->GetNumElements() == NumTriangles);
-    const auto& SurfaceElements = Mesh->GetSurfaceElements();
+	const auto& SurfaceElements = Mesh->GetSurfaceElements();
 	Mesh->GetPointToTriangleMap(); // Builds map for later use by GetPointNormals().
 
 	// Assign per particle mass proportional to connected area.
@@ -289,104 +294,104 @@ void ClothingSimulation::CreateActor(USkeletalMeshComponent* InOwnerComponent, U
 	for (uint32 i = Offset; i < Particles.Size(); i++)
 	{
 		Particles.M(i) = FMath::Max(Particles.M(i), ChaosClothSimConfig->MinPerParticleMass);
-		Particles.InvM(i) = PhysMesh->IsFixed(i-Offset) ? 0.0 : 1.0/Particles.M(i);
+		Particles.InvM(i) = PhysMesh->IsFixed(i - Offset) ? 0.0 : 1.0 / Particles.M(i);
 	}
 
-    // Add Model
-    if (ChaosClothSimConfig->ShapeTargetStiffness)
-    {
-        check(ChaosClothSimConfig->ShapeTargetStiffness > 0.f && ChaosClothSimConfig->ShapeTargetStiffness <= 1.f);
+	// Add Model
+	if (ChaosClothSimConfig->ShapeTargetStiffness)
+	{
+		check(ChaosClothSimConfig->ShapeTargetStiffness > 0.f && ChaosClothSimConfig->ShapeTargetStiffness <= 1.f);
 		Evolution->AddPBDConstraintFunction([ShapeConstraints = Chaos::TPerParticlePBDShapeConstraints<float, 3>(Evolution->Particles(), AnimationPositions, ChaosClothSimConfig->ShapeTargetStiffness)](TPBDParticles<float, 3>& InParticles, const float Dt) {
 			ShapeConstraints.Apply(InParticles, Dt);
 		});
-    }
-    if (ChaosClothSimConfig->EdgeStiffness)
-    {
-        check(ChaosClothSimConfig->EdgeStiffness > 0.f && ChaosClothSimConfig->EdgeStiffness <= 1.f);
-        Evolution->AddPBDConstraintFunction([SpringConstraints = Chaos::TPBDSpringConstraints<float, 3>(Evolution->Particles(), SurfaceElements, ChaosClothSimConfig->EdgeStiffness)](TPBDParticles<float, 3>& InParticles, const float Dt) {
-            SpringConstraints.Apply(InParticles, Dt);
-        });
-    }
-    if (ChaosClothSimConfig->BendingStiffness)
-    {
-        check(ChaosClothSimConfig->BendingStiffness > 0.f && ChaosClothSimConfig->BendingStiffness <= 1.f);
-        if (ChaosClothSimConfig->bUseBendingElements)
-        {
-            TArray<Chaos::TVector<int32, 4>> BendingConstraints = Mesh->GetUniqueAdjacentElements();
+	}
+	if (ChaosClothSimConfig->EdgeStiffness)
+	{
+		check(ChaosClothSimConfig->EdgeStiffness > 0.f && ChaosClothSimConfig->EdgeStiffness <= 1.f);
+		Evolution->AddPBDConstraintFunction([SpringConstraints = Chaos::TPBDSpringConstraints<float, 3>(Evolution->Particles(), SurfaceElements, ChaosClothSimConfig->EdgeStiffness)](TPBDParticles<float, 3>& InParticles, const float Dt) {
+			SpringConstraints.Apply(InParticles, Dt);
+		});
+	}
+	if (ChaosClothSimConfig->BendingStiffness)
+	{
+		check(ChaosClothSimConfig->BendingStiffness > 0.f && ChaosClothSimConfig->BendingStiffness <= 1.f);
+		if (ChaosClothSimConfig->bUseBendingElements)
+		{
+			TArray<Chaos::TVector<int32, 4>> BendingConstraints = Mesh->GetUniqueAdjacentElements();
 			Evolution->AddPBDConstraintFunction([BendConstraints = Chaos::TPBDBendingConstraints<float>(Evolution->Particles(), MoveTemp(BendingConstraints))](TPBDParticles<float, 3>& InParticles, const float Dt) {
 				BendConstraints.Apply(InParticles, Dt);
 			});
-        }
-        else
-        {
-            TArray<Chaos::TVector<int32, 2>> BendingConstraints = Mesh->GetUniqueAdjacentPoints();
-            Evolution->AddPBDConstraintFunction([SpringConstraints = Chaos::TPBDSpringConstraints<float, 3>(Evolution->Particles(), MoveTemp(BendingConstraints), ChaosClothSimConfig->BendingStiffness)](TPBDParticles<float, 3>& InParticles, const float Dt) {
-                SpringConstraints.Apply(InParticles, Dt);
-            });
-        }
-    }
-    if (ChaosClothSimConfig->AreaStiffness)
-    {
-        TArray<Chaos::TVector<int32, 3>> SurfaceConstraints = SurfaceElements;
+		}
+		else
+		{
+			TArray<Chaos::TVector<int32, 2>> BendingConstraints = Mesh->GetUniqueAdjacentPoints();
+			Evolution->AddPBDConstraintFunction([SpringConstraints = Chaos::TPBDSpringConstraints<float, 3>(Evolution->Particles(), MoveTemp(BendingConstraints), ChaosClothSimConfig->BendingStiffness)](TPBDParticles<float, 3>& InParticles, const float Dt) {
+				SpringConstraints.Apply(InParticles, Dt);
+			});
+		}
+	}
+	if (ChaosClothSimConfig->AreaStiffness)
+	{
+		TArray<Chaos::TVector<int32, 3>> SurfaceConstraints = SurfaceElements;
 		Evolution->AddPBDConstraintFunction([SurfConstraints = Chaos::TPBDAxialSpringConstraints<float, 3>(Evolution->Particles(), MoveTemp(SurfaceConstraints), ChaosClothSimConfig->AreaStiffness)](TPBDParticles<float, 3>& InParticles, const float Dt) {
 			SurfConstraints.Apply(InParticles, Dt);
 		});
-    }
-    if (ChaosClothSimConfig->VolumeStiffness)
-    {
-        check(ChaosClothSimConfig->VolumeStiffness > 0.f && ChaosClothSimConfig->VolumeStiffness <= 1.f);
-        if (ChaosClothSimConfig->bUseTetrahedralConstraints)
-        {
-            // TODO(mlentine): Need to tetrahedralize surface to support this
-            check(false);
-        }
-        else if (ChaosClothSimConfig->bUseThinShellVolumeConstraints)
-        {
-            TArray<Chaos::TVector<int32, 2>> BendingConstraints = Mesh->GetUniqueAdjacentPoints();
-            TArray<Chaos::TVector<int32, 2>> DoubleBendingConstraints;
-            {
-                TMap<int32, TArray<int32>> BendingHash;
-                for (int32 i = 0; i < BendingConstraints.Num(); ++i)
-                {
-                    BendingHash.FindOrAdd(BendingConstraints[i][0]).Add(BendingConstraints[i][1]);
-                    BendingHash.FindOrAdd(BendingConstraints[i][1]).Add(BendingConstraints[i][0]);
-                }
-                TSet<Chaos::TVector<int32, 2>> Visited;
-                for (auto Elem : BendingHash)
-                {
-                    for (int32 i = 0; i < Elem.Value.Num(); ++i)
-                    {
-                        for (int32 j = i + 1; j < Elem.Value.Num(); ++j)
-                        {
-                            if (Elem.Value[i] == Elem.Value[j])
-                                continue;
-                            auto NewElem = Chaos::TVector<int32, 2>(Elem.Value[i], Elem.Value[j]);
-                            if (!Visited.Contains(NewElem))
-                            {
-                                DoubleBendingConstraints.Add(NewElem);
-                                Visited.Add(NewElem);
-                                Visited.Add(Chaos::TVector<int32, 2>(Elem.Value[j], Elem.Value[i]));
-                            }
-                        }
-                    }
-                }
-            }
-            Evolution->AddPBDConstraintFunction([SpringConstraints = Chaos::TPBDSpringConstraints<float, 3>(Evolution->Particles(), MoveTemp(DoubleBendingConstraints), ChaosClothSimConfig->VolumeStiffness)](TPBDParticles<float, 3>& InParticles, const float Dt) {
-                SpringConstraints.Apply(InParticles, Dt);
-            });
-        }
-        else
-        {
-            TArray<Chaos::TVector<int32, 3>> SurfaceConstraints = SurfaceElements;
-			Chaos::TPBDVolumeConstraint<float> PBDVolumeConstraint (Evolution->Particles(), MoveTemp(SurfaceConstraints));
+	}
+	if (ChaosClothSimConfig->VolumeStiffness)
+	{
+		check(ChaosClothSimConfig->VolumeStiffness > 0.f && ChaosClothSimConfig->VolumeStiffness <= 1.f);
+		if (ChaosClothSimConfig->bUseTetrahedralConstraints)
+		{
+			// TODO(mlentine): Need to tetrahedralize surface to support this
+			check(false);
+		}
+		else if (ChaosClothSimConfig->bUseThinShellVolumeConstraints)
+		{
+			TArray<Chaos::TVector<int32, 2>> BendingConstraints = Mesh->GetUniqueAdjacentPoints();
+			TArray<Chaos::TVector<int32, 2>> DoubleBendingConstraints;
+			{
+				TMap<int32, TArray<int32>> BendingHash;
+				for (int32 i = 0; i < BendingConstraints.Num(); ++i)
+				{
+					BendingHash.FindOrAdd(BendingConstraints[i][0]).Add(BendingConstraints[i][1]);
+					BendingHash.FindOrAdd(BendingConstraints[i][1]).Add(BendingConstraints[i][0]);
+				}
+				TSet<Chaos::TVector<int32, 2>> Visited;
+				for (auto Elem : BendingHash)
+				{
+					for (int32 i = 0; i < Elem.Value.Num(); ++i)
+					{
+						for (int32 j = i + 1; j < Elem.Value.Num(); ++j)
+						{
+							if (Elem.Value[i] == Elem.Value[j])
+								continue;
+							auto NewElem = Chaos::TVector<int32, 2>(Elem.Value[i], Elem.Value[j]);
+							if (!Visited.Contains(NewElem))
+							{
+								DoubleBendingConstraints.Add(NewElem);
+								Visited.Add(NewElem);
+								Visited.Add(Chaos::TVector<int32, 2>(Elem.Value[j], Elem.Value[i]));
+							}
+						}
+					}
+				}
+			}
+			Evolution->AddPBDConstraintFunction([SpringConstraints = Chaos::TPBDSpringConstraints<float, 3>(Evolution->Particles(), MoveTemp(DoubleBendingConstraints), ChaosClothSimConfig->VolumeStiffness)](TPBDParticles<float, 3>& InParticles, const float Dt) {
+				SpringConstraints.Apply(InParticles, Dt);
+			});
+		}
+		else
+		{
+			TArray<Chaos::TVector<int32, 3>> SurfaceConstraints = SurfaceElements;
+			Chaos::TPBDVolumeConstraint<float> PBDVolumeConstraint(Evolution->Particles(), MoveTemp(SurfaceConstraints));
 			Evolution->AddPBDConstraintFunction([PBDVolumeConstraint = MoveTemp(PBDVolumeConstraint)](TPBDParticles<float, 3>& InParticles, const float Dt)
 			{
 				PBDVolumeConstraint.Apply(InParticles, Dt);
-			});            
-        }
-    }
-    if (ChaosClothSimConfig->StrainLimitingStiffness)
-    {
+			});
+		}
+	}
+	if (ChaosClothSimConfig->StrainLimitingStiffness)
+	{
 		check(Mesh->GetNumElements() > 0);
 		Chaos::TPerParticlePBDLongRangeConstraints<float, 3> PerParticlePBDLongRangeConstraints(
 			Evolution->Particles(),
@@ -401,7 +406,7 @@ void ClothingSimulation::CreateActor(USkeletalMeshComponent* InOwnerComponent, U
 	}
 
 	// Maximum Distance Constraints
-	const UEnum* const MeshTargets = PhysMesh->GetFloatArrayTargets();	
+	const UEnum* const MeshTargets = PhysMesh->GetFloatArrayTargets();
 	const uint32 PhysMeshMaxDistanceIndex = MeshTargets->GetValueByName(TEXT("MaxDistance"));
 	if (PhysMesh->GetFloatArray(PhysMeshMaxDistanceIndex)->Num() > 0)
 	{
@@ -421,23 +426,24 @@ void ClothingSimulation::CreateActor(USkeletalMeshComponent* InOwnerComponent, U
 		check(Mesh->GetNumElements() > 0);
 		check(PhysMesh->GetFloatArray(PhysMeshBackstopRadiusIndex)->Num() == PhysMesh->GetFloatArray(PhysMeshBackstopDistanceIndex)->Num());
 
-		Chaos::PBDSphericalConstraint<float, 3> SphericalContraint(Offset, PhysMesh->GetFloatArray(PhysMeshBackstopRadiusIndex)->Num(), false, &AnimationPositions, 
+		Chaos::PBDSphericalConstraint<float, 3> SphericalContraint(Offset, PhysMesh->GetFloatArray(PhysMeshBackstopRadiusIndex)->Num(), false, &AnimationPositions,
 			PhysMesh->GetFloatArray(PhysMeshBackstopRadiusIndex), PhysMesh->GetFloatArray(PhysMeshBackstopDistanceIndex), &AnimationNormals);
 		Evolution->AddPBDConstraintFunction([SphericalContraint = MoveTemp(SphericalContraint)](TPBDParticles<float, 3>& InParticles, const float Dt)
 		{
 			SphericalContraint.Apply(InParticles, Dt);
-		});		
-	}	
+		});
+	}
 
 	// Animation Drive Constraints
 	const uint32 PhysMeshAnimDriveIndex = MeshTargets->GetValueByName(TEXT("AnimDriveMultiplier"));
 	if (PhysMesh->GetFloatArray(PhysMeshAnimDriveIndex)->Num() > 0)
 	{
 		check(Mesh->GetNumElements() > 0);
-		TPBDAnimDriveConstraint<float, 3> PBDAnimDriveConstraint(Offset, &AnimationPositions, PhysMesh->GetFloatArray(PhysMeshAnimDriveIndex), ChaosClothSimConfig->AnimDriveSpringStiffness);
+		TPBDAnimDriveConstraint<float, 3> PBDAnimDriveConstraint(Offset, &AnimationPositions, PhysMesh->GetFloatArray(PhysMeshAnimDriveIndex), AnimDriveSpringStiffness[InSimDataIndex]);
 		Evolution->AddPBDConstraintFunction(
-			[PBDAnimDriveConstraint = MoveTemp(PBDAnimDriveConstraint)](TPBDParticles<float, 3>& InParticles, const float Dt)
+			[PBDAnimDriveConstraint = MoveTemp(PBDAnimDriveConstraint), &Stiffness = AnimDriveSpringStiffness[InSimDataIndex]](TPBDParticles<float, 3>& InParticles, const float Dt) mutable
 		{
+			PBDAnimDriveConstraint.SetSpringStiffness(Stiffness);
 			PBDAnimDriveConstraint.Apply(InParticles, Dt);
 		});
 	}
@@ -483,11 +489,11 @@ void ClothingSimulation::PostActorCreationInitialize()
 	{
 		if (Asset)
 		{
-			// If we don't have a sim config, create one
+			// If we don't have a shared sim config, create one. This is only possible if none of the cloth Assets had a configuration during Actor creation
 			if (ClothSharedSimConfig == nullptr)
 			{
-					// None of the cloth assets had a clothSharedSimConfig, so we will create it
-					ClothSharedSimConfig = NewObject<UChaosClothSharedSimConfig>(Asset, UChaosClothSharedSimConfig::StaticClass()->GetFName());
+				// None of the cloth assets had a clothSharedSimConfig, so we will create it
+				ClothSharedSimConfig = NewObject<UChaosClothSharedSimConfig>(Asset, UChaosClothSharedSimConfig::StaticClass()->GetFName());
 			}
 			Asset->ClothSharedSimConfig = ClothSharedSimConfig;
 		}
@@ -1182,6 +1188,14 @@ void ClothingSimulation::GetCollisions(FClothCollisionData& OutCollisions, bool
 	UE_LOG(LogChaosCloth, VeryVerbose, TEXT("Returned collisions: %d spheres, %d capsules, %d convexes, %d boxes."), OutCollisions.Spheres.Num() - 2 * OutCollisions.SphereConnections.Num(), OutCollisions.SphereConnections.Num(), OutCollisions.Convexes.Num(), OutCollisions.Boxes.Num());
 }
 
+void ClothingSimulation::SetAnimDriveSpringStiffness(float InStiffness)
+{
+	for (float& stiffness : AnimDriveSpringStiffness)
+	{
+		stiffness = InStiffness;
+	}
+}
+
 #if WITH_EDITOR
 void ClothingSimulation::AddReferencedObjects(FReferenceCollector& Collector)
 {
diff --git a/Engine/Plugins/Experimental/ChaosCloth/Source/Chaos/Private/ChaosCloth/ChaosClothingSimulationFactory.cpp b/Engine/Plugins/Experimental/ChaosCloth/Source/Chaos/Private/ChaosCloth/ChaosClothingSimulationFactory.cpp
--- a/Engine/Plugins/Experimental/ChaosCloth/Source/Chaos/Private/ChaosCloth/ChaosClothingSimulationFactory.cpp
+++ b/Engine/Plugins/Experimental/ChaosCloth/Source/Chaos/Private/ChaosCloth/ChaosClothingSimulationFactory.cpp
@@ -3,6 +3,7 @@
 #include "ChaosCloth/ChaosClothingSimulationFactory.h"
 #include "ChaosCloth/ChaosClothConfig.h"
 #include "ChaosCloth/ChaosClothingSimulation.h"
+#include "ChaosCloth/ChaosClothingSimulationInteractor.h"
 
 IClothingSimulation* UChaosClothingSimulationFactory::CreateSimulation()
 {
@@ -22,12 +23,12 @@ bool UChaosClothingSimulationFactory::SupportsAsset(UClothingAssetBase* InAsset)
 
 bool UChaosClothingSimulationFactory::SupportsRuntimeInteraction()
 {
-    return false;
+    return true;
 }
 
 UClothingSimulationInteractor* UChaosClothingSimulationFactory::CreateInteractor()
 {
-	return nullptr;
+	return NewObject<UChaosClothingSimulationInteractor>(GetTransientPackage());
 }
 
 UClothConfigBase* UChaosClothingSimulationFactory::CreateDefaultClothConfig(const FObjectInitializer& ObjectInitializer, UObject* Outer)
diff --git a/Engine/Plugins/Experimental/ChaosCloth/Source/Chaos/Private/ChaosCloth/ChaosClothingSimulationInteractor.cpp b/Engine/Plugins/Experimental/ChaosCloth/Source/Chaos/Private/ChaosCloth/ChaosClothingSimulationInteractor.cpp
new file mode 100644
--- /dev/null
+++ b/Engine/Plugins/Experimental/ChaosCloth/Source/Chaos/Private/ChaosCloth/ChaosClothingSimulationInteractor.cpp
@@ -0,0 +1,50 @@
+// Copyright 1998-2019 Epic Games, Inc. All Rights Reserved.
+
+#include "ChaosCloth/ChaosClothingSimulationInteractor.h"
+
+void UChaosClothingSimulationInteractor::PhysicsAssetUpdated()
+{
+	//TODO
+}
+
+void UChaosClothingSimulationInteractor::ClothConfigUpdated()
+{
+	//TODO
+}
+
+void UChaosClothingSimulationInteractor::Sync(IClothingSimulation* InSimulation, IClothingSimulationContext* InContext)
+{
+	check(InSimulation);
+	check(InContext);
+
+	Chaos::ClothingSimulation* ChaosSim = static_cast<Chaos::ClothingSimulation*>(InSimulation);
+	Chaos::ClothingSimulationContext* ChaosContext = static_cast<Chaos::ClothingSimulationContext*>(InContext);
+
+	for(ChaosClothInteractorCommand& Command : Commands)
+	{
+		Command.Execute(ChaosSim, ChaosContext);
+	}
+	Commands.Reset();
+	bDirty = false;
+}
+
+void UChaosClothingSimulationInteractor::SetAnimDriveSpringStiffness(float InStiffness)
+{
+	Commands.Add(ChaosClothInteractorCommand::CreateLambda([InStiffness](Chaos::ClothingSimulation* InSimulation, Chaos::ClothingSimulationContext* InContext)
+	{
+		InSimulation->SetAnimDriveSpringStiffness(InStiffness);
+	}));
+
+	MarkDirty();
+}
+
+void UChaosClothingSimulationInteractor::EnableGravityOverride(const FVector& InVector)
+{
+	//TODO
+}
+
+void UChaosClothingSimulationInteractor::DisableGravityOverride()
+{
+	//TODO
+}
+
diff --git a/Engine/Plugins/Experimental/ChaosCloth/Source/Chaos/Public/ChaosCloth/ChaosClothingSimulation.h b/Engine/Plugins/Experimental/ChaosCloth/Source/Chaos/Public/ChaosCloth/ChaosClothingSimulation.h
--- a/Engine/Plugins/Experimental/ChaosCloth/Source/Chaos/Public/ChaosCloth/ChaosClothingSimulation.h
+++ b/Engine/Plugins/Experimental/ChaosCloth/Source/Chaos/Public/ChaosCloth/ChaosClothingSimulation.h
@@ -34,6 +34,9 @@ namespace Chaos
 		ClothingSimulation();
 		virtual ~ClothingSimulation();
 
+		// Set the animation drive stiffness for all actors
+		void SetAnimDriveSpringStiffness(float InStiffness);
+
 #if WITH_EDITOR
 		// FGCObject interface
 		void AddReferencedObjects(FReferenceCollector& Collector) override;
@@ -94,6 +97,11 @@ namespace Chaos
 		TArray<UClothingAssetCommon*> Assets;
 		UChaosClothSharedSimConfig* ClothSharedSimConfig;
 
+		// Cloth Interaction Parameters
+		// These simulation parameters can be changed through blueprints
+		// They will only be updated when the simulation is not running (so are safe to use on any cloth thread)
+		TArray<float> AnimDriveSpringStiffness; // One for every Asset
+
 		// Collision Data
 		FClothCollisionData ExtractedCollisions;  // Collisions extracted from the referenced physics asset
 		FClothCollisionData ExternalCollisions;  // External collisions
diff --git a/Engine/Plugins/Experimental/ChaosCloth/Source/Chaos/Public/ChaosCloth/ChaosClothingSimulationInteractor.h b/Engine/Plugins/Experimental/ChaosCloth/Source/Chaos/Public/ChaosCloth/ChaosClothingSimulationInteractor.h
new file mode 100644
--- /dev/null
+++ b/Engine/Plugins/Experimental/ChaosCloth/Source/Chaos/Public/ChaosCloth/ChaosClothingSimulationInteractor.h
@@ -0,0 +1,40 @@
+// Copyright 1998-2019 Epic Games, Inc. All Rights Reserved.
+
+#pragma once
+
+#include "ClothingSimulationInteractor.h"
+#include "ChaosClothingSimulation.h"
+#include "ChaosClothingSimulationInteractor.generated.h"
+
+// Command signature for handling synced command buffer
+DECLARE_DELEGATE_TwoParams(ChaosClothInteractorCommand, Chaos::ClothingSimulation*, Chaos::ClothingSimulationContext*)
+
+UCLASS(BlueprintType)
+class CHAOSCLOTH_API UChaosClothingSimulationInteractor : public UClothingSimulationInteractor
+{
+	GENERATED_BODY()
+public:
+
+	// UClothingSimulationInteractor Interface
+	virtual void PhysicsAssetUpdated() override;
+	virtual void ClothConfigUpdated() override;
+	virtual void Sync(IClothingSimulation* InSimulation, IClothingSimulationContext* InContext) override;
+	//////////////////////////////////////////////////////////////////////////
+
+	// Set the stiffness of the spring force for the animation drive
+	UFUNCTION(BlueprintCallable, Category=ClothingSimulation)
+	void SetAnimDriveSpringStiffness(float InStiffness);
+
+	// Set a new gravity override and enable the override
+	UFUNCTION(BlueprintCallable, Category=ClothingSimulation)
+	void EnableGravityOverride(const FVector& InVector);
+
+	// Disable any currently set gravity override
+	UFUNCTION(BlueprintCallable, Category = ClothingSimulation)
+	void DisableGravityOverride();
+
+private:
+
+	// Command queue processed when we hit a sync
+	TArray<ChaosClothInteractorCommand> Commands;
+};
diff --git a/Engine/Source/Runtime/Experimental/Chaos/Public/Chaos/PBDAnimDriveConstraint.h b/Engine/Source/Runtime/Experimental/Chaos/Public/Chaos/PBDAnimDriveConstraint.h
--- a/Engine/Source/Runtime/Experimental/Chaos/Public/Chaos/PBDAnimDriveConstraint.h
+++ b/Engine/Source/Runtime/Experimental/Chaos/Public/Chaos/PBDAnimDriveConstraint.h
@@ -13,7 +13,7 @@ class TPBDAnimDriveConstraint : public TParticleRule<T, d>
 		const int InParticleIndexOffset
 		, const TArray<TVector<T, d>>* const InSpringNeutralPositions
 		, const TArray<T>* const InSpringStiffnessMultiplier
-		, float InSpringStiffness
+		, const T InSpringStiffness
 	)
 		: ParticleIndexOffset(InParticleIndexOffset)
 		, SpringNeutralPositions(InSpringNeutralPositions)
@@ -48,11 +48,16 @@ class TPBDAnimDriveConstraint : public TParticleRule<T, d>
 		InParticles.P(Index + ParticleIndexOffset) = RelaxationFactor * NeutralPosition + (1 - RelaxationFactor) * ParticlePosition;
 	}
 
+	inline void SetSpringStiffness(T InSpringStiffness)
+	{
+		SpringStiffness = InSpringStiffness;
+	}
+
 private:
 	const int32 ParticleIndexOffset;
-	const TArray<TVector<T, d>>* const SpringNeutralPositions; // Size: Same as full particle array
-	const TArray<T>* const SpringStiffnessMultiplier; // Size: Number of Animation drive constraints to solve
+	const TArray<TVector<T, d>>* SpringNeutralPositions; // Size: Same as full particle array
+	const TArray<T>* SpringStiffnessMultiplier; // Size: Number of Animation drive constraints to solve
 
-	const float SpringStiffness;
+	T SpringStiffness;
 };
 }
