Chaos Cloth Bug Fix: Now support skeletal meshes with all disabled cloth instances

diff --git a/Engine/Plugins/Experimental/ChaosCloth/Source/Chaos/Private/ChaosCloth/ChaosClothingSimulation.cpp b/Engine/Plugins/Experimental/ChaosCloth/Source/Chaos/Private/ChaosCloth/ChaosClothingSimulation.cpp
--- a/Engine/Plugins/Experimental/ChaosCloth/Source/Chaos/Private/ChaosCloth/ChaosClothingSimulation.cpp
+++ b/Engine/Plugins/Experimental/ChaosCloth/Source/Chaos/Private/ChaosCloth/ChaosClothingSimulation.cpp
@@ -494,19 +494,21 @@ void ClothingSimulation::PostActorCreationInitialize()
 			{
 				// None of the cloth assets had a clothSharedSimConfig, so we will create it
 				ClothSharedSimConfig = NewObject<UChaosClothSharedSimConfig>(Asset, UChaosClothSharedSimConfig::StaticClass()->GetFName());
+				check(ClothSharedSimConfig);
 			}
 			Asset->ClothSharedSimConfig = ClothSharedSimConfig;
 		}
 	}
 
-	check(ClothSharedSimConfig);
-
-	// Now set all the common parameters on the simulation
-	Evolution->SetIterations(ClothSharedSimConfig->IterationCount);
-	Evolution->SetSelfCollisionThickness(ClothSharedSimConfig->SelfCollisionThickness);
-	Evolution->SetCollisionThickness(ClothSharedSimConfig->CollisionThickness);
-	Evolution->SetDamping(ClothSharedSimConfig->Damping);
-	Evolution->GetGravityForces().SetAcceleration(Chaos::TVector<float, 3>(ClothSharedSimConfig->Gravity));
+	if (ClothSharedSimConfig) // ClothSharedSimConfig will be a null pointer if all cloth instances are disabled in which case we will use default Evolution parameters
+	{
+		// Now set all the common parameters on the simulation
+		Evolution->SetIterations(ClothSharedSimConfig->IterationCount);
+		Evolution->SetSelfCollisionThickness(ClothSharedSimConfig->SelfCollisionThickness);
+		Evolution->SetCollisionThickness(ClothSharedSimConfig->CollisionThickness);
+		Evolution->SetDamping(ClothSharedSimConfig->Damping);
+		Evolution->GetGravityForces().SetAcceleration(Chaos::TVector<float, 3>(ClothSharedSimConfig->Gravity));
+	}	
 }
 
 void ClothingSimulation::UpdateCollisionTransforms(const ClothingSimulationContext& Context)
