Chaos cloth minor feature: Override gravity using clothinteractors in Blueprints

diff --git a/Engine/Plugins/Experimental/ChaosCloth/Source/Chaos/Private/ChaosCloth/ChaosClothingSimulation.cpp b/Engine/Plugins/Experimental/ChaosCloth/Source/Chaos/Private/ChaosCloth/ChaosClothingSimulation.cpp
--- a/Engine/Plugins/Experimental/ChaosCloth/Source/Chaos/Private/ChaosCloth/ChaosClothingSimulation.cpp
+++ b/Engine/Plugins/Experimental/ChaosCloth/Source/Chaos/Private/ChaosCloth/ChaosClothingSimulation.cpp
@@ -1198,6 +1198,16 @@ void ClothingSimulation::SetAnimDriveSpringStiffness(float InStiffness)
 	}
 }
 
+void ClothingSimulation::SetGravityOverride(const TVector<float, 3>& InGravityOverride) 
+{
+	Evolution->GetGravityForces().SetAcceleration(InGravityOverride);
+}
+
+void ClothingSimulation::DisableGravityOverride()
+{
+	Evolution->GetGravityForces().SetAcceleration(ClothSharedSimConfig->Gravity);
+}
+
 #if WITH_EDITOR
 void ClothingSimulation::AddReferencedObjects(FReferenceCollector& Collector)
 {
diff --git a/Engine/Plugins/Experimental/ChaosCloth/Source/Chaos/Private/ChaosCloth/ChaosClothingSimulationInteractor.cpp b/Engine/Plugins/Experimental/ChaosCloth/Source/Chaos/Private/ChaosCloth/ChaosClothingSimulationInteractor.cpp
--- a/Engine/Plugins/Experimental/ChaosCloth/Source/Chaos/Private/ChaosCloth/ChaosClothingSimulationInteractor.cpp
+++ b/Engine/Plugins/Experimental/ChaosCloth/Source/Chaos/Private/ChaosCloth/ChaosClothingSimulationInteractor.cpp
@@ -40,11 +40,21 @@ void UChaosClothingSimulationInteractor::SetAnimDriveSpringStiffness(float InSti
 
 void UChaosClothingSimulationInteractor::EnableGravityOverride(const FVector& InVector)
 {
-	//TODO
+	Commands.Add(ChaosClothInteractorCommand::CreateLambda([InVector](Chaos::ClothingSimulation* InSimulation, Chaos::ClothingSimulationContext* InContext)
+	{
+		InSimulation->SetGravityOverride(InVector);
+	}));
+
+	MarkDirty();
 }
 
 void UChaosClothingSimulationInteractor::DisableGravityOverride()
 {
-	//TODO
+	Commands.Add(ChaosClothInteractorCommand::CreateLambda([](Chaos::ClothingSimulation* InSimulation, Chaos::ClothingSimulationContext* InContext)
+	{
+		InSimulation->DisableGravityOverride();
+	}));
+
+	MarkDirty();
 }
 
diff --git a/Engine/Plugins/Experimental/ChaosCloth/Source/Chaos/Public/ChaosCloth/ChaosClothingSimulation.h b/Engine/Plugins/Experimental/ChaosCloth/Source/Chaos/Public/ChaosCloth/ChaosClothingSimulation.h
--- a/Engine/Plugins/Experimental/ChaosCloth/Source/Chaos/Public/ChaosCloth/ChaosClothingSimulation.h
+++ b/Engine/Plugins/Experimental/ChaosCloth/Source/Chaos/Public/ChaosCloth/ChaosClothingSimulation.h
@@ -36,6 +36,8 @@ namespace Chaos
 
 		// Set the animation drive stiffness for all actors
 		void SetAnimDriveSpringStiffness(float InStiffness);
+		void SetGravityOverride(const TVector<float, 3>& InGravityOverride);
+		void DisableGravityOverride();
 
 #if WITH_EDITOR
 		// FGCObject interface
