Added chaos cloth debug rendering for:
- AnimDrives
- Backstops
- Maximum Distances

diff --git a/Engine/Plugins/Experimental/ChaosCloth/Source/Chaos/Private/ChaosCloth/ChaosClothingSimulation.cpp b/Engine/Plugins/Experimental/ChaosCloth/Source/Chaos/Private/ChaosCloth/ChaosClothingSimulation.cpp
--- a/Engine/Plugins/Experimental/ChaosCloth/Source/Chaos/Private/ChaosCloth/ChaosClothingSimulation.cpp
+++ b/Engine/Plugins/Experimental/ChaosCloth/Source/Chaos/Private/ChaosCloth/ChaosClothingSimulation.cpp
@@ -51,6 +51,7 @@ ClothingSimulation::ClothingSimulation()
 {
 #if WITH_EDITOR
 	DebugClothMaterial = LoadObject<UMaterial>(nullptr, TEXT("/Engine/EditorMaterials/Cloth/CameraLitDoubleSided.CameraLitDoubleSided"), nullptr, LOAD_None, nullptr);  // LOAD_EditorOnly
+	DebugClothMaterialVertex = LoadObject<UMaterial>(nullptr, TEXT("/Engine/EditorMaterials/WidgetVertexColorMaterial"), nullptr, LOAD_None, nullptr);  // LOAD_EditorOnly
 #endif  // #if WITH_EDITOR
 }
 
@@ -1097,7 +1098,7 @@ void ClothingSimulation::GetSimulationData(
 		const TUniquePtr<Chaos::TTriangleMesh<float>>& Mesh = Meshes[i];
 		if (!Mesh)
 			continue;
-		Mesh->GetFaceNormals(FaceNormals[i], Evolution->Particles().X(), false);
+		Mesh->GetFaceNormals(FaceNormals[i], Evolution->Particles().X(), false);  // No need to add a point index offset here since that is baked into the triangles
 		Mesh->GetPointNormals(PointNormals[i], FaceNormals[i], /*bReturnEmptyOnError =*/ false, /*bFillAtStartIndex =*/ false);
 
 		FClothSimulData& Data = OutData.FindOrAdd(i);
@@ -1543,12 +1544,95 @@ void ClothingSimulation::DebugDrawCollision(USkeletalMeshComponent* OwnerCompone
 
 void ClothingSimulation::DebugDrawBackstops(USkeletalMeshComponent* OwnerComponent, FPrimitiveDrawInterface* PDI) const
 {
-	// TODO: Add when GetCurrentSkinnedPositions is ever implemented
+	for (int32 i = 0; i < IndexToRangeMap.Num(); ++i)
+	{
+		const UClothingAssetCommon* const Asset = Assets[i];
+		if (Asset == nullptr)
+		{
+			continue;
+		}
+
+		// Get Backstop Distances
+		const UClothLODDataBase* const AssetLodData = Asset->ClothLodData[0];
+		check(AssetLodData);
+		check(AssetLodData->PhysicalMeshData);
+		const UClothPhysicalMeshDataBase* const PhysMesh = AssetLodData->PhysicalMeshData;
+		const UEnum* const MeshTargets = PhysMesh->GetFloatArrayTargets();
+		const uint32 PhysMeshBackstopIndex = MeshTargets->GetValueByName(TEXT("BackstopDistance"));
+		if (PhysMesh->GetFloatArray(PhysMeshBackstopIndex)->Num() == 0)
+		{
+			continue;
+		}
+
+		const uint32 PhysMeshBackstopRadiusIndex = MeshTargets->GetValueByName(TEXT("BackstopRadius"));
+		if (PhysMesh->GetFloatArray(PhysMeshBackstopRadiusIndex)->Num() == 0)
+		{
+			continue;
+		}
+
+		for (uint32 ParticleIndex = IndexToRangeMap[i][0]; ParticleIndex < IndexToRangeMap[i][1]; ++ParticleIndex)
+		{
+			const float Radius = (*PhysMesh->GetFloatArray(PhysMeshBackstopRadiusIndex))[ParticleIndex - IndexToRangeMap[i][0]];
+			const float Distance = (*PhysMesh->GetFloatArray(PhysMeshBackstopIndex))[ParticleIndex - IndexToRangeMap[i][0]];
+			PDI->DrawLine(AnimationPositions[ParticleIndex], AnimationPositions[ParticleIndex] - AnimationNormals[ParticleIndex] * (Distance - Radius), FColor::White, SDPG_World, 0.0f, 0.001f);
+			if (Radius > 0.0f)
+			{
+				const FVector& Normal = AnimationNormals[ParticleIndex];
+				const FVector& Position = AnimationPositions[ParticleIndex];
+				auto DrawBackstop = [Radius, Distance, &Normal, &Position, PDI](const FVector& Axis, const FColor& Color)
+				{
+					const float ArcLength = 5.0f; // Arch length in cm
+					const float ArcAngle = ArcLength * 360.0f / (Radius * 2.0f * PI);
+					
+					const float MaxCosAngle = 0.99f;
+					if (FMath::Abs(FVector::DotProduct(Normal, Axis)) < MaxCosAngle)
+					{
+						DrawArc(PDI, Position - Normal * Distance, Normal, FVector::CrossProduct(Axis, Normal).GetSafeNormal(), -ArcAngle / 2.0f, ArcAngle / 2.0f, Radius, 10, Color, SDPG_World);
+					}
+				};
+				DrawBackstop(FVector::ForwardVector, FColor::Blue);
+				DrawBackstop(FVector::UpVector, FColor::Blue);
+				DrawBackstop(FVector::RightVector, FColor::Blue);
+			}
+		}
+	}
 }
 
 void ClothingSimulation::DebugDrawMaxDistances(USkeletalMeshComponent* OwnerComponent, FPrimitiveDrawInterface* PDI) const
 {
-	// TODO: Add when GetCurrentSkinnedPositions is ever implemented
+	for (int32 i = 0; i < IndexToRangeMap.Num(); ++i)
+	{
+		const UClothingAssetCommon* const Asset = Assets[i];
+		if (Asset == nullptr)
+		{
+			continue;
+		}
+
+		// Get Maximum Distances
+		const UClothLODDataBase* const AssetLodData = Asset->ClothLodData[0];
+		check(AssetLodData);
+		check(AssetLodData->PhysicalMeshData);
+		UClothPhysicalMeshDataBase* PhysMesh = AssetLodData->PhysicalMeshData;
+		const UEnum* const MeshTargets = PhysMesh->GetFloatArrayTargets();
+		const uint32 PhysMeshMaxDistanceIndex = MeshTargets->GetValueByName(TEXT("MaxDistance"));
+		if (PhysMesh->GetFloatArray(PhysMeshMaxDistanceIndex)->Num() == 0)
+		{
+			continue;
+		}
+		
+		for (uint32 ParticleIndex = IndexToRangeMap[i][0]; ParticleIndex < IndexToRangeMap[i][1]; ++ParticleIndex)
+		{
+			const float Distance = (*PhysMesh->GetFloatArray(PhysMeshMaxDistanceIndex))[ParticleIndex - IndexToRangeMap[i][0]];
+			if (Distance == 0.0f)
+			{
+				DrawSphere(PDI, AnimationPositions[ParticleIndex], FRotator::ZeroRotator, FVector(0.5f, 0.5f, 0.5f), 10, 10, DebugClothMaterialVertex->GetRenderProxy(), SDPG_World, false);
+			}
+			else
+			{
+				PDI->DrawLine(AnimationPositions[ParticleIndex], AnimationPositions[ParticleIndex] + AnimationNormals[ParticleIndex] * Distance, FColor::White, SDPG_World, 0.0f, 0.001f);
+			}
+		}
+	}
 }
 
 void ClothingSimulation::DebugDrawSelfCollision(USkeletalMeshComponent* OwnerComponent, FPrimitiveDrawInterface* PDI) const
@@ -1580,6 +1664,32 @@ void ClothingSimulation::DebugDrawSelfCollision(USkeletalMeshComponent* OwnerCom
 
 void ClothingSimulation::DebugDrawAnimDrive(USkeletalMeshComponent* OwnerComponent, FPrimitiveDrawInterface* PDI) const
 {
-	// TODO: Add when GetCurrentSkinnedPositions is ever implemented
+	const TPBDParticles<float, 3>& Particles = Evolution->Particles();
+	for (int32 i = 0; i < IndexToRangeMap.Num(); ++i)
+	{
+		const UClothingAssetCommon* const Asset = Assets[i];
+		if (Asset == nullptr)
+		{
+			continue;
+		}
+
+		// Get Animdrive Multiplier
+		const UClothLODDataBase* const AssetLodData = Asset->ClothLodData[0];
+		check(AssetLodData);
+		check(AssetLodData->PhysicalMeshData);
+		const UClothPhysicalMeshDataBase* const PhysMesh = AssetLodData->PhysicalMeshData;
+		const UEnum* const MeshTargets = PhysMesh->GetFloatArrayTargets();
+		const uint32 PhysMeshAnimDriveIndex = MeshTargets->GetValueByName(TEXT("AnimDriveMultiplier"));
+		if (PhysMesh->GetFloatArray(PhysMeshAnimDriveIndex)->Num() == 0)
+		{
+			continue;
+		}
+
+		for (uint32 ParticleIndex = IndexToRangeMap[i][0]; ParticleIndex < IndexToRangeMap[i][1]; ++ParticleIndex)
+		{
+			const float Multiplier = (*PhysMesh->GetFloatArray(PhysMeshAnimDriveIndex))[ParticleIndex - IndexToRangeMap[i][0]];
+			PDI->DrawLine(AnimationPositions[ParticleIndex], Particles.X(ParticleIndex), Multiplier * AnimDriveSpringStiffness[i] * FColor::Cyan, SDPG_World, 0.0f, 0.001f);
+		}
+	}
 }
 #endif  // #if WITH_EDITOR
diff --git a/Engine/Plugins/Experimental/ChaosCloth/Source/Chaos/Public/ChaosCloth/ChaosClothingSimulation.h b/Engine/Plugins/Experimental/ChaosCloth/Source/Chaos/Public/ChaosCloth/ChaosClothingSimulation.h
--- a/Engine/Plugins/Experimental/ChaosCloth/Source/Chaos/Public/ChaosCloth/ChaosClothingSimulation.h
+++ b/Engine/Plugins/Experimental/ChaosCloth/Source/Chaos/Public/ChaosCloth/ChaosClothingSimulation.h
@@ -136,6 +136,7 @@ namespace Chaos
 #if WITH_EDITOR
 		// Visualization material
 		UMaterial* DebugClothMaterial;
+		UMaterial* DebugClothMaterialVertex;
 #endif  // #if WITH_EDITOR
 	};
 } // namespace Chaos
\ No newline at end of file
diff --git a/Engine/Plugins/Experimental/ChaosClothEditor/Source/Chaos/Private/ChaosClothEditor/ChaosSimulationEditorExtender.cpp b/Engine/Plugins/Experimental/ChaosClothEditor/Source/Chaos/Private/ChaosClothEditor/ChaosSimulationEditorExtender.cpp
--- a/Engine/Plugins/Experimental/ChaosClothEditor/Source/Chaos/Private/ChaosClothEditor/ChaosSimulationEditorExtender.cpp
+++ b/Engine/Plugins/Experimental/ChaosClothEditor/Source/Chaos/Private/ChaosClothEditor/ChaosSimulationEditorExtender.cpp
@@ -53,9 +53,9 @@ const FVisualizationOption FVisualizationOption::OptionData[] =
 	FVisualizationOption(&Chaos::ClothingSimulation::DebugDrawInversedFaceNormals , LOCTEXT("ChaosVisName_InversedFaceNormals" , "Inversed Face Normals"      ), LOCTEXT("ChaosVisName_InversedFaceNormals_ToolTip" , "Draws the inversed face normals for the simulation mesh")),
 	FVisualizationOption(&Chaos::ClothingSimulation::DebugDrawCollision           , LOCTEXT("ChaosVisName_Collision"           , "Collisions"                 ), LOCTEXT("ChaosVisName_Collision_ToolTip"           , "Draws the collision bodies the simulation is currently using")),
 	FVisualizationOption(&Chaos::ClothingSimulation::DebugDrawSelfCollision       , LOCTEXT("ChaosVisName_SelfCollision"       , "Self Collision Radii"       ), LOCTEXT("ChaosVisName_SelfCollision_ToolTip"       , "Draws the self collision radius for each particle if self collision is enabled")),
-	//FVisualizationOption(&Chaos::ClothingSimulation::DebugDrawBackstops           , LOCTEXT("ChaosVisName_Backstop"            , "Backstops"                  ), LOCTEXT("ChaosVisName_Backstop_ToolTip"            , "Draws the backstop offset for each simulation particle"), /*bDisablesSimulation =*/true),
-	//FVisualizationOption(&Chaos::ClothingSimulation::DebugDrawMaxDistances        , LOCTEXT("ChaosVisName_MaxDistance"         , "Max Distances"              ), LOCTEXT("ChaosVisName_MaxDistance_ToolTip"         , "Draws the current max distances for the sim particles as a line along its normal"), true),
-	//FVisualizationOption(&Chaos::ClothingSimulation::DebugDrawAnimDrive           , LOCTEXT("ChaosVisName_AnimDrive"           , "Anim Drive"                 ), LOCTEXT("ChaosVisName_AnimDrive_Tooltip"           , "Draws the current skinned reference mesh for the simulation which anim drive will attempt to reach if enabled"))
+	FVisualizationOption(&Chaos::ClothingSimulation::DebugDrawBackstops           , LOCTEXT("ChaosVisName_Backstop"            , "Backstops"                  ), LOCTEXT("ChaosVisName_Backstop_ToolTip"            , "Draws the backstop offset for each simulation particle"), /*bDisablesSimulation =*/true),
+	FVisualizationOption(&Chaos::ClothingSimulation::DebugDrawMaxDistances        , LOCTEXT("ChaosVisName_MaxDistance"         , "Max Distances"              ), LOCTEXT("ChaosVisName_MaxDistance_ToolTip"         , "Draws the current max distances for the sim particles as a line along its normal"), true),
+	FVisualizationOption(&Chaos::ClothingSimulation::DebugDrawAnimDrive           , LOCTEXT("ChaosVisName_AnimDrive"           , "Anim Drive"                 ), LOCTEXT("ChaosVisName_AnimDrive_Tooltip"           , "Draws the current skinned reference mesh for the simulation which anim drive will attempt to reach if enabled"))
 };
 const uint32 FVisualizationOption::Count = sizeof(OptionData) / sizeof(FVisualizationOption);
 
