Chaos Cloth: Small Bug fixes
- Adding missing index offset
- Debug rendering fixes

diff --git a/Engine/Plugins/Experimental/ChaosCloth/Source/Chaos/Private/ChaosCloth/ChaosClothingSimulation.cpp b/Engine/Plugins/Experimental/ChaosCloth/Source/Chaos/Private/ChaosCloth/ChaosClothingSimulation.cpp
--- a/Engine/Plugins/Experimental/ChaosCloth/Source/Chaos/Private/ChaosCloth/ChaosClothingSimulation.cpp
+++ b/Engine/Plugins/Experimental/ChaosCloth/Source/Chaos/Private/ChaosCloth/ChaosClothingSimulation.cpp
@@ -1263,7 +1263,7 @@ void ClothingSimulation::DebugDrawPhysMeshShaded(USkeletalMeshComponent* OwnerCo
 				const FVector& Pos1 = Particles.X(Element.Y);
 				const FVector& Pos2 = Particles.X(Element.Z);
 
-				const FVector& Normal = FVector::CrossProduct(Pos1 - Pos0, Pos2 - Pos0).GetSafeNormal();
+				const FVector& Normal = FVector::CrossProduct(Pos2 - Pos0, Pos1 - Pos0).GetSafeNormal();
 				const FVector Tangent = ((Pos1 + Pos2) * 0.5f - Pos0).GetSafeNormal();
 
 				MeshBuilder.AddVertex(FDynamicMeshVertex(Pos0, Tangent, Normal, FVector2D(0.f, 0.f), FColor::White));
@@ -1600,6 +1600,7 @@ void ClothingSimulation::DebugDrawBackstops(USkeletalMeshComponent* OwnerCompone
 
 void ClothingSimulation::DebugDrawMaxDistances(USkeletalMeshComponent* OwnerComponent, FPrimitiveDrawInterface* PDI) const
 {
+	const TPBDParticles<float, 3>& Particles = Evolution->Particles();
 	for (int32 i = 0; i < IndexToRangeMap.Num(); ++i)
 	{
 		const UClothingAssetCommon* const Asset = Assets[i];
@@ -1623,9 +1624,12 @@ void ClothingSimulation::DebugDrawMaxDistances(USkeletalMeshComponent* OwnerComp
 		for (uint32 ParticleIndex = IndexToRangeMap[i][0]; ParticleIndex < IndexToRangeMap[i][1]; ++ParticleIndex)
 		{
 			const float Distance = (*PhysMesh->GetFloatArray(PhysMeshMaxDistanceIndex))[ParticleIndex - IndexToRangeMap[i][0]];
-			if (Distance == 0.0f)
+			if (Particles.InvM(ParticleIndex) == 0.0f)
 			{
-				DrawSphere(PDI, AnimationPositions[ParticleIndex], FRotator::ZeroRotator, FVector(0.5f, 0.5f, 0.5f), 10, 10, DebugClothMaterialVertex->GetRenderProxy(), SDPG_World, false);
+				const FMatrix& ViewMatrix = PDI->View->ViewMatrices.GetViewMatrix();
+				const FVector& XAxis = ViewMatrix.GetColumn(0); // Just using transpose here (orthogonal transform assumed)
+				const FVector& YAxis = ViewMatrix.GetColumn(1);
+				DrawDisc(PDI, AnimationPositions[ParticleIndex], XAxis, YAxis, FColor::White, 0.2f, 10, DebugClothMaterialVertex->GetRenderProxy(), SDPG_World);
 			}
 			else
 			{
diff --git a/Engine/Source/Runtime/Experimental/Chaos/Public/Chaos/PBDAnimDriveConstraint.h b/Engine/Source/Runtime/Experimental/Chaos/Public/Chaos/PBDAnimDriveConstraint.h
--- a/Engine/Source/Runtime/Experimental/Chaos/Public/Chaos/PBDAnimDriveConstraint.h
+++ b/Engine/Source/Runtime/Experimental/Chaos/Public/Chaos/PBDAnimDriveConstraint.h
@@ -36,7 +36,7 @@ class TPBDAnimDriveConstraint : public TParticleRule<T, d>
 	inline void ApplyAnimDriveConstraint(TPBDParticles<T, d>& InParticles, const T Dt, const int32 Index) const
 	{
 		
-		if (InParticles.InvM(Index) == 0)
+		if (InParticles.InvM(Index + ParticleIndexOffset) == 0.0f)
 		{
 			return;
 		}
