Add a test against null on the Chaos Cloth shared config in the gravity override function, and change the SetGravityOverride function to take an FVector.

diff --git a/Engine/Plugins/Experimental/ChaosCloth/Source/Chaos/Private/ChaosCloth/ChaosClothingSimulation.cpp b/Engine/Plugins/Experimental/ChaosCloth/Source/Chaos/Private/ChaosCloth/ChaosClothingSimulation.cpp
--- a/Engine/Plugins/Experimental/ChaosCloth/Source/Chaos/Private/ChaosCloth/ChaosClothingSimulation.cpp
+++ b/Engine/Plugins/Experimental/ChaosCloth/Source/Chaos/Private/ChaosCloth/ChaosClothingSimulation.cpp
@@ -1199,14 +1199,15 @@ void ClothingSimulation::SetAnimDriveSpringStiffness(float InStiffness)
 	}
 }
 
-void ClothingSimulation::SetGravityOverride(const TVector<float, 3>& InGravityOverride) 
+void ClothingSimulation::SetGravityOverride(const FVector& InGravityOverride) 
 {
 	Evolution->GetGravityForces().SetAcceleration(InGravityOverride);
 }
 
 void ClothingSimulation::DisableGravityOverride()
 {
-	Evolution->GetGravityForces().SetAcceleration(ClothSharedSimConfig->Gravity);
+	static const FVector DefaultGravity(0.f, 0.f, -980.665f);
+	Evolution->GetGravityForces().SetAcceleration(!ClothSharedSimConfig ? DefaultGravity : ClothSharedSimConfig->Gravity);
 }
 
 #if WITH_EDITOR
diff --git a/Engine/Plugins/Experimental/ChaosCloth/Source/Chaos/Public/ChaosCloth/ChaosClothingSimulation.h b/Engine/Plugins/Experimental/ChaosCloth/Source/Chaos/Public/ChaosCloth/ChaosClothingSimulation.h
--- a/Engine/Plugins/Experimental/ChaosCloth/Source/Chaos/Public/ChaosCloth/ChaosClothingSimulation.h
+++ b/Engine/Plugins/Experimental/ChaosCloth/Source/Chaos/Public/ChaosCloth/ChaosClothingSimulation.h
@@ -36,7 +36,7 @@ namespace Chaos
 
 		// Set the animation drive stiffness for all actors
 		void SetAnimDriveSpringStiffness(float InStiffness);
-		void SetGravityOverride(const TVector<float, 3>& InGravityOverride);
+		void SetGravityOverride(const FVector& InGravityOverride);
 		void DisableGravityOverride();
 
 #if WITH_EDITOR
