ChaosCloth Performance:
Performance Improvements (Test Asset 980us -> 580us for test asset on test platform):
Removed unnecessary copies and buffer creation
Collapsed transform calculations
In place skinning function
Removed duplicate spring constraints
Moved critical functions to header files so that they can be inlined
Long Range Constraints are switched to iterate per constraint instead of per particle

Scoped Cycle counters added

#[fyi] Ryan.Kautzman, Michael.Lentine, Benn.Gallagher, Kriss.Gossart
#[review]-10494959 [REDACTED_HANDLE] Gallagher

diff --git a/Engine/Plugins/Experimental/ChaosCloth/Source/Chaos/Private/ChaosCloth/ChaosClothPrivate.h b/Engine/Plugins/Experimental/ChaosCloth/Source/Chaos/Private/ChaosCloth/ChaosClothPrivate.h
--- a/Engine/Plugins/Experimental/ChaosCloth/Source/Chaos/Private/ChaosCloth/ChaosClothPrivate.h
+++ b/Engine/Plugins/Experimental/ChaosCloth/Source/Chaos/Private/ChaosCloth/ChaosClothPrivate.h
@@ -3,5 +3,7 @@
 #pragma once
 
 #include "CoreMinimal.h"
+#include "Stats/Stats.h"
 
 DECLARE_LOG_CATEGORY_EXTERN(LogChaosCloth, Log, All);
+DECLARE_STATS_GROUP(TEXT("ChaosCloth"), STATGROUP_ChaosCloth, STATCAT_Advanced);
diff --git a/Engine/Plugins/Experimental/ChaosCloth/Source/Chaos/Private/ChaosCloth/ChaosClothingSimulation.cpp b/Engine/Plugins/Experimental/ChaosCloth/Source/Chaos/Private/ChaosCloth/ChaosClothingSimulation.cpp
--- a/Engine/Plugins/Experimental/ChaosCloth/Source/Chaos/Private/ChaosCloth/ChaosClothingSimulation.cpp
+++ b/Engine/Plugins/Experimental/ChaosCloth/Source/Chaos/Private/ChaosCloth/ChaosClothingSimulation.cpp
@@ -18,6 +18,7 @@
 #include "Chaos/PBDAnimDriveConstraint.h"
 #include "Chaos/PBDAxialSpringConstraints.h"
 #include "Chaos/PBDBendingConstraints.h"
+#include "Chaos/PBDLongRangeConstraints.h"
 #include "Chaos/PBDParticles.h"
 #include "Chaos/PBDSphericalConstraint.h"
 #include "Chaos/PBDSpringConstraints.h"
@@ -33,6 +34,7 @@
 #include "Chaos/Utilities.h"
 #include "Chaos/Vector.h"
 #include "ChaosCloth/ChaosClothConfig.h"
+#include "ChaosCloth/ChaosClothPrivate.h"
 
 #if WITH_PHYSX && !PLATFORM_LUMIN && !PLATFORM_ANDROID
 #include "PhysXIncludes.h"
@@ -45,6 +47,10 @@
 
 using namespace Chaos;
 
+DECLARE_CYCLE_STAT(TEXT("Chaos Cloth Simulate"), STAT_ChaosClothSimulate, STATGROUP_ChaosCloth);
+DECLARE_CYCLE_STAT(TEXT("Chaos Cloth Get Animation Data"), STAT_ChaosClothGetAnimationData, STATGROUP_ChaosCloth);
+DECLARE_CYCLE_STAT(TEXT("Chaos Cloth Update Collision Transforms"), STAT_ChaosClothUpdateCollisionTransforms, STATGROUP_ChaosCloth);
+
 ClothingSimulation::ClothingSimulation()
 	: ClothSharedSimConfig(nullptr)
 	, ExternalCollisionsOffset(0)
@@ -86,7 +92,7 @@ void ClothingSimulation::Initialize()
     Evolution->GetGravityForces().SetAcceleration(Chaos::TVector<float, 3>(0.f, 0.f, -1.f)*GravityMagnitude);
 
     Evolution->SetKinematicUpdateFunction(
-		[&](Chaos::TPBDParticles<float, 3>& ParticlesInput, const float Dt, const float LocalTime, const int32 Index)
+		[this](Chaos::TPBDParticles<float, 3>& ParticlesInput, const float Dt, const float LocalTime, const int32 Index)
 		{
 			if (!OldAnimationPositions.IsValidIndex(Index) || ParticlesInput.InvM(Index) > 0)
 				return;
@@ -96,7 +102,7 @@ void ClothingSimulation::Initialize()
 
 	Evolution->SetCollisionKinematicUpdateFunction(
 //		[&](Chaos::TKinematicGeometryParticles<float, 3>& ParticlesInput, const float Dt, const float LocalTime, const int32 Index)
-		[&](Chaos::TKinematicGeometryClothParticles<float, 3>& ParticlesInput, const float Dt, const float LocalTime, const int32 Index)
+		[this](Chaos::TKinematicGeometryClothParticles<float, 3>& ParticlesInput, const float Dt, const float LocalTime, const int32 Index)
 		{
 			checkSlow(Dt > SMALL_NUMBER && DeltaTime > SMALL_NUMBER);
 			const float Alpha = (LocalTime - Time) / DeltaTime;
@@ -179,26 +185,19 @@ void ClothingSimulation::CreateActor(USkeletalMeshComponent* InOwnerComponent, U
 	TArray<Chaos::TVector<float, 3>> TempAnimationPositions;
 	TArray<Chaos::TVector<float, 3>> TempAnimationNormals;
 
+	TempAnimationPositions.SetNumUninitialized(PhysMesh->Vertices.Num());
+	TempAnimationNormals.SetNumUninitialized(PhysMesh->Vertices.Num());
+
 	FTransform RootBoneTransform = Context.BoneTransforms[Asset->ReferenceBoneIndex];
-	ClothingMeshUtils::SkinPhysicsMesh(
+	ClothingMeshUtils::SkinPhysicsMesh<true, false>(
 		Asset->UsedBoneIndices,
 		*PhysMesh, // curr pos and norm
-		RootBoneTransform,
+		Context.ComponentToWorld,
 		Context.RefToLocals.GetData(),
 		Context.RefToLocals.Num(),
 		reinterpret_cast<TArray<FVector>&>(TempAnimationPositions),
 		reinterpret_cast<TArray<FVector>&>(TempAnimationNormals));
 
-	// Transform points & normals to world space
-	RootBoneTransform.SetScale3D(FVector(1.0f));
-	const FTransform RootBoneWorldTransform = RootBoneTransform * Context.ComponentToWorld;
-	ParallelFor(TempAnimationPositions.Num(),
-		[&](int32 Index)
-	{
-		TempAnimationPositions[Index] = RootBoneWorldTransform.TransformPosition(TempAnimationPositions[Index]);
-		TempAnimationNormals[Index] = RootBoneWorldTransform.TransformVector(TempAnimationNormals[Index]);
-	});
-
 	// Add particles
 	TPBDParticles<float, 3>& Particles = Evolution->Particles();
 	const uint32 Offset = Particles.Size();
@@ -394,15 +393,17 @@ void ClothingSimulation::CreateActor(USkeletalMeshComponent* InOwnerComponent, U
 	if (ChaosClothSimConfig->StrainLimitingStiffness)
 	{
 		check(Mesh->GetNumElements() > 0);
-		Chaos::TPerParticlePBDLongRangeConstraints<float, 3> PerParticlePBDLongRangeConstraints(
+		// PerFormance note: The Per constraint version of this function is quite a bit faster for smaller assets
+		// There might be a cross-over point where the PerParticle version is faster: To be determined
+		Chaos::TPBDLongRangeConstraints<float, 3> PBDLongRangeConstraints(
 			Evolution->Particles(),
 			Mesh->GetPointToNeighborsMap(),
 			10, // The max number of connected neighbors per particle.  ryan - What should this be?  Was k...
 			ChaosClothSimConfig->StrainLimitingStiffness);
 
-		Evolution->AddPBDConstraintFunction([PerParticlePBDLongRangeConstraints = MoveTemp(PerParticlePBDLongRangeConstraints)](TPBDParticles<float, 3>& InParticles, const float Dt)
+		Evolution->AddPBDConstraintFunction([PBDLongRangeConstraints = MoveTemp(PBDLongRangeConstraints)](TPBDParticles<float, 3>& InParticles, const float Dt)
 		{
-			PerParticlePBDLongRangeConstraints.Apply(InParticles, Dt);
+			PBDLongRangeConstraints.Apply(InParticles, Dt);
 		});
 	}
 
@@ -514,10 +515,16 @@ void ClothingSimulation::PostActorCreationInitialize()
 
 void ClothingSimulation::UpdateCollisionTransforms(const ClothingSimulationContext& Context)
 {
+	SCOPE_CYCLE_COUNTER(STAT_ChaosClothUpdateCollisionTransforms);
+
+	// Save collision transforms into the OldCollision transforms
+	// before overwriting it in this function
+	Swap(OldCollisionTransforms, CollisionTransforms);
+
 	TGeometryClothParticles<float, 3>& CollisionParticles = Evolution->CollisionParticles();
 
 	// Resize the transform arrays
-	const int32 PrevNumCollisions = CollisionTransforms.Num();
+	const int32 PrevNumCollisions = OldCollisionTransforms.Num();
 	const int32 NumCollisions = BaseTransforms.Num();
 	check(NumCollisions == int32(CollisionParticles.Size()));  // BaseTransforms should always automatically grow with the number of collision particles (collection array)
 
@@ -1026,54 +1033,50 @@ void ClothingSimulation::FillContext(USkeletalMeshComponent* InComponent, float
 
 void ClothingSimulation::Simulate(IClothingSimulationContext* InContext)
 {
+	SCOPE_CYCLE_COUNTER(STAT_ChaosClothSimulate);
 	ClothingSimulationContext* Context = static_cast<ClothingSimulationContext*>(InContext);
 	if (Context->DeltaTime == 0)
 		return;
 
 	// Get New Animation Positions and Normals
-	OldCollisionTransforms = CollisionTransforms;
-	OldAnimationPositions = AnimationPositions;
-
-	for (int32 Index = 0; Index < IndexToRangeMap.Num(); ++Index)
 	{
-		const UClothingAssetCommon* const Asset = Assets[Index];
-		if (!Asset)
-			continue;
-
-		const UClothLODDataCommon* AssetLodData = Asset->ClothLodData[0];
-		check(AssetLodData->PhysicalMeshData);
-		const UClothPhysicalMeshDataBase* PhysMesh = AssetLodData->PhysicalMeshData;
-
-		TArray<Chaos::TVector<float, 3>> TempAnimationPositions;
-		TArray<Chaos::TVector<float, 3>> TempAnimationNormals;
+		SCOPE_CYCLE_COUNTER(STAT_ChaosClothGetAnimationData);
 
-		FTransform RootBoneTransform = Context->BoneTransforms[Asset->ReferenceBoneIndex];
-		ClothingMeshUtils::SkinPhysicsMesh(
-			Asset->UsedBoneIndices,
-			*PhysMesh,
-			RootBoneTransform,
-			Context->RefToLocals.GetData(),
-			Context->RefToLocals.Num(),
-			reinterpret_cast<TArray<FVector>&>(TempAnimationPositions),
-			reinterpret_cast<TArray<FVector>&>(TempAnimationNormals));
-
-		RootBoneTransform.SetScale3D(FVector(1.0f));
-
-		// Removing Context->ComponentToWorld means the sim doesn't see updates to the component level xf
-		const FTransform RootBoneWorldTransform = RootBoneTransform * Context->ComponentToWorld;
-
-		const int32 Offset = IndexToRangeMap[Index][0];
-		check(TempAnimationPositions.Num() == IndexToRangeMap[Index][1] - IndexToRangeMap[Index][0]);
-
-		ParallelFor(TempAnimationPositions.Num(),
-		[&](int32 AnimationElementIndex)
+		Swap(OldAnimationPositions, AnimationPositions);
+		check(OldAnimationPositions.Num() == AnimationPositions.Num());
+		
+		for (int32 Index = 0; Index < IndexToRangeMap.Num(); ++Index)
 		{
-			AnimationPositions[Offset + AnimationElementIndex] = RootBoneWorldTransform.TransformPosition(TempAnimationPositions[AnimationElementIndex]);
-			AnimationNormals[Offset + AnimationElementIndex] = RootBoneWorldTransform.TransformVector(TempAnimationNormals[AnimationElementIndex]);
-		});
+			const UClothingAssetCommon* const Asset = Assets[Index];
+			if (!Asset)
+			{
+				continue;
+			}
+
+			const UClothLODDataCommon* AssetLodData = Asset->ClothLodData[0];
+			check(AssetLodData->PhysicalMeshData);
+			const UClothPhysicalMeshDataBase* PhysMesh = AssetLodData->PhysicalMeshData;
+			const uint32 PointCount = IndexToRangeMap[Index][1] - IndexToRangeMap[Index][0];
+
+			// Optimization note:
+			// This function usually receives the RootBoneTransform in order to transform the result from Component space to RootBone space.
+			// (So the mesh vectors and positions (Mesh is in component space) is multiplied by Inv(RootBoneTransform) at the end of the function)
+			// We actually require world space coordinates so will instead pass Inv(ComponentToWorld)
+			// This saves a lot of Matrix multiplication work later
+			const int32 Offset = IndexToRangeMap[Index][0];
+
+			ClothingMeshUtils::SkinPhysicsMesh<true, false>(
+				Asset->UsedBoneIndices,
+				*PhysMesh,
+				Context->ComponentToWorld,
+				Context->RefToLocals.GetData(),
+				Context->RefToLocals.Num(),
+				reinterpret_cast<TArray<FVector>&>(AnimationPositions),
+				reinterpret_cast<TArray<FVector>&>(AnimationNormals), Offset);
+		}
 	}
 
-	// Update collision tranforms
+	// Update collision transforms
 	UpdateCollisionTransforms(*Context);
 
 	// Advance Sim
@@ -1199,7 +1202,7 @@ void ClothingSimulation::SetAnimDriveSpringStiffness(float InStiffness)
 	}
 }
 
-void ClothingSimulation::SetGravityOverride(const FVector& InGravityOverride) 
+void ClothingSimulation::SetGravityOverride(const FVector& InGravityOverride)
 {
 	Evolution->GetGravityForces().SetAcceleration(InGravityOverride);
 }
diff --git a/Engine/Source/Runtime/ClothingSystemRuntimeCommon/Private/Utils/ClothingMeshUtils.cpp b/Engine/Source/Runtime/ClothingSystemRuntimeCommon/Private/Utils/ClothingMeshUtils.cpp
--- a/Engine/Source/Runtime/ClothingSystemRuntimeCommon/Private/Utils/ClothingMeshUtils.cpp
+++ b/Engine/Source/Runtime/ClothingSystemRuntimeCommon/Private/Utils/ClothingMeshUtils.cpp
@@ -19,27 +19,45 @@ DECLARE_CYCLE_STAT(TEXT("Skin Physics Mesh"), STAT_ClothSkinPhysMesh, STATGROUP_
 
 namespace ClothingMeshUtils
 {
+	// Explicit template instantiations of SkinPhysicsMesh
+	template
+	void CLOTHINGSYSTEMRUNTIMECOMMON_API SkinPhysicsMesh<true, false>(const TArray<int32>& BoneMap, const UClothPhysicalMeshDataBase& InMesh, const FTransform& RootBoneTransform,
+		const FMatrix* InBoneMatrices, const int32 InNumBoneMatrices, TArray<FVector>& OutPositions, TArray<FVector>& OutNormals, uint32 ArrayOffset);
+	template
+	void CLOTHINGSYSTEMRUNTIMECOMMON_API SkinPhysicsMesh<false, true>(const TArray<int32>& BoneMap, const UClothPhysicalMeshDataBase& InMesh, const FTransform& RootBoneTransform,
+		const FMatrix* InBoneMatrices, const int32 InNumBoneMatrices, TArray<FVector>& OutPositions, TArray<FVector>& OutNormals, uint32 ArrayOffset);
+
+	template<bool bInPlaceOutput, bool bRemoveScaleAndInvertPostTransform>
 	void SkinPhysicsMesh(
 		const TArray<int32>& BoneMap,
 		const UClothPhysicalMeshDataBase& InMesh,
-		const FTransform& RootBoneTransform,
+		const FTransform& PostTransform,
 		const FMatrix* InBoneMatrices,
 		const int32 InNumBoneMatrices,
 		TArray<FVector>& OutPositions,
-		TArray<FVector>& OutNormals)
+		TArray<FVector>& OutNormals,
+		uint32 ArrayOffset)
 	{
 		SCOPE_CYCLE_COUNTER(STAT_ClothSkinPhysMesh);
 
-		// Ignore any user scale. It's already accounted for in our skinning matrices
-		FTransform RootBoneTransformInternal = RootBoneTransform;
-		RootBoneTransformInternal.SetScale3D(FVector(1.0f));
-
 		const uint32 NumVerts = InMesh.Vertices.Num();
 
-		OutPositions.Reset(NumVerts);
-		OutNormals.Reset(NumVerts);
-		OutPositions.AddZeroed(NumVerts);
-		OutNormals.AddZeroed(NumVerts);
+		if(!bInPlaceOutput)
+		{
+			ensure(ArrayOffset == 0);
+			OutPositions.Reset(NumVerts);
+			OutNormals.Reset(NumVerts);
+			OutPositions.AddZeroed(NumVerts);
+			OutNormals.AddZeroed(NumVerts);
+		}
+		else
+		{
+			check((uint32) OutPositions.Num() >= NumVerts + ArrayOffset);
+			check((uint32) OutNormals.Num() >= NumVerts + ArrayOffset);
+			// PS4 performance note: It is faster to zero the memory first instead of changing this function to work with uninitialized memory
+			FMemory::Memzero((uint8*)OutPositions.GetData() + ArrayOffset * sizeof(FVector), NumVerts * sizeof(FVector));
+			FMemory::Memzero((uint8*)OutNormals.GetData() + ArrayOffset * sizeof(FVector), NumVerts * sizeof(FVector));
+		}
 
 		const int32 MaxInfluences = InMesh.MaxBoneWeights;
 		const FMatrix* RESTRICT BoneMatrices = InBoneMatrices;
@@ -56,8 +74,8 @@ namespace ClothingMeshUtils
 			// everything to compose the final skinned data
 			const FVector& RefParticle = InMesh.Vertices[VertIndex];
 			const FVector& RefNormal = InMesh.Normals[VertIndex];
-			FVector& OutPosition = OutPositions[VertIndex];
-			FVector& OutNormal = OutNormals[VertIndex];
+			FVector& OutPosition = OutPositions[bInPlaceOutput ? VertIndex + ArrayOffset : VertIndex];
+			FVector& OutNormal = OutNormals[bInPlaceOutput ? VertIndex + ArrayOffset : VertIndex];
 			switch (InMesh.BoneData[VertIndex].NumInfluences)
 			{
 			default: break;
@@ -127,8 +145,21 @@ namespace ClothingMeshUtils
 			}
 			}
 
-			OutPosition = RootBoneTransformInternal.InverseTransformPosition(OutPosition);
-			OutNormal = RootBoneTransformInternal.InverseTransformVector(OutNormal);
+			if (bRemoveScaleAndInvertPostTransform)
+			{
+				// Ignore any user scale. It's already accounted for in our skinning matrices
+				// This is the use case for NVcloth
+				FTransform PostTransformInternal = PostTransform;
+				PostTransformInternal.SetScale3D(FVector(1.0f));
+
+				OutPosition = PostTransformInternal.InverseTransformPosition(OutPosition);
+				OutNormal = PostTransformInternal.InverseTransformVector(OutNormal);
+			}
+			else
+			{
+				OutPosition = PostTransform.TransformPosition(OutPosition);
+				OutNormal = PostTransform.TransformVector(OutNormal);
+			}
 
 			if (OutNormal.SizeSquared() > SMALL_NUMBER)
 			{
diff --git a/Engine/Source/Runtime/ClothingSystemRuntimeCommon/Public/Utils/ClothingMeshUtils.h b/Engine/Source/Runtime/ClothingSystemRuntimeCommon/Public/Utils/ClothingMeshUtils.h
--- a/Engine/Source/Runtime/ClothingSystemRuntimeCommon/Public/Utils/ClothingMeshUtils.h
+++ b/Engine/Source/Runtime/ClothingSystemRuntimeCommon/Public/Utils/ClothingMeshUtils.h
@@ -29,15 +29,19 @@ namespace ClothingMeshUtils
 	};
 
 	// Static method for calculating a skinned mesh result from source data
+	// The bInPlaceOutput allows us to directly populate arrays that are already allocated
+	// bRemoveScaleAndInvertPostTransform will determine if the PostTransform should be inverted and the scale removed (NvCloth uses this). It is templated to remove branches at compile time
+	template<bool bInPlaceOutput = false, bool bRemoveScaleAndInvertPostTransform = true>
 	void CLOTHINGSYSTEMRUNTIMECOMMON_API
 	SkinPhysicsMesh(
 		const TArray<int32>& BoneMap, // UClothingAssetCommon::UsedBoneIndices
 		const UClothPhysicalMeshDataBase& InMesh, 
-		const FTransform& RootBoneTransform, 
+		const FTransform& PostTransform, // Final transform to apply to component space positions and normals
 		const FMatrix* InBoneMatrices, 
 		const int32 InNumBoneMatrices, 
 		TArray<FVector>& OutPositions, 
-		TArray<FVector>& OutNormals);
+		TArray<FVector>& OutNormals,
+		uint32 ArrayOffset = 0); // Used for Chaos Cloth
 
 	/**
 	* Given mesh information for two meshes, generate a list of skinning data to embed mesh0 in mesh1
diff --git a/Engine/Source/Runtime/Experimental/Chaos/Private/Chaos/PBDEvolution.cpp b/Engine/Source/Runtime/Experimental/Chaos/Private/Chaos/PBDEvolution.cpp
--- a/Engine/Source/Runtime/Experimental/Chaos/Private/Chaos/PBDEvolution.cpp
+++ b/Engine/Source/Runtime/Experimental/Chaos/Private/Chaos/PBDEvolution.cpp
@@ -12,6 +12,18 @@
 #include "Chaos/PerParticlePBDEulerStep.h"
 #include "Chaos/PerParticlePBDGroundConstraint.h"
 #include "Chaos/PerParticlePBDUpdateFromDeltaPosition.h"
+#include "ChaosStats.h"
+
+
+DECLARE_CYCLE_STAT(TEXT("Chaos PBD Advance Time"), STAT_ChaosPBDVAdvanceTime, STATGROUP_Chaos);
+DECLARE_CYCLE_STAT(TEXT("Chaos PBD Velocity Damping State Update"), STAT_ChaosPBDVelocityDampUpdateState, STATGROUP_Chaos);
+DECLARE_CYCLE_STAT(TEXT("Chaos PBD Velocity Damping"), STAT_ChaosPBDVelocityDampUpdate, STATGROUP_Chaos);
+DECLARE_CYCLE_STAT(TEXT("Chaos PBD Pre Iteration Updates"), STAT_ChaosPBDPreIterationUpdates, STATGROUP_Chaos);
+DECLARE_CYCLE_STAT(TEXT("Chaos PBD Constraint Rule"), STAT_ChaosPBDConstraintRule, STATGROUP_Chaos);
+DECLARE_CYCLE_STAT(TEXT("Chaos PBD Self Collision"), STAT_ChaosPBDSelfCollisionRule, STATGROUP_Chaos);
+DECLARE_CYCLE_STAT(TEXT("Chaos PBD Collision"), STAT_ChaosPBDCollisionRule, STATGROUP_Chaos);
+DECLARE_CYCLE_STAT(TEXT("Chaos PBD Collider Friction"), STAT_ChaosPBDCollisionRuleFriction, STATGROUP_Chaos);
+DECLARE_CYCLE_STAT(TEXT("Chaos PBD Collider Kinematic Update"), STAT_CollisionKinematicUpdate, STATGROUP_Chaos);
 
 using namespace Chaos;
 
@@ -40,23 +52,22 @@ TPBDEvolution<T, d>::TPBDEvolution(TPBDParticles<T, d>&& InParticles, TKinematic
 					PBDUpdateRule.Apply(MParticlesInput, Dt, Index);
 				}, NonParallelUpdate);
 			});
-			
-	AddForceFunction([this](TPBDParticles<T, d>& ParticlesInput, const T Dt, const int Index)
-	{
-		GravityForces.Apply(ParticlesInput, Dt, Index);
-	});
 }
 
 template<class T, int d>
 void TPBDEvolution<T, d>::AdvanceOneTimeStep(const T Dt)
 {
+	SCOPE_CYCLE_COUNTER(STAT_ChaosPBDVAdvanceTime);
 	TPerParticleInitForce<T, d> InitForceRule;
 	TPerParticleEulerStepVelocity<T, d> EulerStepVelocityRule;
 	TPerParticleDampVelocity<T, d> DampVelocityRule(MDamping);
 	TPerParticlePBDEulerStep<T, d> EulerStepRule;
 	TPerParticlePBDCollisionConstraint<T, d, EGeometryParticlesSimType::Other> CollisionRule(MCollisionParticles, MCollided, MCollisionThickness, MCoefficientOfFriction);
 
-	DampVelocityRule.UpdatePositionBasedState(MParticles);
+	{
+		SCOPE_CYCLE_COUNTER(STAT_ChaosPBDVelocityDampUpdateState);
+		DampVelocityRule.UpdatePositionBasedState(MParticles);
+	}	
 
 	// Don't bother with threaded execution if we don't have enough work to make it worth while.
 	const bool NonParallelUpdate = MParticles.Size() < 1000; // TODO: 1000 is a guess, tune this!
@@ -66,25 +77,29 @@ void TPBDEvolution<T, d>::AdvanceOneTimeStep(const T Dt)
 	//for(int32 i=0; i < MCollided.Num(); i++)
 	//	MCollided[i] = false;
 	memset(MCollided.GetData(), 0, MCollided.Num()*sizeof(bool));
-
-	PhysicsParallelFor(MParticles.Size(), [&](int32 Index) 
 	{
-		InitForceRule.Apply(MParticles, Dt, Index); // F = TV(0)
-		for (auto ForceRule : MForceRules)
-		{
-			ForceRule(MParticles, Dt, Index); // F += M * A
-		}
-		if (MKinematicUpdate)
+		SCOPE_CYCLE_COUNTER(STAT_ChaosPBDPreIterationUpdates);
+		PhysicsParallelFor(MParticles.Size(), [&](int32 Index)
 		{
-			MKinematicUpdate(MParticles, Dt, MTime + Dt, Index); // X = ...
-		}
-		EulerStepVelocityRule.Apply(MParticles, Dt, Index);
-		DampVelocityRule.Apply(MParticles, Dt, Index);
-		EulerStepRule.Apply(MParticles, Dt, Index);
-	}, NonParallelUpdate);
+			InitForceRule.Apply(MParticles, Dt, Index); // F = TV(0)
+			GravityForces.Apply(MParticles, Dt, Index);
+			for (TFunction<void(TPBDParticles<T, d>&, const T, const int32)>& ForceRule : MForceRules)
+			{
+				ForceRule(MParticles, Dt, Index); // F += M * A
+			}
+			if (MKinematicUpdate)
+			{
+				MKinematicUpdate(MParticles, Dt, MTime + Dt, Index); // X = ...
+			}
+			EulerStepVelocityRule.Apply(MParticles, Dt, Index);
+			DampVelocityRule.Apply(MParticles, Dt, Index);
+			EulerStepRule.Apply(MParticles, Dt, Index);
+		}, NonParallelUpdate);
+	}
 
 	if (MCollisionKinematicUpdate)
 	{
+		SCOPE_CYCLE_COUNTER(STAT_CollisionKinematicUpdate);
 		PhysicsParallelFor(MCollisionParticles.Size(), [&](int32 Index) {
 			MCollisionKinematicUpdate(MCollisionParticles, Dt, MTime + Dt, Index);
 		}, NonParallelUpdate);
@@ -94,19 +109,27 @@ void TPBDEvolution<T, d>::AdvanceOneTimeStep(const T Dt)
 #endif
 	for (int i = 0; i < MNumIterations; ++i)
 	{
-		for (auto ConstraintRule : MConstraintRules)
+		for (TFunction<void(TPBDParticles<T, d>&, const T)>& ConstraintRule : MConstraintRules)
 		{
+			SCOPE_CYCLE_COUNTER(STAT_ChaosPBDConstraintRule);
 			ConstraintRule(MParticles, Dt); // P +/-= ...
 		}
 #if !COMPILE_WITHOUT_UNREAL_SUPPORT
-		SelfCollisionRule.Apply(MParticles, Dt);
+		{
+			SCOPE_CYCLE_COUNTER(STAT_ChaosPBDSelfCollisionRule);
+			SelfCollisionRule.Apply(MParticles, Dt);
+		}
 #endif
-		CollisionRule.ApplyPerParticle(MParticles, Dt);
+		{
+			SCOPE_CYCLE_COUNTER(STAT_ChaosPBDCollisionRule);
+			CollisionRule.ApplyPerParticle(MParticles, Dt);
+		}
 	}
 	check(MParticleUpdate);
 	MParticleUpdate(MParticles, Dt); // V = (P - X) / Dt; X = P;
 	if (MCoefficientOfFriction > 0)
 	{
+		SCOPE_CYCLE_COUNTER(STAT_ChaosPBDCollisionRuleFriction);
 		PhysicsParallelFor(MParticles.Size(), [&](int32 Index) {
 			CollisionRule.ApplyFriction(MParticles, Dt, Index);
 		}, NonParallelUpdate);
diff --git a/Engine/Source/Runtime/Experimental/Chaos/Private/Chaos/PBDLongRangeConstraintsBase.cpp b/Engine/Source/Runtime/Experimental/Chaos/Private/Chaos/PBDLongRangeConstraintsBase.cpp
--- a/Engine/Source/Runtime/Experimental/Chaos/Private/Chaos/PBDLongRangeConstraintsBase.cpp
+++ b/Engine/Source/Runtime/Experimental/Chaos/Private/Chaos/PBDLongRangeConstraintsBase.cpp
@@ -316,39 +316,4 @@ void TPBDLongRangeConstraintsBase<T, d>::ComputeGeodesicConstraints(
 	MConstraints = NewConstraints;
 }
 
-template<class T, int d>
-TVector<T, d> TPBDLongRangeConstraintsBase<T, d>::GetDelta(const TPBDParticles<T, d>& InParticles, const int32 i) const
-{
-	const TArray<uint32>& Constraint = MConstraints[i];
-	check(Constraint.Num() > 1);
-	const uint32 i1 = Constraint[0];
-	const uint32 i2 = Constraint[Constraint.Num() - 1];
-	const uint32 i2m1 = Constraint[Constraint.Num() - 2];
-	check(InParticles.InvM(i1) == 0);
-	check(InParticles.InvM(i2) > 0);
-	const T Distance = ComputeGeodesicDistance(InParticles, Constraint);
-	if (Distance < MDists[i])
-		return TVector<T, d>(0);
-
-	//const TVector<T, d> Direction = (InParticles.P(i2m1) - InParticles.P(i2)).GetSafeNormal();
-	TVector<T, d> Direction = InParticles.P(i2m1) - InParticles.P(i2);
-	const T DirLen = Direction.SafeNormalize();
-
-	const T Offset = Distance - MDists[i];
-	const TVector<T, d> Delta = MStiffness * Offset * Direction;
-
-	/*  // ryan - this currently fails:
-
-	const T NewDirLen = (InParticles.P(i2) + Delta - InParticles.P(i2m1)).Size();
-	//T Correction = (InParticles.P(i2) - InParticles.P(i2m1)).Size() - (InParticles.P(i2) + Delta - InParticles.P(i2m1)).Size();
-	const T Correction = DirLen - NewDirLen;
-	check(Correction >= 0);
-
-	//T NewDist = (Distance - (InParticles.P(i2) - InParticles.P(i2m1)).Size() + (InParticles.P(i2) + Delta - InParticles.P(i2m1)).Size());
-	const T NewDist = Distance - DirLen + NewDirLen;
-	check(FGenericPlatformMath::Abs(NewDist - MDists[i]) < 1e-4);
-*/
-	return Delta;
-}
-
 template class Chaos::TPBDLongRangeConstraintsBase<float, 3>;
diff --git a/Engine/Source/Runtime/Experimental/Chaos/Public/Chaos/PBDAnimDriveConstraint.h b/Engine/Source/Runtime/Experimental/Chaos/Public/Chaos/PBDAnimDriveConstraint.h
--- a/Engine/Source/Runtime/Experimental/Chaos/Public/Chaos/PBDAnimDriveConstraint.h
+++ b/Engine/Source/Runtime/Experimental/Chaos/Public/Chaos/PBDAnimDriveConstraint.h
@@ -9,6 +9,8 @@ template<typename T, int d>
 class TPBDAnimDriveConstraint : public TParticleRule<T, d>
 {
   public:
+	// InSpringNeutralPositions starts at index: InParticleIndexOffset
+	// InSpringStiffnessMultiplier starts at index: 0
 	TPBDAnimDriveConstraint(
 		const int InParticleIndexOffset
 		, const TArray<TVector<T, d>>* const InSpringNeutralPositions
@@ -55,8 +57,8 @@ class TPBDAnimDriveConstraint : public TParticleRule<T, d>
 
 private:
 	const int32 ParticleIndexOffset;
-	const TArray<TVector<T, d>>* SpringNeutralPositions; // Size: Same as full particle array
-	const TArray<T>* SpringStiffnessMultiplier; // Size: Number of Animation drive constraints to solve
+	const TArray<TVector<T, d>>* SpringNeutralPositions; // Size: Same as full particle array // Starts at Index: ParticleIndexOffset
+	const TArray<T>* SpringStiffnessMultiplier; // Size: Number of Animation drive constraints to solve // Starts at Index 0
 
 	T SpringStiffness;
 };
diff --git a/Engine/Source/Runtime/Experimental/Chaos/Public/Chaos/PBDAxialSpringConstraints.h b/Engine/Source/Runtime/Experimental/Chaos/Public/Chaos/PBDAxialSpringConstraints.h
--- a/Engine/Source/Runtime/Experimental/Chaos/Public/Chaos/PBDAxialSpringConstraints.h
+++ b/Engine/Source/Runtime/Experimental/Chaos/Public/Chaos/PBDAxialSpringConstraints.h
@@ -5,9 +5,11 @@
 #include "Chaos/PBDAxialSpringConstraintsBase.h"
 #include "Chaos/PBDParticles.h"
 #include "Chaos/PerParticleRule.h"
+#include "ChaosStats.h"
 
 #include <algorithm>
 
+DECLARE_CYCLE_STAT(TEXT("Chaos PBD Axial Spring Constraint"), STAT_PBD_AxialSpring, STATGROUP_Chaos);
 namespace Chaos
 {
 template<class T, int d>
@@ -24,6 +26,7 @@ class TPBDAxialSpringConstraints : public TParticleRule<T, d>, public TPBDAxialS
 
 	void Apply(TPBDParticles<T, d>& InParticles, const T Dt) const override //-V762
 	{
+		SCOPE_CYCLE_COUNTER(STAT_PBD_AxialSpring);
 		for (int32 i = 0; i < MConstraints.Num(); ++i)
 		{
 			const auto& constraint = MConstraints[i];
diff --git a/Engine/Source/Runtime/Experimental/Chaos/Public/Chaos/PBDLongRangeConstraints.h b/Engine/Source/Runtime/Experimental/Chaos/Public/Chaos/PBDLongRangeConstraints.h
--- a/Engine/Source/Runtime/Experimental/Chaos/Public/Chaos/PBDLongRangeConstraints.h
+++ b/Engine/Source/Runtime/Experimental/Chaos/Public/Chaos/PBDLongRangeConstraints.h
@@ -4,28 +4,47 @@
 #include "Chaos/PBDLongRangeConstraintsBase.h"
 #include "Chaos/PBDParticles.h"
 #include "Chaos/PBDConstraintContainer.h"
+#include "ChaosStats.h"
+
+DECLARE_CYCLE_STAT(TEXT("Chaos PBD Long Range Constraint"), STAT_PBD_LongRange, STATGROUP_Chaos);
 
 namespace Chaos
 {
 template<class T, int d>
-class CHAOS_API PBDLongRangeConstraints : public TPBDLongRangeConstraintsBase<T, d>, public FPBDConstraintContainer
+class TPBDLongRangeConstraints : public TPBDLongRangeConstraintsBase<T, d>, public FPBDConstraintContainer
 {
 	typedef TPBDLongRangeConstraintsBase<T, d> Base;
 	using Base::MConstraints;
 
   public:
-	PBDLongRangeConstraints(const TDynamicParticles<T, d>& InParticles, const TTriangleMesh<T>& Mesh, const int32 NumberOfAttachments = 1, const T Stiffness = (T)1)
-	    : TPBDLongRangeConstraintsBase<T, d>(InParticles, Mesh, NumberOfAttachments, Stiffness) {}
-	virtual ~PBDLongRangeConstraints() {}
+	TPBDLongRangeConstraints(const TDynamicParticles<T, d>& InParticles, const TMap<int32, TSet<uint32>>& PointToNeighbors, const int32 NumberOfAttachments = 1, const T Stiffness = (T)1)
+	    : TPBDLongRangeConstraintsBase<T, d>(InParticles, PointToNeighbors, NumberOfAttachments, Stiffness) {}
+	virtual ~TPBDLongRangeConstraints() {}
+
+
+	void Apply(TPBDParticles<T, d>& InParticles, const T Dt, int32 Index) const
+	{
+		const auto& Constraint = MConstraints[Index];
+		int32 i2 = Constraint[Constraint.Num() - 1];
+		check(InParticles.InvM(i2) > 0);
+		InParticles.P(i2) += Base::GetDelta(InParticles, Index);
+	}
+
+	void Apply(TPBDParticles<T, d>& InParticles, const T Dt) const 
+	{
+		SCOPE_CYCLE_COUNTER(STAT_PBD_LongRange);
+		for (int32 i = 0; i < MConstraints.Num(); ++i)
+		{
+			Apply(InParticles, Dt, i);
+		}
+	}
 
 	void Apply(TPBDParticles<T, d>& InParticles, const T Dt, const TArray<int32>& InConstraintIndices) const
 	{
+		SCOPE_CYCLE_COUNTER(STAT_PBD_LongRange);
 		for (int32 i : InConstraintIndices)
 		{
-			const auto& Constraint = MConstraints[i];
-			int32 i2 = Constraint[Constraint.Num() - 1];
-			check(InParticles.InvM(i2) > 0);
-			InParticles.P(i2) += Base::GetDelta(InParticles, i);
+			Apply(InParticles, Dt, i);
 		}
 	}
 };
diff --git a/Engine/Source/Runtime/Experimental/Chaos/Public/Chaos/PBDLongRangeConstraintsBase.h b/Engine/Source/Runtime/Experimental/Chaos/Public/Chaos/PBDLongRangeConstraintsBase.h
--- a/Engine/Source/Runtime/Experimental/Chaos/Public/Chaos/PBDLongRangeConstraintsBase.h
+++ b/Engine/Source/Runtime/Experimental/Chaos/Public/Chaos/PBDLongRangeConstraintsBase.h
@@ -15,7 +15,39 @@ class CHAOS_API TPBDLongRangeConstraintsBase
 	TPBDLongRangeConstraintsBase(const TDynamicParticles<T, d>& InParticles, const TMap<int32, TSet<uint32>>& PointToNeighbors, const int32 NumberOfAttachments = 1, const T Stiffness = (T)1);
 	virtual ~TPBDLongRangeConstraintsBase() {}
 
-	TVector<T, d> GetDelta(const TPBDParticles<T, d>& InParticles, const int32 i) const;
+	inline TVector<T, d> GetDelta(const TPBDParticles<T, d>& InParticles, const int32 i) const
+	{
+		const TArray<uint32>& Constraint = MConstraints[i];
+		check(Constraint.Num() > 1);
+		const uint32 i1 = Constraint[0];
+		const uint32 i2 = Constraint[Constraint.Num() - 1];
+		const uint32 i2m1 = Constraint[Constraint.Num() - 2];
+		check(InParticles.InvM(i1) == 0);
+		check(InParticles.InvM(i2) > 0);
+		const T Distance = ComputeGeodesicDistance(InParticles, Constraint); // This function is used for either Euclidian or Geodisic distances
+		if (Distance < MDists[i])
+			return TVector<T, d>(0);
+
+		//const TVector<T, d> Direction = (InParticles.P(i2m1) - InParticles.P(i2)).GetSafeNormal();
+		TVector<T, d> Direction = InParticles.P(i2m1) - InParticles.P(i2);
+		const T DirLen = Direction.SafeNormalize();
+
+		const T Offset = Distance - MDists[i];
+		const TVector<T, d> Delta = MStiffness * Offset * Direction;
+
+	/*  // ryan - this currently fails:
+
+		const T NewDirLen = (InParticles.P(i2) + Delta - InParticles.P(i2m1)).Size();
+		//T Correction = (InParticles.P(i2) - InParticles.P(i2m1)).Size() - (InParticles.P(i2) + Delta - InParticles.P(i2m1)).Size();
+		const T Correction = DirLen - NewDirLen;
+		check(Correction >= 0);
+
+		//T NewDist = (Distance - (InParticles.P(i2) - InParticles.P(i2m1)).Size() + (InParticles.P(i2) + Delta - InParticles.P(i2m1)).Size());
+		const T NewDist = Distance - DirLen + NewDirLen;
+		check(FGenericPlatformMath::Abs(NewDist - MDists[i]) < 1e-4);
+	*/
+		return Delta;
+	};
 
   private:
 	static TArray<TArray<uint32>> ComputeIslands(const TDynamicParticles<T, d>& InParticles, const TMap<int32, TSet<uint32>>& PointToNeighbors,/*const TTriangleMesh<T>& Mesh,*/ const TArray<uint32>& KinematicParticles);
diff --git a/Engine/Source/Runtime/Experimental/Chaos/Public/Chaos/PBDSphericalConstraint.h b/Engine/Source/Runtime/Experimental/Chaos/Public/Chaos/PBDSphericalConstraint.h
--- a/Engine/Source/Runtime/Experimental/Chaos/Public/Chaos/PBDSphericalConstraint.h
+++ b/Engine/Source/Runtime/Experimental/Chaos/Public/Chaos/PBDSphericalConstraint.h
@@ -2,6 +2,9 @@
 #pragma once
 
 #include "ParticleRule.h"
+#include "ChaosStats.h"
+
+DECLARE_CYCLE_STAT(TEXT("Chaos PBD Spherical Constraints"), STAT_PBD_Spherical, STATGROUP_Chaos);
 
 namespace Chaos
 {
@@ -33,6 +36,7 @@ public:
 
 	inline void Apply(TPBDParticles<T, d>& InParticles, const T Dt) const override
 	{
+		SCOPE_CYCLE_COUNTER(STAT_PBD_Spherical);
 		const T CenterDeadZoneRadius = KINDA_SMALL_NUMBER; // We will not push the particle away in the dead zone
 
 		PhysicsParallelFor(ParticleCount, [&](int32 Index)
diff --git a/Engine/Source/Runtime/Experimental/Chaos/Public/Chaos/PBDSpringConstraints.h b/Engine/Source/Runtime/Experimental/Chaos/Public/Chaos/PBDSpringConstraints.h
--- a/Engine/Source/Runtime/Experimental/Chaos/Public/Chaos/PBDSpringConstraints.h
+++ b/Engine/Source/Runtime/Experimental/Chaos/Public/Chaos/PBDSpringConstraints.h
@@ -5,9 +5,12 @@
 #include "Chaos/PBDParticles.h"
 #include "Chaos/PBDSpringConstraintsBase.h"
 #include "Chaos/PBDConstraintContainer.h"
+#include "ChaosStats.h"
 
 #include "Templates/EnableIf.h"
 
+DECLARE_CYCLE_STAT(TEXT("Chaos PBD Spring Constraint"), STAT_PBD_Spring, STATGROUP_Chaos);
+
 namespace Chaos
 {
 template<class T, int32 d>
@@ -63,6 +66,7 @@ class TPBDSpringConstraints : public TPBDSpringConstraintsBase<T, d>, public FPB
 
 	void Apply(TPBDParticles<T, d>& InParticles, const T Dt) const
 	{
+		SCOPE_CYCLE_COUNTER(STAT_PBD_Spring);
 		for (int32 i = 0; i < MConstraints.Num(); ++i)
 		{
 			Apply(InParticles, Dt, i);
@@ -71,6 +75,7 @@ class TPBDSpringConstraints : public TPBDSpringConstraintsBase<T, d>, public FPB
 
 	void Apply(TPBDRigidParticles<T, d>& InParticles, const T Dt, const TArray<int32>& InConstraintIndices) const
 	{
+		SCOPE_CYCLE_COUNTER(STAT_PBD_Spring);
 		for (int32 i : InConstraintIndices)
 		{
 			const auto& Constraint = MConstraints[i];
diff --git a/Engine/Source/Runtime/Experimental/Chaos/Public/Chaos/PBDSpringConstraintsBase.h b/Engine/Source/Runtime/Experimental/Chaos/Public/Chaos/PBDSpringConstraintsBase.h
--- a/Engine/Source/Runtime/Experimental/Chaos/Public/Chaos/PBDSpringConstraintsBase.h
+++ b/Engine/Source/Runtime/Experimental/Chaos/Public/Chaos/PBDSpringConstraintsBase.h
@@ -21,18 +21,21 @@ namespace Chaos
 		TPBDSpringConstraintsBase(const TDynamicParticles<T, d>& InParticles, TArray<TVector<int32, 2>>&& Constraints, const T Stiffness = (T)1)
 		    : MConstraints(MoveTemp(Constraints)), MStiffness(Stiffness)
 		{
+			RemoveDuplicateConstraints();
 			UpdateDistances(InParticles);
 		}
 		TPBDSpringConstraintsBase(const TDynamicParticles<T, d>& InParticles, const TArray<TVector<int32, 3>>& Constraints, const T Stiffness = (T)1)
 		    : MStiffness(Stiffness)
 		{
 			Init<3>(Constraints);
+			RemoveDuplicateConstraints();
 			UpdateDistances(InParticles);
 		}
 		TPBDSpringConstraintsBase(const TDynamicParticles<T, d>& InParticles, const TArray<TVector<int32, 4>>& Constraints, const T Stiffness = (T)1)
 		    : MStiffness(Stiffness)
 		{
 			Init<4>(Constraints);
+			RemoveDuplicateConstraints();
 			UpdateDistances(InParticles);
 		}
 		virtual ~TPBDSpringConstraintsBase()
@@ -73,10 +76,32 @@ namespace Chaos
 				}
 			}
 		}
+		
+		uint32 RemoveDuplicateConstraints()
+		{
+			const uint32 OriginalSize = MConstraints.Num();
+			TArray<TVector<int32, 2>> TrimmedConstraints;
+			TSet<TVector<int32, 2>>   ConstraintsAlreadyAdded;
+			TrimmedConstraints.Reserve(MConstraints.Num());
+			ConstraintsAlreadyAdded.Reserve(MConstraints.Num());
+			for (TVector<int32, 2> Constraint : MConstraints)
+			{
+				if (Constraint[0] > Constraint[1])
+				{
+					Swap(Constraint[0], Constraint[1]);
+				}
+				if (!ConstraintsAlreadyAdded.Contains(Constraint))
+				{
+					TrimmedConstraints.Add(Constraint);
+					ConstraintsAlreadyAdded.Add(Constraint);
+				}
+			}
+			MConstraints = MoveTemp(TrimmedConstraints);
+			return OriginalSize - MConstraints.Num();
+		}
 
 		template<class T_PARTICLES>
-		void
-		UpdateDistances(const T_PARTICLES& InParticles)
+		void UpdateDistances(const T_PARTICLES& InParticles)
 		{
 			MDists.Reset();
 			MDists.Reserve(MConstraints.Num());
@@ -91,8 +116,7 @@ namespace Chaos
 		}
 
 		template<class T_PARTICLES>
-		TVector<T, d>
-		GetDelta(const T_PARTICLES& InParticles, const int32 i) const
+		inline TVector<T, d> GetDelta(const T_PARTICLES& InParticles, const int32 i) const
 		{
 			const auto& Constraint = MConstraints[i];
 			const int32 i1 = Constraint[0];
@@ -115,8 +139,7 @@ namespace Chaos
 		// the dynamic particle positions prior to normalizing. Use this if you happen
 		// to know that the particle positions aren't coincident.
 		template<class T_PARTICLES>
-		TVector<T, d>
-		GetUnsafeDelta(const T_PARTICLES& InParticles, const int32 i) const
+		inline TVector<T, d> GetUnsafeDelta(const T_PARTICLES& InParticles, const int32 i) const
 		{
 			const auto& Constraint = MConstraints[i];
 			const int32 i1 = Constraint[0];
diff --git a/Engine/Source/Runtime/Experimental/Chaos/Public/Chaos/PerParticlePBDLongRangeConstraints.h b/Engine/Source/Runtime/Experimental/Chaos/Public/Chaos/PerParticlePBDLongRangeConstraints.h
--- a/Engine/Source/Runtime/Experimental/Chaos/Public/Chaos/PerParticlePBDLongRangeConstraints.h
+++ b/Engine/Source/Runtime/Experimental/Chaos/Public/Chaos/PerParticlePBDLongRangeConstraints.h
@@ -4,6 +4,9 @@
 #include "Chaos/PBDLongRangeConstraintsBase.h"
 #include "Chaos/PBDParticles.h"
 #include "Chaos/PerParticleRule.h"
+#include "ChaosStats.h"
+
+DECLARE_CYCLE_STAT(TEXT("Chaos PBD Long Range Per Particle Constraint"), STAT_PBD_LongRangePerParticle, STATGROUP_Chaos);
 
 namespace Chaos
 {
@@ -31,8 +34,8 @@ class TPerParticlePBDLongRangeConstraints : public TPerParticleRule<T, d>, publi
 	{
 		for (int32 i = 0; i < MParticleToConstraints[Index].Num(); ++i)
 		{
-			const auto CIndex = MParticleToConstraints[Index][i];
-			const auto& Constraint = MConstraints[CIndex];
+			const int32 CIndex = MParticleToConstraints[Index][i];  // Cache misses all the time
+			const TArray<uint32>& Constraint = MConstraints[CIndex];
 			check(Index == Constraint[Constraint.Num() - 1]);
 			check(InParticles.InvM(Index) > 0);
 			InParticles.P(Index) += Base::GetDelta(InParticles, CIndex);
@@ -41,7 +44,8 @@ class TPerParticlePBDLongRangeConstraints : public TPerParticleRule<T, d>, publi
 
 	void Apply(TPBDParticles<T, d>& InParticles, const T Dt) const override //-V762
 	{
-		PhysicsParallelFor(MParticleToConstraints.Num(), [&](int32 Index) {
+		SCOPE_CYCLE_COUNTER(STAT_PBD_LongRangePerParticle);
+		PhysicsParallelFor(MParticleToConstraints.Num(), [this, &InParticles, Dt](int32 Index) {
 			Apply(InParticles, Dt, Index);
 		});
 	}
