Fix CIS

diff --git a/Engine/Plugins/Experimental/ChaosCloth/Source/Chaos/Private/ChaosCloth/ChaosClothingSimulation.cpp b/Engine/Plugins/Experimental/ChaosCloth/Source/Chaos/Private/ChaosCloth/ChaosClothingSimulation.cpp
--- a/Engine/Plugins/Experimental/ChaosCloth/Source/Chaos/Private/ChaosCloth/ChaosClothingSimulation.cpp
+++ b/Engine/Plugins/Experimental/ChaosCloth/Source/Chaos/Private/ChaosCloth/ChaosClothingSimulation.cpp
@@ -1397,7 +1397,7 @@ void ClothingSimulation::DebugDrawCollision(USkeletalMeshComponent* OwnerCompone
 		DrawWireSphere(PDI, Transform, Color, Radius, 12, SDPG_World, 0.0f, 0.001f, false);
 	};
 
-	auto DrawBox = [&PDI](const Chaos::TAABB<float, 3>& Box, const TRotation<float, 3>& Rotation, const Chaos::TVector<float, 3>& Position, const FLinearColor& Color)
+	auto DrawBox = [&PDI](const Chaos::TBox<float, 3>& Box, const TRotation<float, 3>& Rotation, const Chaos::TVector<float, 3>& Position, const FLinearColor& Color)
 	{
 		const FMatrix BoxToWorld = FTransform(Rotation, Position).ToMatrixNoScale();
 		const FVector Radii = Box.Extents() * 0.5f;
