Fix CIS issue with cloth weight map target enums.

diff --git a/Engine/Plugins/Experimental/ChaosCloth/Source/Chaos/Private/ChaosCloth/ChaosWeightMapTarget.h b/Engine/Plugins/Experimental/ChaosCloth/Source/Chaos/Private/ChaosCloth/ChaosWeightMapTarget.h
--- a/Engine/Plugins/Experimental/ChaosCloth/Source/Chaos/Private/ChaosCloth/ChaosWeightMapTarget.h
+++ b/Engine/Plugins/Experimental/ChaosCloth/Source/Chaos/Private/ChaosCloth/ChaosWeightMapTarget.h
@@ -9,10 +9,10 @@
 UENUM()
 enum class EChaosWeightMapTarget : uint8
 {
-	None                = EWeightMapTargetCommon::None,
-	MaxDistance         = EWeightMapTargetCommon::MaxDistance,
-	BackstopDistance    = EWeightMapTargetCommon::BackstopDistance,
-	BackstopRadius      = EWeightMapTargetCommon::BackstopRadius,
-	AnimDriveMultiplier = EWeightMapTargetCommon::AnimDriveMultiplier
+	None                = (uint8)EWeightMapTargetCommon::None,
+	MaxDistance         = (uint8)EWeightMapTargetCommon::MaxDistance,
+	BackstopDistance    = (uint8)EWeightMapTargetCommon::BackstopDistance,
+	BackstopRadius      = (uint8)EWeightMapTargetCommon::BackstopRadius,
+	AnimDriveMultiplier = (uint8)EWeightMapTargetCommon::AnimDriveMultiplier
 	// Add Chaos specific maps below this line
 };
