Fix an Object is not packaged error, when creating the Chaos Cloth shared config due to the config trying to link itself to an empty asset.

diff --git a/Engine/Plugins/Experimental/ChaosCloth/Source/Chaos/Private/ChaosCloth/ChaosClothingSimulation.cpp b/Engine/Plugins/Experimental/ChaosCloth/Source/Chaos/Private/ChaosCloth/ChaosClothingSimulation.cpp
--- a/Engine/Plugins/Experimental/ChaosCloth/Source/Chaos/Private/ChaosCloth/ChaosClothingSimulation.cpp
+++ b/Engine/Plugins/Experimental/ChaosCloth/Source/Chaos/Private/ChaosCloth/ChaosClothingSimulation.cpp
@@ -502,19 +502,19 @@ void ClothingSimulation::PostActorCreationInitialize()
 {
 	if (Assets.Num() > 0)
 	{
-		// If we don't have a shared sim config, create one. This is only possible if none of the cloth Assets had a configuration during Actor creation
-		if (ClothSharedSimConfig == nullptr)
-		{
-			// None of the cloth assets had a clothSharedSimConfig, so we will create it
-			ClothSharedSimConfig = NewObject<UChaosClothSharedSimConfig>(Assets[0], UChaosClothSharedSimConfig::StaticClass()->GetFName());
-			check(ClothSharedSimConfig);
-		}
-
 		// Let all assets point to the same shared configuration
 		for (int32 AssetIndex = 0; AssetIndex < Assets.Num(); ++AssetIndex)
 		{
 			if (UClothingAssetCommon* const Asset = Assets[AssetIndex])
 			{
+				// If we don't have a shared sim config, create one. This is only possible if none of the cloth Assets had a configuration during Actor creation
+				if (ClothSharedSimConfig == nullptr)
+				{
+					// None of the cloth assets had a clothSharedSimConfig, so we will create it and assign it to the first available asset
+					ClothSharedSimConfig = NewObject<UChaosClothSharedSimConfig>(Assets[AssetIndex], UChaosClothSharedSimConfig::StaticClass()->GetFName());
+					check(ClothSharedSimConfig);
+				}
+
 				const UChaosClothSharedSimConfig* const AssetClothConfigShared = Asset->GetClothConfig<UChaosClothSharedSimConfig>();
 				if (!AssetClothConfigShared || AssetClothConfigShared != ClothSharedSimConfig)
 				{
