Chaos CLoth: Refactor CreateActor to enable new functionaliy like updating cloth config without restarting the simulation

diff --git a/Engine/Plugins/Experimental/ChaosCloth/Source/Chaos/Private/ChaosCloth/ChaosClothingSimulation.cpp b/Engine/Plugins/Experimental/ChaosCloth/Source/Chaos/Private/ChaosCloth/ChaosClothingSimulation.cpp
--- a/Engine/Plugins/Experimental/ChaosCloth/Source/Chaos/Private/ChaosCloth/ChaosClothingSimulation.cpp
+++ b/Engine/Plugins/Experimental/ChaosCloth/Source/Chaos/Private/ChaosCloth/ChaosClothingSimulation.cpp
@@ -177,36 +177,29 @@
 		AnimDriveSpringStiffness.SetNumZeroed(InSimDataIndex + 1);
 	}
 	Assets[InSimDataIndex] = Asset;
-	AnimDriveSpringStiffness[InSimDataIndex] = ChaosClothSimConfig->AnimDriveSpringStiffness;
 
+	if (Meshes.Num() <= InSimDataIndex)
+	{
+		Meshes.SetNum(InSimDataIndex + 1);
+		FaceNormals.SetNum(InSimDataIndex + 1);
+		PointNormals.SetNum(InSimDataIndex + 1);
+	}
+
 	check(Asset->GetNumLods() == 1);
 	const UClothLODDataCommon* const AssetLodData = Asset->ClothLodData[0];
 	const FClothPhysicalMeshData& PhysMesh = AssetLodData->ClothPhysicalMeshData;
 
-	// SkinPhysicsMesh() strips scale from RootBoneTransform ("Ignore any user scale.
-	// It's already accounted for in our skinning matrices."), and returns all points
-	// in that space.
-	TArray<Chaos::TVector<float, 3>> TempAnimationPositions;
-	TArray<Chaos::TVector<float, 3>> TempAnimationNormals;
-
-	TempAnimationPositions.SetNumUninitialized(PhysMesh.Vertices.Num());
-	TempAnimationNormals.SetNumUninitialized(PhysMesh.Vertices.Num());
-
-	FTransform RootBoneTransform = Context.BoneTransforms[Asset->ReferenceBoneIndex];
-	ClothingMeshUtils::SkinPhysicsMesh<true, false>(
-		Asset->UsedBoneIndices,
-		PhysMesh, // curr pos and norm
-		Context.ComponentToWorld,
-		Context.RefToLocals.GetData(),
-		Context.RefToLocals.Num(),
-		reinterpret_cast<TArray<FVector>&>(TempAnimationPositions),
-		reinterpret_cast<TArray<FVector>&>(TempAnimationNormals));
-
 	// Add particles
 	TPBDParticles<float, 3>& Particles = Evolution->Particles();
 	const uint32 Offset = Particles.Size();
 	Particles.AddParticles(PhysMesh.Vertices.Num());
 
+	if (IndexToRangeMap.Num() <= InSimDataIndex)
+	{
+		IndexToRangeMap.SetNum(InSimDataIndex + 1);
+	}
+	IndexToRangeMap[InSimDataIndex] = Chaos::TVector<uint32, 2>(Offset, Particles.Size());
+
 	// ClothSharedSimConfig should either be a nullptr, or point to an object common to the whole skeletal mesh
 	UChaosClothSharedSimConfig* const AssetClothConfigShared = Asset->GetClothConfig<UChaosClothSharedSimConfig>();
 	if (ClothSharedSimConfig == nullptr)
@@ -216,7 +209,7 @@
 	else if (!AssetClothConfigShared || AssetClothConfigShared != ClothSharedSimConfig)
 	{
 		UE_CLOG(AssetClothConfigShared,
-			LogChaosCloth, Warning, 
+			LogChaosCloth, Warning,
 			TEXT("Found a different Chaos Cloth shared config in the Cloth LOD asset being added to %s in sim slot %d. ")
 			TEXT("The asset's shared config will have to be replaced."),
 			InOwnerComponent->GetOwner() ? *InOwnerComponent->GetOwner()->GetName() : TEXT("None"), InSimDataIndex);
@@ -227,14 +220,18 @@
 	AnimationPositions.SetNum(Particles.Size());
 	AnimationNormals.SetNum(Particles.Size());
 
-	if (IndexToRangeMap.Num() <= InSimDataIndex)
-		IndexToRangeMap.SetNum(InSimDataIndex + 1);
-	IndexToRangeMap[InSimDataIndex] = Chaos::TVector<uint32, 2>(Offset, Particles.Size());
+	ClothingMeshUtils::SkinPhysicsMesh<true, false>(
+		Asset->UsedBoneIndices,
+		PhysMesh, // curr pos and norm
+		Context.ComponentToWorld,
+		Context.RefToLocals.GetData(),
+		Context.RefToLocals.Num(),
+		reinterpret_cast<TArray<FVector>&>(AnimationPositions),
+		reinterpret_cast<TArray<FVector>&>(AnimationNormals),
+		Offset);
 
 	for (uint32 i = Offset; i < Particles.Size(); ++i)
 	{
-		AnimationPositions[i] = TempAnimationPositions[i - Offset];
-		AnimationNormals[i] = TempAnimationNormals[i - Offset];
 		Particles.X(i) = AnimationPositions[i];
 		Particles.V(i) = Chaos::TVector<float, 3>(0.f, 0.f, 0.f);
 		// Initialize mass to 0, to be overridden later
@@ -243,31 +240,110 @@
 
 	OldAnimationPositions = AnimationPositions;  // Also update the old positions array to avoid any interpolation issues
 
-	const int32 NumTriangles = PhysMesh.Indices.Num() / 3;
+	TUniquePtr<Chaos::TTriangleMesh<float>>& Mesh = Meshes[InSimDataIndex];
+	BuildMesh(Mesh, PhysMesh, Offset);
+
+	const FPointWeightMap& MaxDistances = PhysMesh.GetWeightMap(EChaosWeightMapTarget::MaxDistance);
+	SetParticleMasses(ChaosClothSimConfig, Particles, Offset, PhysMesh.Vertices.Num(), *Mesh, MaxDistances);
+
+	AddConstraints(ChaosClothSimConfig, PhysMesh, *Mesh, Offset, InSimDataIndex);
+
+	// Add Self Collisions
+	if (ChaosClothSimConfig->bUseSelfCollisions)
+	{
+		AddSelfCollisions(Mesh, Particles, Offset);
+	}
+
+	// Collision particles
+	TGeometryClothParticles<float, 3>& CollisionParticles = Evolution->CollisionParticles();
+
+	// Pull collisions from the specified physics asset inside the clothing asset
+	ExtractPhysicsAssetCollisions(Asset);
+
+	// Extract the legacy Apex collision from the clothing asset
+	ExtractLegacyAssetCollisions(Asset, InOwnerComponent);
+
+	// Set the external collision starting index
+	checkf(ExternalCollisions.Spheres.Num() == 0 &&
+		ExternalCollisions.SphereConnections.Num() == 0 &&
+		ExternalCollisions.Convexes.Num() == 0 &&
+		ExternalCollisions.Boxes.Num() == 0, TEXT("There cannot be any external collisions added before all the cloth assets collisions are processed."));
+	ExternalCollisionsOffset = CollisionParticles.Size();
+}
+
+void ClothingSimulation::PostActorCreationInitialize()
+{
+	if (Assets.Num() > 0)
+	{
+		// Let all assets point to the same shared configuration
+		for (int32 AssetIndex = 0; AssetIndex < Assets.Num(); ++AssetIndex)
+		{
+			if (UClothingAssetCommon* const Asset = Assets[AssetIndex])
+			{
+				// If we don't have a shared sim config, create one. This is only possible if none of the cloth Assets had a configuration during Actor creation
+				if (ClothSharedSimConfig == nullptr)
+				{
+					// None of the cloth assets had a clothSharedSimConfig, so we will create it and assign it to the first available asset
+					ClothSharedSimConfig = NewObject<UChaosClothSharedSimConfig>(Assets[AssetIndex], UChaosClothSharedSimConfig::StaticClass()->GetFName());
+					check(ClothSharedSimConfig);
+				}
+
+				const UChaosClothSharedSimConfig* const AssetClothConfigShared = Asset->GetClothConfig<UChaosClothSharedSimConfig>();
+				if (!AssetClothConfigShared || AssetClothConfigShared != ClothSharedSimConfig)
+				{
+					UE_CLOG(AssetClothConfigShared,
+						LogChaosCloth, Warning,
+						TEXT("Found a different Chaos Cloth shared config in Cloth LOD asset %d. ")
+						TEXT("The asset's shared config will have to be replaced."),
+						AssetIndex);
+
+					Asset->SetClothConfig(ClothSharedSimConfig);
+				}
+			}
+		}
+	}
+
+	if (ClothSharedSimConfig) // ClothSharedSimConfig will be a null pointer if all cloth instances are disabled in which case we will use default Evolution parameters
+	{
+		// Now set all the common parameters on the simulation
+		Evolution->SetIterations(ClothSharedSimConfig->IterationCount);
+		Evolution->SetSelfCollisionThickness(ClothSharedSimConfig->SelfCollisionThickness);
+		Evolution->SetCollisionThickness(ClothSharedSimConfig->CollisionThickness);
+		Evolution->SetDamping(ClothSharedSimConfig->Damping);
+		Evolution->GetGravityForces().SetAcceleration(Chaos::TVector<float, 3>(ClothSharedSimConfig->Gravity));
+	}
+}
+
+void ClothingSimulation::BuildMesh(TUniquePtr<Chaos::TTriangleMesh<float>>& OutChaosMesh, const FClothPhysicalMeshData& InPhysMesh, uint32 Offset) const
+{
+	const int32 NumTriangles = InPhysMesh.Indices.Num() / 3;
 	TArray<Chaos::TVector<int32, 3>> InputSurfaceElements;
 	InputSurfaceElements.Reserve(NumTriangles);
 	for (int i = 0; i < NumTriangles; ++i)
 	{
 		const int32 Index = 3 * i;
 		InputSurfaceElements.Add(
-			{ static_cast<int32>(Offset + PhysMesh.Indices[Index]),
-			 static_cast<int32>(Offset + PhysMesh.Indices[Index + 1]),
-			 static_cast<int32>(Offset + PhysMesh.Indices[Index + 2]) });
+			{ static_cast<int32>(Offset + InPhysMesh.Indices[Index]),
+			 static_cast<int32>(Offset + InPhysMesh.Indices[Index + 1]),
+			 static_cast<int32>(Offset + InPhysMesh.Indices[Index + 2]) });
 	}
 	check(InputSurfaceElements.Num() == NumTriangles);
-	if (Meshes.Num() <= InSimDataIndex)
-	{
-		Meshes.SetNum(InSimDataIndex + 1);
-		FaceNormals.SetNum(InSimDataIndex + 1);
-		PointNormals.SetNum(InSimDataIndex + 1);
-	}
-	TUniquePtr<Chaos::TTriangleMesh<float>>& Mesh = Meshes[InSimDataIndex];
-	Mesh.Reset(new Chaos::TTriangleMesh<float>(MoveTemp(InputSurfaceElements)));
-	check(Mesh->GetNumElements() == NumTriangles);
-	const auto& SurfaceElements = Mesh->GetSurfaceElements();
-	Mesh->GetPointToTriangleMap(); // Builds map for later use by GetPointNormals().
+	OutChaosMesh.Reset(new Chaos::TTriangleMesh<float>(MoveTemp(InputSurfaceElements)));
+	check(OutChaosMesh->GetNumElements() == NumTriangles);
+	OutChaosMesh->GetPointToTriangleMap(); // Builds map for later use by GetPointNormals()
+}
 
+void ClothingSimulation::SetParticleMasses(
+	const UChaosClothConfig* ChaosClothSimConfig,
+	TPBDParticles<float, 3>& Particles,
+	uint32 Offset, // Offset into Particles array
+	uint32 ParticlesCount, // Number of particles
+	Chaos::TTriangleMesh<float>& Mesh, // Mesh with indices into Particles with offset already taken into account
+	const FPointWeightMap& MaxDistances
+	) const
+{
 	// Assign per particle mass proportional to connected area.
+	const TArray<TVector<int32, 3>>& SurfaceElements = Mesh.GetSurfaceElements();
 	float TotalArea = 0.0;
 	for (const Chaos::TVector<int32, 3>& Tri : SurfaceElements)
 	{
@@ -280,7 +356,7 @@
 		Particles.M(Tri[1]) += ThirdTriArea;
 		Particles.M(Tri[2]) += ThirdTriArea;
 	}
-	const TSet<int32> Vertices = Mesh->GetVertices();
+	const TSet<int32> Vertices = Mesh.GetVertices();
 	switch (ChaosClothSimConfig->MassMode)
 	{
 	case EClothMassMode::UniformMass:
@@ -305,15 +381,20 @@
 		}
 		break;
 	};
+
 	// Clamp and enslave
-	const FPointWeightMap& MaxDistances = PhysMesh.GetWeightMap(EChaosWeightMapTarget::MaxDistance);
-	for (uint32 i = Offset; i < Particles.Size(); i++)
+	check(Particles.Size() >= Offset + ParticlesCount);
+	for (uint32 i = Offset; i < Offset + ParticlesCount; i++)
 	{
 		Particles.M(i) = FMath::Max(Particles.M(i), ChaosClothSimConfig->MinPerParticleMass);
 		Particles.InvM(i) = MaxDistances.IsBelowThreshold(i - Offset) ? 0.f : 1.f / Particles.M(i);
 	}
+}
 
-	// Add Model
+void ClothingSimulation::AddConstraints(const UChaosClothConfig* ChaosClothSimConfig, const FClothPhysicalMeshData& PhysMesh, const Chaos::TTriangleMesh<float>& Mesh, uint32 Offset, int32 InSimDataIndex)
+{
+	const TArray<TVector<int32, 3>>& SurfaceElements = Mesh.GetSurfaceElements();
+
 	if (ChaosClothSimConfig->ShapeTargetStiffness)
 	{
 		check(ChaosClothSimConfig->ShapeTargetStiffness > 0.f && ChaosClothSimConfig->ShapeTargetStiffness <= 1.f);
@@ -333,14 +414,14 @@
 		check(ChaosClothSimConfig->BendingStiffness > 0.f && ChaosClothSimConfig->BendingStiffness <= 1.f);
 		if (ChaosClothSimConfig->bUseBendingElements)
 		{
-			TArray<Chaos::TVector<int32, 4>> BendingConstraints = Mesh->GetUniqueAdjacentElements();
+			TArray<Chaos::TVector<int32, 4>> BendingConstraints = Mesh.GetUniqueAdjacentElements();
 			Evolution->AddPBDConstraintFunction([BendConstraints = Chaos::TPBDBendingConstraints<float>(Evolution->Particles(), MoveTemp(BendingConstraints))](TPBDParticles<float, 3>& InParticles, const float Dt) {
 				BendConstraints.Apply(InParticles, Dt);
 			});
 		}
 		else
 		{
-			TArray<Chaos::TVector<int32, 2>> BendingConstraints = Mesh->GetUniqueAdjacentPoints();
+			TArray<Chaos::TVector<int32, 2>> BendingConstraints = Mesh.GetUniqueAdjacentPoints();
 			Evolution->AddPBDConstraintFunction([SpringConstraints = Chaos::TPBDSpringConstraints<float, 3>(Evolution->Particles(), MoveTemp(BendingConstraints), ChaosClothSimConfig->BendingStiffness)](TPBDParticles<float, 3>& InParticles, const float Dt) {
 				SpringConstraints.Apply(InParticles, Dt);
 			});
@@ -363,7 +444,7 @@
 		}
 		else if (ChaosClothSimConfig->bUseThinShellVolumeConstraints)
 		{
-			TArray<Chaos::TVector<int32, 2>> BendingConstraints = Mesh->GetUniqueAdjacentPoints();
+			TArray<Chaos::TVector<int32, 2>> BendingConstraints = Mesh.GetUniqueAdjacentPoints();
 			TArray<Chaos::TVector<int32, 2>> DoubleBendingConstraints;
 			{
 				TMap<int32, TArray<int32>> BendingHash;
@@ -408,12 +489,12 @@
 	}
 	if (ChaosClothSimConfig->StrainLimitingStiffness)
 	{
-		check(Mesh->GetNumElements() > 0);
+		check(Mesh.GetNumElements() > 0);
 		// PerFormance note: The Per constraint version of this function is quite a bit faster for smaller assets
 		// There might be a cross-over point where the PerParticle version is faster: To be determined
 		Chaos::TPBDLongRangeConstraints<float, 3> PBDLongRangeConstraints(
 			Evolution->Particles(),
-			Mesh->GetPointToNeighborsMap(),
+			Mesh.GetPointToNeighborsMap(),
 			10, // The max number of connected neighbors per particle.  ryan - What should this be?  Was k...
 			ChaosClothSimConfig->StrainLimitingStiffness);
 
@@ -424,9 +505,10 @@
 	}
 
 	// Maximum Distance Constraints
+	const FPointWeightMap& MaxDistances = PhysMesh.GetWeightMap(EChaosWeightMapTarget::MaxDistance);
 	if (MaxDistances.Num() > 0)
 	{
-		check(Mesh->GetNumElements() > 0);
+		check(Mesh.GetNumElements() > 0);
 		Chaos::PBDSphericalConstraint<float, 3> SphericalContraint(Offset, MaxDistances.Num(), true, &AnimationPositions, &MaxDistances.Values);
 		Evolution->AddPBDConstraintFunction([SphericalContraint = MoveTemp(SphericalContraint)](TPBDParticles<float, 3>& InParticles, const float Dt)
 		{
@@ -439,7 +521,7 @@
 	const FPointWeightMap& BackstopDistances = PhysMesh.GetWeightMap(EChaosWeightMapTarget::BackstopDistance);
 	if (BackstopRadiuses.Num() > 0 && BackstopDistances.Num() > 0)
 	{
-		check(Mesh->GetNumElements() > 0);
+		check(Mesh.GetNumElements() > 0);
 		check(BackstopRadiuses.Num() == BackstopDistances.Num());
 
 		Chaos::PBDSphericalConstraint<float, 3> SphericalContraint(Offset, BackstopRadiuses.Num(), false, &AnimationPositions,
@@ -451,10 +533,11 @@
 	}
 
 	// Animation Drive Constraints
+	AnimDriveSpringStiffness[InSimDataIndex] = ChaosClothSimConfig->AnimDriveSpringStiffness;
 	const FPointWeightMap& AnimDriveMultipliers = PhysMesh.GetWeightMap(EChaosWeightMapTarget::AnimDriveMultiplier);
 	if (AnimDriveMultipliers.Num() > 0)
 	{
-		check(Mesh->GetNumElements() > 0);
+		check(Mesh.GetNumElements() > 0);
 		TPBDAnimDriveConstraint<float, 3> PBDAnimDriveConstraint(Offset, &AnimationPositions, &AnimDriveMultipliers.Values, AnimDriveSpringStiffness[InSimDataIndex]);
 		Evolution->AddPBDConstraintFunction(
 			[PBDAnimDriveConstraint = MoveTemp(PBDAnimDriveConstraint), &Stiffness = AnimDriveSpringStiffness[InSimDataIndex]](TPBDParticles<float, 3>& InParticles, const float Dt) mutable
@@ -463,82 +546,23 @@
 			PBDAnimDriveConstraint.Apply(InParticles, Dt);
 		});
 	}
-
-    // Add Self Collisions
-    if (ChaosClothSimConfig->bUseSelfCollisions)
-    {
-        // TODO(mlentine): Parallelize these for multiple meshes
-        Evolution->CollisionTriangles().Append(SurfaceElements);
-        for (uint32 i = Offset; i < Particles.Size(); ++i)
-        {
-            auto Neighbors = Mesh->GetNRing(i, 5);
-            for (const auto& Element : Neighbors)
-            {
-                check(i != Element);
-                Evolution->DisabledCollisionElements().Add(Chaos::TVector<int32, 2>(i, Element));
-                Evolution->DisabledCollisionElements().Add(Chaos::TVector<int32, 2>(Element, i));
-            }
-        }
-    }
-
-	// Collision particles
-	TGeometryClothParticles<float, 3>& CollisionParticles = Evolution->CollisionParticles();
-
-	// Pull collisions from the specified physics asset inside the clothing asset
-	ExtractPhysicsAssetCollisions(Asset);
-
-	// Extract the legacy Apex collision from the clothing asset
-	ExtractLegacyAssetCollisions(Asset, InOwnerComponent);
-
-	// Set the external collision starting index
-	checkf(ExternalCollisions.Spheres.Num() == 0 &&
-		ExternalCollisions.SphereConnections.Num() == 0 &&
-		ExternalCollisions.Convexes.Num() == 0 &&
-		ExternalCollisions.Boxes.Num() == 0, TEXT("There cannot be any external collisions added before all the cloth assets collisions are processed."));
-	ExternalCollisionsOffset = CollisionParticles.Size();
 }
 
-void ClothingSimulation::PostActorCreationInitialize()
+void ClothingSimulation::AddSelfCollisions(const TUniquePtr<Chaos::TTriangleMesh<float>>& Mesh, const TPBDParticles<float, 3>& Particles, int32 Offset)
 {
-	if (Assets.Num() > 0)
+	// TODO(mlentine): Parallelize these for multiple meshes
+	const TArray<TVector<int32, 3>>& SurfaceElements = Mesh->GetSurfaceElements();
+	Evolution->CollisionTriangles().Append(SurfaceElements);
+	for (uint32 i = Offset; i < Particles.Size(); ++i)
 	{
-		// Let all assets point to the same shared configuration
-		for (int32 AssetIndex = 0; AssetIndex < Assets.Num(); ++AssetIndex)
+		TSet<int32> Neighbors = Mesh->GetNRing(i, 5);
+		for (int32 Element : Neighbors)
 		{
-			if (UClothingAssetCommon* const Asset = Assets[AssetIndex])
-			{
-				// If we don't have a shared sim config, create one. This is only possible if none of the cloth Assets had a configuration during Actor creation
-				if (ClothSharedSimConfig == nullptr)
-				{
-					// None of the cloth assets had a clothSharedSimConfig, so we will create it and assign it to the first available asset
-					ClothSharedSimConfig = NewObject<UChaosClothSharedSimConfig>(Assets[AssetIndex], UChaosClothSharedSimConfig::StaticClass()->GetFName());
-					check(ClothSharedSimConfig);
-				}
-
-				const UChaosClothSharedSimConfig* const AssetClothConfigShared = Asset->GetClothConfig<UChaosClothSharedSimConfig>();
-				if (!AssetClothConfigShared || AssetClothConfigShared != ClothSharedSimConfig)
-				{
-					UE_CLOG(AssetClothConfigShared,
-						LogChaosCloth, Warning,
-						TEXT("Found a different Chaos Cloth shared config in Cloth LOD asset %d. ")
-						TEXT("The asset's shared config will have to be replaced."),
-						AssetIndex);
-
-					Asset->SetClothConfig(ClothSharedSimConfig);
-				}
-			}
+			check(i != Element);
+			Evolution->DisabledCollisionElements().Add(Chaos::TVector<int32, 2>(i, Element));
+			Evolution->DisabledCollisionElements().Add(Chaos::TVector<int32, 2>(Element, i));
 		}
 	}
-
-	if (ClothSharedSimConfig) // ClothSharedSimConfig will be a null pointer if all cloth instances are disabled in which case we will use default Evolution parameters
-	{
-		// Now set all the common parameters on the simulation
-		Evolution->SetIterations(ClothSharedSimConfig->IterationCount);
-		Evolution->SetSelfCollisionThickness(ClothSharedSimConfig->SelfCollisionThickness);
-		Evolution->SetCollisionThickness(ClothSharedSimConfig->CollisionThickness);
-		Evolution->SetDamping(ClothSharedSimConfig->Damping);
-		Evolution->GetGravityForces().SetAcceleration(Chaos::TVector<float, 3>(ClothSharedSimConfig->Gravity));
-	}
 }
 
 void ClothingSimulation::UpdateCollisionTransforms(const ClothingSimulationContext& Context)
diff --git a/Engine/Plugins/Experimental/ChaosCloth/Source/Chaos/Public/ChaosCloth/ChaosClothingSimulation.h b/Engine/Plugins/Experimental/ChaosCloth/Source/Chaos/Public/ChaosCloth/ChaosClothingSimulation.h
--- a/Engine/Plugins/Experimental/ChaosCloth/Source/Chaos/Public/ChaosCloth/ChaosClothingSimulation.h
+++ b/Engine/Plugins/Experimental/ChaosCloth/Source/Chaos/Public/ChaosCloth/ChaosClothingSimulation.h
@@ -79,6 +79,22 @@ namespace Chaos
 		// End of IClothingSimulation interface
 
 	private:
+
+		void BuildMesh(TUniquePtr<Chaos::TTriangleMesh<float>>& OutChaosMesh, const FClothPhysicalMeshData& InPhysMesh, uint32 Offset) const;
+
+		void SetParticleMasses(
+			const UChaosClothConfig* ChaosClothSimConfig,
+			TPBDParticles<float, 3>& Particles,
+			uint32 Offset, // Offset into Particles array
+			uint32 ParticlesCount, // Number of particles
+			Chaos::TTriangleMesh<float>& Mesh, // Mesh with indices into Particles with offset already taken into account
+			const FPointWeightMap& MaxDistances
+			) const;
+
+		void AddConstraints(const UChaosClothConfig* ChaosClothSimConfig, const FClothPhysicalMeshData& PhysMesh, const Chaos::TTriangleMesh<float>& Mesh, uint32 Offset, int32 InSimDataIndex);
+
+		void AddSelfCollisions(const TUniquePtr<Chaos::TTriangleMesh<float>>& Mesh, const TPBDParticles<float, 3>& Particles, int32 Offset);
+
 		// Extract the collisions from the physics asset referenced by the specified clothing asset
 		void ExtractPhysicsAssetCollisions(UClothingAssetCommon* Asset);
 		// Extract the collisions from the specified clothing asset 
diff --git a/Engine/Source/Runtime/Experimental/Chaos/Private/Chaos/TriangleMesh.cpp b/Engine/Source/Runtime/Experimental/Chaos/Private/Chaos/TriangleMesh.cpp
--- a/Engine/Source/Runtime/Experimental/Chaos/Private/Chaos/TriangleMesh.cpp
+++ b/Engine/Source/Runtime/Experimental/Chaos/Private/Chaos/TriangleMesh.cpp
@@ -108,7 +108,7 @@ void TTriangleMesh<T>::GetVertexSet(TSet<int32>& VertexSet) const
 }
 
 template<class T>
-const TMap<int32, TSet<uint32>>& TTriangleMesh<T>::GetPointToNeighborsMap()
+const TMap<int32, TSet<uint32>>& TTriangleMesh<T>::GetPointToNeighborsMap() const
 {
 	if (MPointToNeighborsMap.Num())
 	{
@@ -136,7 +136,7 @@ const TMap<int32, TSet<uint32>>& TTriangleMesh<T>::GetPointToNeighborsMap()
 }
 
 template<class T>
-const TMap<int32, TArray<int32>>& TTriangleMesh<T>::GetPointToTriangleMap()
+const TMap<int32, TArray<int32>>& TTriangleMesh<T>::GetPointToTriangleMap() const
 {
 	if (MPointToTriangleMap.Num())
 	{
@@ -154,7 +154,7 @@ const TMap<int32, TArray<int32>>& TTriangleMesh<T>::GetPointToTriangleMap()
 }
 
 template<class T>
-TArray<TVector<int32, 2>> TTriangleMesh<T>::GetUniqueAdjacentPoints()
+TArray<TVector<int32, 2>> TTriangleMesh<T>::GetUniqueAdjacentPoints() const
 {
 	TArray<TVector<int32, 2>> BendingConstraints;
 	auto BendingElements = GetUniqueAdjacentElements();
@@ -166,7 +166,7 @@ TArray<TVector<int32, 2>> TTriangleMesh<T>::GetUniqueAdjacentPoints()
 }
 
 template<class T>
-TArray<TVector<int32, 4>> TTriangleMesh<T>::GetUniqueAdjacentElements()
+TArray<TVector<int32, 4>> TTriangleMesh<T>::GetUniqueAdjacentElements() const
 {
 	TArray<TVector<int32, 4>> BendingConstraints;
 	TSet<TArray<int32>> BendingElements;
diff --git a/Engine/Source/Runtime/Experimental/Chaos/Public/Chaos/TriangleMesh.h b/Engine/Source/Runtime/Experimental/Chaos/Public/Chaos/TriangleMesh.h
--- a/Engine/Source/Runtime/Experimental/Chaos/Public/Chaos/TriangleMesh.h
+++ b/Engine/Source/Runtime/Experimental/Chaos/Public/Chaos/TriangleMesh.h
@@ -80,10 +80,10 @@ namespace Chaos
 
 		int32 GetNumElements() const { return MElements.Num(); }
 
-		const TMap<int32, TSet<uint32>>& GetPointToNeighborsMap();
+		const TMap<int32, TSet<uint32>>& GetPointToNeighborsMap() const;
 		const TSet<uint32>& GetNeighbors(const int32 Element) { return GetPointToNeighborsMap()[Element]; }
 
-		const TMap<int32, TArray<int32>>& GetPointToTriangleMap();
+		const TMap<int32, TArray<int32>>& GetPointToTriangleMap() const;
 		const TArray<int32>& GetCoincidentTriangles(const int32 Element) { return GetPointToTriangleMap()[Element]; }
 
 		TSet<int32> GetNRing(const int32 Element, const int32 N)
@@ -123,8 +123,8 @@ namespace Chaos
 			return Neighbors;
 		}
 
-		TArray<Chaos::TVector<int32, 2>> GetUniqueAdjacentPoints();
-		TArray<Chaos::TVector<int32, 4>> GetUniqueAdjacentElements();
+		TArray<Chaos::TVector<int32, 2>> GetUniqueAdjacentPoints() const;
+		TArray<Chaos::TVector<int32, 4>> GetUniqueAdjacentElements() const;
 
 		/** The GetFaceNormals functions assume Counter Clockwise triangle windings in a Left Handed coordinate system
 			If this is not the case the returned face normals may be inverted
@@ -244,8 +244,8 @@ namespace Chaos
 
 		TArray<TVector<int32, 3>> MElements;
 
-		TMap<int32, TArray<int32>> MPointToTriangleMap;
-		TMap<int32, TSet<uint32>> MPointToNeighborsMap;
+		mutable TMap<int32, TArray<int32>> MPointToTriangleMap;
+		mutable TMap<int32, TSet<uint32>> MPointToNeighborsMap;
 
 		TSegmentMesh<T> MSegmentMesh;
 		TArray<TVector<int32, 3>> MFaceToEdges;
