Fix crash in Chaos Cloth when using a cloth asset with more than one LOD, and replace the offending check with a conditional log warning message.

diff --git a/Engine/Plugins/Experimental/ChaosCloth/Source/Chaos/Private/ChaosCloth/ChaosClothingSimulation.cpp b/Engine/Plugins/Experimental/ChaosCloth/Source/Chaos/Private/ChaosCloth/ChaosClothingSimulation.cpp
--- a/Engine/Plugins/Experimental/ChaosCloth/Source/Chaos/Private/ChaosCloth/ChaosClothingSimulation.cpp
+++ b/Engine/Plugins/Experimental/ChaosCloth/Source/Chaos/Private/ChaosCloth/ChaosClothingSimulation.cpp
@@ -185,7 +185,10 @@ void ClothingSimulation::CreateActor(USkeletalMeshComponent* InOwnerComponent, U
 		PointNormals.SetNum(InSimDataIndex + 1);
 	}
 
-	check(Asset->GetNumLods() == 1);
+	check(Asset->GetNumLods() > 0);
+	UE_CLOG(Asset->GetNumLods() != 1,
+		LogChaosCloth, Warning, TEXT("More than one LOD with the current cloth asset %s in sim slot %d. Only LOD 0 is supported with Chaos Cloth for now."),
+		InOwnerComponent->GetOwner() ? *InOwnerComponent->GetOwner()->GetName() : TEXT("None"), InSimDataIndex);
 	const UClothLODDataCommon* const AssetLodData = Asset->ClothLodData[0];
 	const FClothPhysicalMeshData& PhysMesh = AssetLodData->ClothPhysicalMeshData;
 
@@ -793,11 +796,6 @@ void ClothingSimulation::ExtractPhysicsAssetCollisions(UClothingAssetCommon* Ass
 
 void ClothingSimulation::ExtractLegacyAssetCollisions(UClothingAssetCommon* Asset, const USkeletalMeshComponent* InOwnerComponent)
 {
-	check(Asset->GetNumLods() > 0);
-	ensure(Asset->GetNumLods() == 1);
-	UE_CLOG(Asset->GetNumLods() != 1,
-		LogChaosCloth, Warning, TEXT("More than one LOD with the current cloth asset. Only LOD 0 is supported with the current system."));
-
 	if (const UClothLODDataCommon* const AssetLodData = Asset->ClothLodData[0])
 	{
 		const FClothCollisionData& LodCollData = AssetLodData->CollisionData;
