Fix crash in Chaos Cloth by replacing zero length tapered capsule collision volumes by spheres.

diff --git a/Engine/Plugins/Experimental/ChaosCloth/Source/Chaos/Private/ChaosCloth/ChaosClothingSimulation.cpp b/Engine/Plugins/Experimental/ChaosCloth/Source/Chaos/Private/ChaosCloth/ChaosClothingSimulation.cpp
--- a/Engine/Plugins/Experimental/ChaosCloth/Source/Chaos/Private/ChaosCloth/ChaosClothingSimulation.cpp
+++ b/Engine/Plugins/Experimental/ChaosCloth/Source/Chaos/Private/ChaosCloth/ChaosClothingSimulation.cpp
@@ -851,32 +851,49 @@ void ClothingSimulation::AddCollisions(const FClothCollisionData& ClothCollision
 
 			const Chaos::TVector<float, 3> X0 = Sphere0.LocalPosition;
 			const Chaos::TVector<float, 3> X1 = Sphere1.LocalPosition;
+			const Chaos::TVector<float, 3> Axis = X1 - X0;
+			const float AxisSize = Axis.Size();
 
 			const float Radius0 = Sphere0.Radius;
 			const float Radius1 = Sphere1.Radius;
+			float MinRadius, MaxRadius;
+			if (Radius0 <= Radius1) { MinRadius = Radius0; MaxRadius = Radius1; }
+			else { MinRadius = Radius1; MaxRadius = Radius0; }
 
-			if (FMath::Abs(Radius0 - Radius1) < SMALL_NUMBER)
+			if (AxisSize < KINDA_SMALL_NUMBER)
+			{
+				// Sphere
+				BaseTransforms[i] = Chaos::TRigidTransform<float, 3>(FTransform::Identity);
+
+				CollisionParticles.SetDynamicGeometry(
+					i,
+					MakeUnique<Chaos::TSphere<float, 3>>(
+						X0,
+						MaxRadius));
+			}
+			else if (MaxRadius - MinRadius < KINDA_SMALL_NUMBER)
 			{
 				// Capsule
 				const Chaos::TVector<float, 3> Center = (X0 + X1) * 0.5f;  // Construct a capsule centered at the origin along the Z axis
-				const Chaos::TVector<float, 3> Axis = X1 - X0;
 				const Chaos::TRotation<float, 3> Rotation = Chaos::TRotation<float, 3>::FromRotatedVector(
 					Chaos::TVector<float, 3>::AxisVector(2),
 					Axis.GetSafeNormal());
 
 				BaseTransforms[i] = Chaos::TRigidTransform<float, 3>(Center, Rotation);
 
-				const float HalfHeight = Axis.Size() * 0.5f;
+				const float HalfHeight = AxisSize * 0.5f;
 				CollisionParticles.SetDynamicGeometry(
 					i,
 					MakeUnique<Chaos::TCapsule<float>>(
 						Chaos::TVector<float, 3>(0.f, 0.f, -HalfHeight), // Min
 						Chaos::TVector<float, 3>(0.f, 0.f, HalfHeight), // Max
-						Radius0));
+						MaxRadius));
 			}
 			else
 			{
 				// Tapered capsule
+				BaseTransforms[i] = Chaos::TRigidTransform<float, 3>(FTransform::Identity);
+
 				TArray<TUniquePtr<Chaos::FImplicitObject>> Objects;
 				Objects.Reserve(3);
 				Objects.Add(TUniquePtr<Chaos::FImplicitObject>(
@@ -887,7 +904,7 @@ void ClothingSimulation::AddCollisions(const FClothCollisionData& ClothCollision
 					new Chaos::TSphere<float, 3>(X1, Radius1)));
 				CollisionParticles.SetDynamicGeometry(
 					i,
-					MakeUnique<Chaos::TImplicitObjectUnion<float, 3>>(MoveTemp(Objects)));  // TODO(Kriss.Gossart): Replace this once a TTaperedCapsule implicit type is implemented
+					MakeUnique<Chaos::TImplicitObjectUnion<float, 3>>(MoveTemp(Objects)));  // TODO(Kriss.Gossart): Replace this once a TTaperedCapsule implicit type is implemented (note: this tapered cylinder with spheres is an approximation of a real tapered capsule)
 			}
 
 			// Skip spheres added as end caps for the capsule.
