Add missing stiffness value clamps in the Chaos Cloth config migration code.

diff --git a/Engine/Plugins/Experimental/ChaosCloth/Source/Chaos/Private/ChaosCloth/ChaosClothConfig.cpp b/Engine/Plugins/Experimental/ChaosCloth/Source/Chaos/Private/ChaosCloth/ChaosClothConfig.cpp
--- a/Engine/Plugins/Experimental/ChaosCloth/Source/Chaos/Private/ChaosCloth/ChaosClothConfig.cpp
+++ b/Engine/Plugins/Experimental/ChaosCloth/Source/Chaos/Private/ChaosCloth/ChaosClothConfig.cpp
@@ -42,24 +42,24 @@ void UChaosClothConfig::MigrateFrom(const FClothConfig_Legacy& ClothConfig)
 	const float HorizontalStiffness =
 		ClothConfig.HorizontalConstraintConfig.Stiffness *
 		ClothConfig.HorizontalConstraintConfig.StiffnessMultiplier;
-	EdgeStiffness = (VerticalStiffness + HorizontalStiffness) * 0.5f;
+	EdgeStiffness = FMath::Clamp((VerticalStiffness + HorizontalStiffness) * 0.5f, 0.f, 1.f);
 
-	BendingStiffness =
+	BendingStiffness = FMath::Clamp(
 		ClothConfig.BendConstraintConfig.Stiffness *
-		ClothConfig.BendConstraintConfig.StiffnessMultiplier;
+		ClothConfig.BendConstraintConfig.StiffnessMultiplier, 0.f, 1.f);
 
-	AreaStiffness =
+	AreaStiffness = FMath::Clamp(
 		ClothConfig.ShearConstraintConfig.Stiffness *
-		ClothConfig.ShearConstraintConfig.StiffnessMultiplier;
+		ClothConfig.ShearConstraintConfig.StiffnessMultiplier, 0.f, 1.f);
 
-	AnimDriveSpringStiffness = ClothConfig.AnimDriveSpringStiffness;
+	AnimDriveSpringStiffness = FMath::Clamp(ClothConfig.AnimDriveSpringStiffness, 0.f, 1.f);
 
 	CoefficientOfFriction = FMath::Clamp(ClothConfig.Friction, 0.f, 10.f);
 
 	bUseBendingElements = false;
 	bUseSelfCollisions = (ClothConfig.SelfCollisionRadius > 0.f && ClothConfig.SelfCollisionStiffness > 0.f);
 
-	StrainLimitingStiffness = ClothConfig.TetherStiffness;
+	StrainLimitingStiffness = FMath::Clamp(ClothConfig.TetherStiffness, 0.f, 1.f);
 	ShapeTargetStiffness = 0.f;
 }
 
