Chaos Cloth - Add live editing for cloth property changes.
- Add Chaos Cloth physics refresh, called when UI manipulators are used on collision volumes.
- Tidied up ChaosClothingSimulation.cpp to remove redundant function arguments.
- Bug fix: Shape constraints now have a version where it is not applied on all particles but only to the relevant instance.

diff --git a/Engine/Plugins/Experimental/ChaosCloth/Source/Chaos/Private/ChaosCloth/ChaosClothConfig.cpp b/Engine/Plugins/Experimental/ChaosCloth/Source/Chaos/Private/ChaosCloth/ChaosClothConfig.cpp
--- a/Engine/Plugins/Experimental/ChaosCloth/Source/Chaos/Private/ChaosCloth/ChaosClothConfig.cpp
+++ b/Engine/Plugins/Experimental/ChaosCloth/Source/Chaos/Private/ChaosCloth/ChaosClothConfig.cpp
@@ -3,6 +3,7 @@
 #include "ChaosCloth/ChaosClothConfig.h"
 #include "ClothConfig_Legacy.h"
 #include "ChaosClothSharedConfigCustomVersion.h"
+#include "ClothingSimulationInteractor.h"
 
 // Legacy parameters not yet migrated to Chaos parameters:
 //  WindDragCoefficient
@@ -104,3 +105,31 @@ void UChaosClothSharedSimConfig::PostLoad()
 		bUseGravityOverride = true;  // Default gravity override would otherwise disable the currently set gravity on older versions
 	}
 }
+
+#if WITH_EDITOR
+void UChaosClothSharedSimConfig::PostEditChangeChainProperty(FPropertyChangedChainEvent& ChainEvent)
+{
+	Super::PostEditChangeChainProperty(ChainEvent);
+
+	// Update the simulation if there is any interactor attached to the skeletal mesh component
+	if (ChainEvent.ChangeType != EPropertyChangeType::Interactive)
+	{
+		if (USkeletalMesh* const OwnerMesh = Cast<USkeletalMesh>(GetOuter()))
+		{
+			for (TObjectIterator<USkeletalMeshComponent> It; It; ++It)
+			{
+				if (const USkeletalMeshComponent* const Component = *It)
+				{
+					if (Component->SkeletalMesh == OwnerMesh)
+					{
+						if (UClothingSimulationInteractor* const CurInteractor = Component->GetClothingSimulationInteractor())
+						{
+							CurInteractor->ClothConfigUpdated();
+						}
+					}
+				}
+			}
+		}
+	}
+}
+#endif  // #if WITH_EDITOR
diff --git a/Engine/Plugins/Experimental/ChaosCloth/Source/Chaos/Private/ChaosCloth/ChaosClothingSimulation.cpp b/Engine/Plugins/Experimental/ChaosCloth/Source/Chaos/Private/ChaosCloth/ChaosClothingSimulation.cpp
--- a/Engine/Plugins/Experimental/ChaosCloth/Source/Chaos/Private/ChaosCloth/ChaosClothingSimulation.cpp
+++ b/Engine/Plugins/Experimental/ChaosCloth/Source/Chaos/Private/ChaosCloth/ChaosClothingSimulation.cpp
@@ -23,9 +23,8 @@
 #include "Chaos/PBDSphericalConstraint.h"
 #include "Chaos/PBDSpringConstraints.h"
 #include "Chaos/PBDVolumeConstraint.h"
+#include "Chaos/PBDShapeConstraints.h"
 #include "Chaos/PerParticleGravity.h"
-#include "Chaos/PerParticlePBDLongRangeConstraints.h"
-#include "Chaos/PerParticlePBDShapeConstraints.h"
 #include "Chaos/Plane.h"
 #include "Chaos/Sphere.h"
 #include "Chaos/TaperedCylinder.h"
@@ -132,7 +131,6 @@ void ClothingSimulation::Shutdown()
 {
 	Assets.Reset();
 	AnimDriveSpringStiffness.Reset();
-	ExtractedCollisions.Reset();
 	ExternalCollisions.Reset();
 	OldCollisionTransforms.Reset();
 	CollisionTransforms.Reset();
@@ -229,48 +227,62 @@ void ClothingSimulation::CreateActor(USkeletalMeshComponent* InOwnerComponent, U
 		reinterpret_cast<TArray<FVector>&>(AnimationNormals),
 		Offset);
 
-	for (uint32 i = Offset; i < Particles.Size(); ++i)
-	{
-		Particles.X(i) = AnimationPositions[i];
-		Particles.V(i) = Chaos::TVector<float, 3>(0.f, 0.f, 0.f);
-		// Initialize mass to 0, to be overridden later
-		Particles.M(i) = 0.f;
-	}
+	ResetParticles(InSimDataIndex);
 
 	OldAnimationPositions = AnimationPositions;  // Also update the old positions array to avoid any interpolation issues
 
-	TUniquePtr<Chaos::TTriangleMesh<float>>& Mesh = Meshes[InSimDataIndex];
-	BuildMesh(Mesh, PhysMesh, Offset);
+	BuildMesh(PhysMesh, InSimDataIndex);
 
-	const FPointWeightMap& MaxDistances = PhysMesh.GetWeightMap(EChaosWeightMapTarget::MaxDistance);
-	SetParticleMasses(ChaosClothSimConfig, Particles, Offset, PhysMesh.Vertices.Num(), *Mesh, MaxDistances);
+	SetParticleMasses(ChaosClothSimConfig, PhysMesh, InSimDataIndex);
 
-	AddConstraints(ChaosClothSimConfig, PhysMesh, *Mesh, Offset, InSimDataIndex);
+	AddConstraints(ChaosClothSimConfig, PhysMesh, InSimDataIndex);
 
 	// Add Self Collisions
 	if (ChaosClothSimConfig->bUseSelfCollisions)
 	{
-		AddSelfCollisions(Mesh, Particles, Offset);
+		AddSelfCollisions(InSimDataIndex);
 	}
 
-	// Collision particles
-	TGeometryClothParticles<float, 3>& CollisionParticles = Evolution->CollisionParticles();
+	// Warn about legacy apex collisions
+	const FClothCollisionData& LodCollData = AssetLodData->CollisionData;
+	UE_CLOG(LodCollData.Spheres.Num() > 0 || LodCollData.SphereConnections.Num() > 0 || LodCollData.Convexes.Num() > 0,
+		LogChaosCloth, Warning, TEXT(
+			"Actor '%s' component '%s' has %d sphere, %d capsule, and %d "
+			"convex collision objects for physics authored as part of a LOD construct, "
+			"probably by the Apex cloth authoring system.  This is deprecated.  "
+			"Please update your asset!"),
+		InOwnerComponent->GetOwner() ? *InOwnerComponent->GetOwner()->GetName() : TEXT("None"),
+		*InOwnerComponent->GetName(),
+		LodCollData.Spheres.Num(),
+		LodCollData.SphereConnections.Num(),
+		LodCollData.Convexes.Num());
+
+	// Extract all collisions for this asset
+	checkf(ExternalCollisions.Spheres.Num() == 0 &&
+		ExternalCollisions.SphereConnections.Num() == 0 &&
+		ExternalCollisions.Convexes.Num() == 0 &&
+		ExternalCollisions.Boxes.Num() == 0, TEXT("There cannot be any external collisions added before all the cloth assets collisions are processed."));
+	ExtractCollisions(Asset);
+}
 
+void ClothingSimulation::ExtractCollisions(const UClothingAssetCommon* Asset)
+{
 	// Pull collisions from the specified physics asset inside the clothing asset
 	ExtractPhysicsAssetCollisions(Asset);
 
 	// Extract the legacy Apex collision from the clothing asset
-	ExtractLegacyAssetCollisions(Asset, InOwnerComponent);
+	ExtractLegacyAssetCollisions(Asset);
 
-	// Set the external collision starting index
-	checkf(ExternalCollisions.Spheres.Num() == 0 &&
-		ExternalCollisions.SphereConnections.Num() == 0 &&
-		ExternalCollisions.Convexes.Num() == 0 &&
-		ExternalCollisions.Boxes.Num() == 0, TEXT("There cannot be any external collisions added before all the cloth assets collisions are processed."));
-	ExternalCollisionsOffset = CollisionParticles.Size();
+	// Update the external collision offset
+	ExternalCollisionsOffset = Evolution->CollisionParticles().Size();
 }
 
 void ClothingSimulation::PostActorCreationInitialize()
+{
+	UpdateSimulationFromSharedSimConfig();
+}
+
+void ClothingSimulation::UpdateSimulationFromSharedSimConfig()
 {
 	if (ClothSharedSimConfig) // ClothSharedSimConfig will be a null pointer if all cloth instances are disabled in which case we will use default Evolution parameters
 	{
@@ -283,8 +295,12 @@ void ClothingSimulation::PostActorCreationInitialize()
 	}
 }
 
-void ClothingSimulation::BuildMesh(TUniquePtr<Chaos::TTriangleMesh<float>>& OutChaosMesh, const FClothPhysicalMeshData& InPhysMesh, uint32 Offset) const
+void ClothingSimulation::BuildMesh(const FClothPhysicalMeshData& InPhysMesh, int32 InSimDataIndex)
 {
+	TUniquePtr<Chaos::TTriangleMesh<float>>& Mesh = Meshes[InSimDataIndex];
+
+	const uint32 Offset = IndexToRangeMap[InSimDataIndex][0];
+
 	const int32 NumTriangles = InPhysMesh.Indices.Num() / 3;
 	TArray<Chaos::TVector<int32, 3>> InputSurfaceElements;
 	InputSurfaceElements.Reserve(NumTriangles);
@@ -297,20 +313,29 @@ void ClothingSimulation::BuildMesh(TUniquePtr<Chaos::TTriangleMesh<float>>& OutC
 			 static_cast<int32>(Offset + InPhysMesh.Indices[Index + 2]) });
 	}
 	check(InputSurfaceElements.Num() == NumTriangles);
-	OutChaosMesh.Reset(new Chaos::TTriangleMesh<float>(MoveTemp(InputSurfaceElements)));
-	check(OutChaosMesh->GetNumElements() == NumTriangles);
-	OutChaosMesh->GetPointToTriangleMap(); // Builds map for later use by GetPointNormals()
+	Mesh.Reset(new Chaos::TTriangleMesh<float>(MoveTemp(InputSurfaceElements)));
+	check(Mesh->GetNumElements() == NumTriangles);
+	Mesh->GetPointToTriangleMap(); // Builds map for later use by GetPointNormals()
+}
+
+void ClothingSimulation::ResetParticles(int32 InSimDataIndex)
+{
+	const uint32 Offset = IndexToRangeMap[InSimDataIndex][0];
+	const uint32 Range = IndexToRangeMap[InSimDataIndex][1];
+
+	for (uint32 i = Offset; i < Range; ++i)
+	{
+		Evolution->Particles().X(i) = AnimationPositions[i];
+		Evolution->Particles().V(i) = Chaos::TVector<float, 3>(0.f);
+		Evolution->Particles().M(i) = 0.f;
+	}
 }
 
-void ClothingSimulation::SetParticleMasses(
-	const UChaosClothConfig* ChaosClothSimConfig,
-	TPBDParticles<float, 3>& Particles,
-	uint32 Offset, // Offset into Particles array
-	uint32 ParticlesCount, // Number of particles
-	Chaos::TTriangleMesh<float>& Mesh, // Mesh with indices into Particles with offset already taken into account
-	const FPointWeightMap& MaxDistances
-	) const
+void ClothingSimulation::SetParticleMasses(const UChaosClothConfig* ChaosClothConfig, const FClothPhysicalMeshData& PhysMesh, int32 InSimDataIndex)
 {
+	const Chaos::TTriangleMesh<float>& Mesh = *Meshes[InSimDataIndex];
+	TPBDParticles<float, 3>& Particles = Evolution->Particles();
+
 	// Assign per particle mass proportional to connected area.
 	const TArray<TVector<int32, 3>>& SurfaceElements = Mesh.GetSurfaceElements();
 	float TotalArea = 0.0;
@@ -326,17 +351,17 @@ void ClothingSimulation::SetParticleMasses(
 		Particles.M(Tri[2]) += ThirdTriArea;
 	}
 	const TSet<int32> Vertices = Mesh.GetVertices();
-	switch (ChaosClothSimConfig->MassMode)
+	switch (ChaosClothConfig->MassMode)
 	{
 	case EClothMassMode::UniformMass:
 		for (const int32 Vertex : Vertices)
 		{
-			Particles.M(Vertex) = ChaosClothSimConfig->UniformMass;
+			Particles.M(Vertex) = ChaosClothConfig->UniformMass;
 		}
 		break;
 	case EClothMassMode::TotalMass:
 	{
-		const float MassPerUnitArea = TotalArea > 0.0 ? ChaosClothSimConfig->TotalMass / TotalArea : 1.0;
+		const float MassPerUnitArea = TotalArea > 0.0 ? ChaosClothConfig->TotalMass / TotalArea : 1.0;
 		for (const int32 Vertex : Vertices)
 		{
 			Particles.M(Vertex) *= MassPerUnitArea;
@@ -346,28 +371,36 @@ void ClothingSimulation::SetParticleMasses(
 	case EClothMassMode::Density:
 		for (const int32 Vertex : Vertices)
 		{
-			Particles.M(Vertex) *= ChaosClothSimConfig->Density;
+			Particles.M(Vertex) *= ChaosClothConfig->Density;
 		}
 		break;
 	};
 
 	// Clamp and enslave
-	check(Particles.Size() >= Offset + ParticlesCount);
-	for (uint32 i = Offset; i < Offset + ParticlesCount; i++)
+	const FPointWeightMap& MaxDistances = PhysMesh.GetWeightMap(EChaosWeightMapTarget::MaxDistance);
+	const uint32 Offset = IndexToRangeMap[InSimDataIndex][0];
+	const uint32 Range = IndexToRangeMap[InSimDataIndex][1];
+
+	check(Particles.Size() >= Range);
+	for (uint32 i = Offset; i < Range; i++)
 	{
-		Particles.M(i) = FMath::Max(Particles.M(i), ChaosClothSimConfig->MinPerParticleMass);
+		Particles.M(i) = FMath::Max(Particles.M(i), ChaosClothConfig->MinPerParticleMass);
 		Particles.InvM(i) = MaxDistances.IsBelowThreshold(i - Offset) ? 0.f : 1.f / Particles.M(i);
 	}
 }
 
-void ClothingSimulation::AddConstraints(const UChaosClothConfig* ChaosClothSimConfig, const FClothPhysicalMeshData& PhysMesh, const Chaos::TTriangleMesh<float>& Mesh, uint32 Offset, int32 InSimDataIndex)
+void ClothingSimulation::AddConstraints(const UChaosClothConfig* ChaosClothSimConfig, const FClothPhysicalMeshData& PhysMesh, int32 InSimDataIndex)
 {
+	const Chaos::TTriangleMesh<float>& Mesh = *Meshes[InSimDataIndex];
 	const TArray<TVector<int32, 3>>& SurfaceElements = Mesh.GetSurfaceElements();
 
+	const uint32 Offset = IndexToRangeMap[InSimDataIndex][0];
+	const uint32 ParticleCount = IndexToRangeMap[InSimDataIndex][1] - Offset;
+
 	if (ChaosClothSimConfig->ShapeTargetStiffness)
 	{
 		check(ChaosClothSimConfig->ShapeTargetStiffness > 0.f && ChaosClothSimConfig->ShapeTargetStiffness <= 1.f);
-		Evolution->AddPBDConstraintFunction([ShapeConstraints = Chaos::TPerParticlePBDShapeConstraints<float, 3>(Evolution->Particles(), AnimationPositions, ChaosClothSimConfig->ShapeTargetStiffness)](TPBDParticles<float, 3>& InParticles, const float Dt) {
+		Evolution->AddPBDConstraintFunction([ShapeConstraints = Chaos::TPBDShapeConstraints<float, 3>(Evolution->Particles(), Offset, ParticleCount, AnimationPositions, ChaosClothSimConfig->ShapeTargetStiffness)](TPBDParticles<float, 3>& InParticles, const float Dt) {
 			ShapeConstraints.Apply(InParticles, Dt);
 		});
 	}
@@ -517,14 +550,17 @@ void ClothingSimulation::AddConstraints(const UChaosClothConfig* ChaosClothSimCo
 	}
 }
 
-void ClothingSimulation::AddSelfCollisions(const TUniquePtr<Chaos::TTriangleMesh<float>>& Mesh, const TPBDParticles<float, 3>& Particles, int32 Offset)
+void ClothingSimulation::AddSelfCollisions(int32 InSimDataIndex) 
 {
 	// TODO(mlentine): Parallelize these for multiple meshes
-	const TArray<TVector<int32, 3>>& SurfaceElements = Mesh->GetSurfaceElements();
-	Evolution->CollisionTriangles().Append(SurfaceElements);
-	for (uint32 i = Offset; i < Particles.Size(); ++i)
+	const Chaos::TTriangleMesh<float>& Mesh = *Meshes[InSimDataIndex];
+	Evolution->CollisionTriangles().Append(Mesh.GetSurfaceElements());
+
+	const uint32 Offset = IndexToRangeMap[InSimDataIndex][0];
+	const uint32 Range = IndexToRangeMap[InSimDataIndex][1];
+	for (uint32 i = Offset; i < Range; ++i)
 	{
-		TSet<int32> Neighbors = Mesh->GetNRing(i, 5);
+		const TSet<int32> Neighbors = Mesh.GetNRing(i, 5);
 		for (int32 Element : Neighbors)
 		{
 			check(i != Element);
@@ -580,9 +616,9 @@ void ClothingSimulation::UpdateCollisionTransforms(const ClothingSimulationConte
 	}
 }
 
-void ClothingSimulation::ExtractPhysicsAssetCollisions(UClothingAssetCommon* Asset)
+void ClothingSimulation::ExtractPhysicsAssetCollisions(const UClothingAssetCommon* Asset)
 {
-	ExtractedCollisions.Reset();
+	FClothCollisionData ExtractedCollisions;
 
 	TGeometryClothParticles<float, 3>& CollisionParticles = Evolution->CollisionParticles();
 
@@ -760,24 +796,13 @@ void ClothingSimulation::ExtractPhysicsAssetCollisions(UClothingAssetCommon* Ass
 	}  // End if Asset->PhysicsAsset
 }
 
-void ClothingSimulation::ExtractLegacyAssetCollisions(UClothingAssetCommon* Asset, const USkeletalMeshComponent* InOwnerComponent)
+void ClothingSimulation::ExtractLegacyAssetCollisions(const UClothingAssetCommon* Asset)
 {
 	if (const UClothLODDataCommon* const AssetLodData = Asset->ClothLodData[0])
 	{
 		const FClothCollisionData& LodCollData = AssetLodData->CollisionData;
 		if (LodCollData.Spheres.Num() || LodCollData.SphereConnections.Num() || LodCollData.Convexes.Num())
 		{
-			UE_LOG(LogChaosCloth, Warning,
-				TEXT("Actor '%s' component '%s' has %d sphere, %d capsule, and %d "
-					"convex collision objects for physics authored as part of a LOD construct, "
-					"probably by the Apex cloth authoring system.  This is deprecated.  "
-					"Please update your asset!"),
-				InOwnerComponent->GetOwner() ? *InOwnerComponent->GetOwner()->GetName() : TEXT("None"),
-				*InOwnerComponent->GetName(),
-				LodCollData.Spheres.Num(),
-				LodCollData.SphereConnections.Num(),
-				LodCollData.Convexes.Num());
-
 			UE_LOG(LogChaosCloth, VeryVerbose, TEXT("Adding legacy cloth asset collisions..."));
 			AddCollisions(LodCollData, Asset->UsedBoneIndices);
 		}
@@ -1203,7 +1228,6 @@ void ClothingSimulation::ClearExternalCollisions()
 	CollisionParticles.Resize(ExternalCollisionsOffset);  // This will also resize BoneIndices and BaseTransforms
 
 	// Reset external collisions
-	ExternalCollisionsOffset = CollisionParticles.Size();
 	ExternalCollisions.Reset();
 
 	UE_LOG(LogChaosCloth, VeryVerbose, TEXT("Cleared all external collisions."));
@@ -1211,6 +1235,10 @@ void ClothingSimulation::ClearExternalCollisions()
 
 void ClothingSimulation::GetCollisions(FClothCollisionData& OutCollisions, bool bIncludeExternal) const
 {
+	// This code only gathers old apex collisions that don't appear in the physics mesh
+	// It is also never called with bIncludeExternal = true
+	// This function is bound to be deprecated at some point
+
 	OutCollisions.Reset();
 
 	// Add internal asset collisions
@@ -1222,18 +1250,63 @@ void ClothingSimulation::GetCollisions(FClothCollisionData& OutCollisions, bool
 		}
 	}
 
-	// Add collisions extracted from the physics asset
-	// TODO(Kriss.Gossart): Including the following code seems to be the correct behaviour, but this did not appear
-	// in the NvCloth implementation, so best to leave it commented out for now.
-	//OutCollisions.Append(ExtractedCollisions);
-
 	// Add external asset collisions
 	if (bIncludeExternal)
 	{
 		OutCollisions.Append(ExternalCollisions);
 	}
 
-	UE_LOG(LogChaosCloth, VeryVerbose, TEXT("Returned collisions: %d spheres, %d capsules, %d convexes, %d boxes."), OutCollisions.Spheres.Num() - 2 * OutCollisions.SphereConnections.Num(), OutCollisions.SphereConnections.Num(), OutCollisions.Convexes.Num(), OutCollisions.Boxes.Num());
+	UE_LOG(LogChaosCloth, VeryVerbose, TEXT("GetCollisions returned collisions: %d spheres, %d capsules, %d convexes, %d boxes."), OutCollisions.Spheres.Num() - 2 * OutCollisions.SphereConnections.Num(), OutCollisions.SphereConnections.Num(), OutCollisions.Convexes.Num(), OutCollisions.Boxes.Num());
+}
+
+void ClothingSimulation::RefreshClothConfig()
+{
+	UpdateSimulationFromSharedSimConfig();
+
+	Evolution->ResetConstraintRules();
+	Evolution->ResetSelfCollision();
+
+	for (int32 SimDataIndex = 0; SimDataIndex < Assets.Num(); ++SimDataIndex)
+	{
+		if (const UClothingAssetCommon* const Asset = Assets[SimDataIndex])
+		{
+			if (const UChaosClothConfig* const ChaosClothConfig = Asset->GetClothConfig<UChaosClothConfig>())
+			{
+				check(Asset->GetNumLods() > 0);
+				const FClothPhysicalMeshData& PhysMesh = Asset->ClothLodData[0]->ClothPhysicalMeshData;
+
+				ResetParticles(SimDataIndex);
+
+				SetParticleMasses(ChaosClothConfig, PhysMesh, SimDataIndex);
+
+				AddConstraints(ChaosClothConfig, PhysMesh, SimDataIndex);
+
+				// Add Self Collisions
+				if (ChaosClothConfig->bUseSelfCollisions)
+				{
+					AddSelfCollisions(SimDataIndex);
+				}
+			}
+		}
+	}
+	UE_LOG(LogChaosCloth, VeryVerbose, TEXT("RefreshClothConfig, all constraints and self-collisions have been updated for all clothing assets"));
+}
+
+void ClothingSimulation::RefreshPhysicsAsset()
+{
+	// Clear all collisions
+	TGeometryClothParticles<float, 3>& CollisionParticles = Evolution->CollisionParticles();
+	CollisionParticles.Resize(0);  // This will also resize BoneIndices and BaseTransforms
+
+	ExternalCollisions.Reset();
+	ExternalCollisionsOffset = 0;
+
+	// Re-extract all collisions from every cloth asset
+	for (const UClothingAssetCommon* Asset : Assets)
+	{
+		ExtractCollisions(Asset);
+	}
+	UE_LOG(LogChaosCloth, VeryVerbose, TEXT("RefreshPhysicsAsset, all collisions have been re-added for all clothing assets"));
 }
 
 void ClothingSimulation::SetAnimDriveSpringStiffness(float InStiffness)
diff --git a/Engine/Plugins/Experimental/ChaosCloth/Source/Chaos/Private/ChaosCloth/ChaosClothingSimulationInteractor.cpp b/Engine/Plugins/Experimental/ChaosCloth/Source/Chaos/Private/ChaosCloth/ChaosClothingSimulationInteractor.cpp
--- a/Engine/Plugins/Experimental/ChaosCloth/Source/Chaos/Private/ChaosCloth/ChaosClothingSimulationInteractor.cpp
+++ b/Engine/Plugins/Experimental/ChaosCloth/Source/Chaos/Private/ChaosCloth/ChaosClothingSimulationInteractor.cpp
@@ -4,12 +4,22 @@
 
 void UChaosClothingSimulationInteractor::PhysicsAssetUpdated()
 {
-	//TODO
+	Commands.Add(ChaosClothInteractorCommand::CreateLambda([](Chaos::ClothingSimulation* InSimulation, Chaos::ClothingSimulationContext* InContext)
+	{
+		InSimulation->RefreshPhysicsAsset();
+	}));
+
+	MarkDirty();
 }
 
 void UChaosClothingSimulationInteractor::ClothConfigUpdated()
 {
-	//TODO
+	Commands.Add(ChaosClothInteractorCommand::CreateLambda([](Chaos::ClothingSimulation* InSimulation, Chaos::ClothingSimulationContext* InContext)
+	{
+		InSimulation->RefreshClothConfig();
+	}));
+
+	MarkDirty();
 }
 
 void UChaosClothingSimulationInteractor::Sync(IClothingSimulation* InSimulation, IClothingSimulationContext* InContext)
diff --git a/Engine/Plugins/Experimental/ChaosCloth/Source/Chaos/Public/ChaosCloth/ChaosClothConfig.h b/Engine/Plugins/Experimental/ChaosCloth/Source/Chaos/Public/ChaosCloth/ChaosClothConfig.h
--- a/Engine/Plugins/Experimental/ChaosCloth/Source/Chaos/Public/ChaosCloth/ChaosClothConfig.h
+++ b/Engine/Plugins/Experimental/ChaosCloth/Source/Chaos/Public/ChaosCloth/ChaosClothConfig.h
@@ -126,6 +126,11 @@ public:
 	/** PostLoad override used to deal with updates/changes in properties. */
 	virtual void PostLoad() override;
 
+#if WITH_EDITOR
+	/** Called after changes in any of the asset properties. */
+	virtual void PostEditChangeChainProperty(FPropertyChangedChainEvent& ChainEvent) override;
+#endif
+
 	// The number of solver iterations
 	// This will increase the stiffness of all constraints but will increase the CPU cost
 	UPROPERTY(EditAnywhere, Category = Simulation, meta = (UIMin = "0", UIMax = "20", ClampMin = "0", ClampMax = "100"))
diff --git a/Engine/Plugins/Experimental/ChaosCloth/Source/Chaos/Public/ChaosCloth/ChaosClothingSimulation.h b/Engine/Plugins/Experimental/ChaosCloth/Source/Chaos/Public/ChaosCloth/ChaosClothingSimulation.h
--- a/Engine/Plugins/Experimental/ChaosCloth/Source/Chaos/Public/ChaosCloth/ChaosClothingSimulation.h
+++ b/Engine/Plugins/Experimental/ChaosCloth/Source/Chaos/Public/ChaosCloth/ChaosClothingSimulation.h
@@ -33,16 +33,23 @@ namespace Chaos
 	{
 	public:
 		ClothingSimulation();
-		virtual ~ClothingSimulation();
+		virtual ~ClothingSimulation() override;
 
 		// Set the animation drive stiffness for all actors
 		void SetAnimDriveSpringStiffness(float InStiffness);
 		void SetGravityOverride(const FVector& InGravityOverride);
 		void DisableGravityOverride();
 
+		// Function to be called if any of the assets' configuration parameters have changed
+		void RefreshClothConfig();
+		// Function to be called if any of the assets' physics assets changes (colliders)
+		// This seems to only happen when UPhysicsAsset::RefreshPhysicsAssetChange is called with
+		// bFullClothRefresh set to false during changes created using the viewport manipulators.
+		void RefreshPhysicsAsset();
+
 #if WITH_EDITOR
 		// FGCObject interface
-		void AddReferencedObjects(FReferenceCollector& Collector) override;
+		virtual void AddReferencedObjects(FReferenceCollector& Collector) override;
 		// End of FGCObject interface
 
 		CHAOSCLOTH_API void DebugDrawPhysMeshWired(USkeletalMeshComponent* OwnerComponent, FPrimitiveDrawInterface* PDI) const;
@@ -59,47 +66,45 @@ namespace Chaos
 
 	protected:
 		// IClothingSimulation interface
-		void Initialize() override;
-		void CreateActor(USkeletalMeshComponent* InOwnerComponent, UClothingAssetBase* InAsset, int32 SimDataIndex) override;
-		void PostActorCreationInitialize() override;
-		IClothingSimulationContext* CreateContext() override { return new ClothingSimulationContext(); }
-		void FillContext(USkeletalMeshComponent* InComponent, float InDeltaTime, IClothingSimulationContext* InOutContext) override;
-		void Shutdown() override;
-		bool ShouldSimulate() const override { return true; }
-		void Simulate(IClothingSimulationContext* InContext) override;
-		void DestroyActors() override;
-		void DestroyContext(IClothingSimulationContext* InContext) override { delete InContext; }
-		void GetSimulationData(TMap<int32, FClothSimulData>& OutData, USkeletalMeshComponent* InOwnerComponent, USkinnedMeshComponent* InOverrideComponent) const override;
-
-		FBoxSphereBounds GetBounds(const USkeletalMeshComponent* InOwnerComponent) const override
+		virtual void Initialize() override;
+		virtual void CreateActor(USkeletalMeshComponent* InOwnerComponent, UClothingAssetBase* InAsset, int32 SimDataIndex) override;
+		virtual void PostActorCreationInitialize() override;
+		virtual IClothingSimulationContext* CreateContext() override { return new ClothingSimulationContext(); }
+		virtual void FillContext(USkeletalMeshComponent* InComponent, float InDeltaTime, IClothingSimulationContext* InOutContext) override;
+		virtual void Shutdown() override;
+		virtual bool ShouldSimulate() const override { return true; }
+		virtual void Simulate(IClothingSimulationContext* InContext) override;
+		virtual void DestroyActors() override;
+		virtual void DestroyContext(IClothingSimulationContext* InContext) override { delete InContext; }
+		virtual void GetSimulationData(TMap<int32, FClothSimulData>& OutData, USkeletalMeshComponent* InOwnerComponent, USkinnedMeshComponent* InOverrideComponent) const override;
+
+		virtual FBoxSphereBounds GetBounds(const USkeletalMeshComponent* InOwnerComponent) const override
 		{ return FBoxSphereBounds(Evolution->Particles().X().GetData(), Evolution->Particles().Size()); }
 
-		void AddExternalCollisions(const FClothCollisionData& InData) override;
-		void ClearExternalCollisions() override;
-		void GetCollisions(FClothCollisionData& OutCollisions, bool bIncludeExternal = true) const override;
+		virtual void AddExternalCollisions(const FClothCollisionData& InData) override;
+		virtual void ClearExternalCollisions() override;
+		virtual void GetCollisions(FClothCollisionData& OutCollisions, bool bIncludeExternal = true) const override;
 		// End of IClothingSimulation interface
 
 	private:
+		void UpdateSimulationFromSharedSimConfig();
+		void BuildMesh(const FClothPhysicalMeshData& PhysMesh, int32 InSimDataIndex);
 
-		void BuildMesh(TUniquePtr<Chaos::TTriangleMesh<float>>& OutChaosMesh, const FClothPhysicalMeshData& InPhysMesh, uint32 Offset) const;
+		// Reset particle positions to animation positions, masses and velocities are cleared
+		void ResetParticles(int32 InSimDataIndex);
 
-		void SetParticleMasses(
-			const UChaosClothConfig* ChaosClothSimConfig,
-			TPBDParticles<float, 3>& Particles,
-			uint32 Offset, // Offset into Particles array
-			uint32 ParticlesCount, // Number of particles
-			Chaos::TTriangleMesh<float>& Mesh, // Mesh with indices into Particles with offset already taken into account
-			const FPointWeightMap& MaxDistances
-			) const;
+		void SetParticleMasses(const UChaosClothConfig* ChaosClothSimConfig, const FClothPhysicalMeshData& PhysMesh, int32 InSimDataIndex);
 
-		void AddConstraints(const UChaosClothConfig* ChaosClothSimConfig, const FClothPhysicalMeshData& PhysMesh, const Chaos::TTriangleMesh<float>& Mesh, uint32 Offset, int32 InSimDataIndex);
+		void AddConstraints(const UChaosClothConfig* ChaosClothSimConfig, const FClothPhysicalMeshData& PhysMesh, int32 InSimDataIndex);
 
-		void AddSelfCollisions(const TUniquePtr<Chaos::TTriangleMesh<float>>& Mesh, const TPBDParticles<float, 3>& Particles, int32 Offset);
+		void AddSelfCollisions(int32 InSimDataIndex);
 
+		// Extract all collisions (from the physics asset and from the legacy apex collisions)
+		void ExtractCollisions(const UClothingAssetCommon* Asset);
 		// Extract the collisions from the physics asset referenced by the specified clothing asset
-		void ExtractPhysicsAssetCollisions(UClothingAssetCommon* Asset);
-		// Extract the collisions from the specified clothing asset 
-		void ExtractLegacyAssetCollisions(UClothingAssetCommon* Asset, const USkeletalMeshComponent* InOwnerComponent);
+		void ExtractPhysicsAssetCollisions(const UClothingAssetCommon* Asset);
+		// Extract the legacy apex collisions from the specified clothing asset 
+		void ExtractLegacyAssetCollisions(const UClothingAssetCommon* Asset);
 		// Add collisions from a ClothCollisionData structure
 		void AddCollisions(const FClothCollisionData& ClothCollisionData, const TArray<int32>& UsedBoneIndices);
 		// Update the collision transforms using the specified context
@@ -121,7 +126,6 @@ namespace Chaos
 		TArray<float> AnimDriveSpringStiffness; // One for every Asset
 
 		// Collision Data
-		FClothCollisionData ExtractedCollisions;  // Collisions extracted from the referenced physics asset
 		FClothCollisionData ExternalCollisions;  // External collisions
 		TArray<Chaos::TRigidTransform<float, 3>> OldCollisionTransforms;  // Used for the kinematic collision transform update
 		TArray<Chaos::TRigidTransform<float, 3>> CollisionTransforms;  // Used for the kinematic collision transform update
diff --git a/Engine/Source/Runtime/ClothingSystemRuntimeCommon/Private/ClothingAsset.cpp b/Engine/Source/Runtime/ClothingSystemRuntimeCommon/Private/ClothingAsset.cpp
--- a/Engine/Source/Runtime/ClothingSystemRuntimeCommon/Private/ClothingAsset.cpp
+++ b/Engine/Source/Runtime/ClothingSystemRuntimeCommon/Private/ClothingAsset.cpp
@@ -686,12 +686,12 @@ void UClothingAssetCommon::CalculateReferenceBoneIndex()
 	}
 }
 
-bool UClothingAssetCommon::IsValidLod(int32 InLodIndex)
+bool UClothingAssetCommon::IsValidLod(int32 InLodIndex) const
 {
 	return ClothLodData.IsValidIndex(InLodIndex);
 }
 
-int32 UClothingAssetCommon::GetNumLods()
+int32 UClothingAssetCommon::GetNumLods() const
 {
 	return ClothLodData.Num();
 }
diff --git a/Engine/Source/Runtime/ClothingSystemRuntimeCommon/Public/ClothingAsset.h b/Engine/Source/Runtime/ClothingSystemRuntimeCommon/Public/ClothingAsset.h
--- a/Engine/Source/Runtime/ClothingSystemRuntimeCommon/Public/ClothingAsset.h
+++ b/Engine/Source/Runtime/ClothingSystemRuntimeCommon/Public/ClothingAsset.h
@@ -134,10 +134,10 @@ public:
 	void CalculateReferenceBoneIndex();
 
 	/** Returns \c true if \p InLodIndex is a valid LOD id (index into \c ClothLodData). */
-	virtual bool IsValidLod(int32 InLodIndex) override;
+	virtual bool IsValidLod(int32 InLodIndex) const override;
 
 	/** Returns the number of valid LOD's (length of the \c ClothLodData array). */
-	virtual int32 GetNumLods() override;
+	virtual int32 GetNumLods() const override;
 
 #if WITH_EDITOR
 	/** * Add a new LOD class instance. */
diff --git a/Engine/Source/Runtime/ClothingSystemRuntimeInterface/Public/ClothingAssetBase.h b/Engine/Source/Runtime/ClothingSystemRuntimeInterface/Public/ClothingAssetBase.h
--- a/Engine/Source/Runtime/ClothingSystemRuntimeInterface/Public/ClothingAssetBase.h
+++ b/Engine/Source/Runtime/ClothingSystemRuntimeInterface/Public/ClothingAssetBase.h
@@ -58,11 +58,11 @@ public:
 	PURE_VIRTUAL(UClothingAssetBase::InvalidateCachedData(), );
 
 	/** Check the validity of a LOD index */
-	virtual bool IsValidLod(int32 InLodIndex)
+	virtual bool IsValidLod(int32 InLodIndex) const
 	PURE_VIRTUAL(UClothingAssetBase::IsValidLod(), return false;);
 
 	/** Get the number of LODs defined in the clothing asset */
-	virtual int32 GetNumLods()
+	virtual int32 GetNumLods() const
 	PURE_VIRTUAL(UClothingAssetBase::GetNumLods(), return 0;);
 
 	/** Add a new LOD class instance. */
diff --git a/Engine/Source/Runtime/Experimental/Chaos/Public/Chaos/PBDEvolution.h b/Engine/Source/Runtime/Experimental/Chaos/Public/Chaos/PBDEvolution.h
--- a/Engine/Source/Runtime/Experimental/Chaos/Public/Chaos/PBDEvolution.h
+++ b/Engine/Source/Runtime/Experimental/Chaos/Public/Chaos/PBDEvolution.h
@@ -56,11 +56,14 @@ class CHAOS_API TPBDEvolution
 
 	T GetTime() const { return MTime; }
 
+	void ResetConstraintRules() { MConstraintRules.Reset(); };
+	void ResetSelfCollision() { MCollisionTriangles.Reset(); MDisabledCollisionElements.Reset(); };
+
   private:
 	TPBDParticles<T, d> MParticles;
 	TKinematicGeometryClothParticles<T, d> MCollisionParticles;
-	TArray<TVector<int32, 3>> MCollisionTriangles;
-	TSet<TVector<int32, 2>> MDisabledCollisionElements;
+	TArray<TVector<int32, 3>> MCollisionTriangles;       // Used for self-collisions
+	TSet<TVector<int32, 2>> MDisabledCollisionElements;  // 
 	TArrayCollectionArray<bool> MCollided;
 	int32 MNumIterations;
 	T MCollisionThickness;
diff --git a/Engine/Source/Runtime/Experimental/Chaos/Public/Chaos/PBDShapeConstraints.h b/Engine/Source/Runtime/Experimental/Chaos/Public/Chaos/PBDShapeConstraints.h
new file mode 100644
--- /dev/null
+++ b/Engine/Source/Runtime/Experimental/Chaos/Public/Chaos/PBDShapeConstraints.h
@@ -0,0 +1,44 @@
+// Copyright 1998-2019 Epic Games, Inc. All Rights Reserved.
+#pragma once
+
+#include "Chaos/Array.h"
+#include "Chaos/Map.h"
+#include "Chaos/PBDParticles.h"
+#include "Chaos/PBDShapeConstraintsBase.h"
+
+namespace Chaos
+{
+	template<class T, int d>
+	class TPBDShapeConstraints : public TPBDShapeConstraintsBase<T, d>
+	{
+		typedef TPBDShapeConstraintsBase<T, d> Base;
+
+	public:
+
+		TPBDShapeConstraints(const TDynamicParticles<T, d>& InParticles, uint32 InParticleIndexOffset, uint32 InParticleCount, const TArray<TVector<float, 3>>& TargetPositions, const T Stiffness = (T)1)
+			: Base(InParticles, TargetPositions, Stiffness), ParticleIndexOffset(InParticleIndexOffset), ParticleCount(InParticleCount)
+		{
+		}
+		virtual ~TPBDShapeConstraints() {}
+
+		void Apply(TPBDParticles<T, d>& InParticles, const T Dt, const int32 Index) const
+		{
+			if (InParticles.InvM(Index) > 0)
+			{
+				InParticles.P(Index) -= InParticles.InvM(Index) * Base::GetDelta(InParticles, Index);
+			}
+		}
+
+		void Apply(TPBDParticles<T, d>& InParticles, const T Dt) const
+		{
+			for (uint32 Index = ParticleIndexOffset; Index < ParticleIndexOffset + ParticleCount; Index++)
+			{
+				Apply(InParticles, Dt, Index);
+			}
+		}
+
+	private:
+		uint32 ParticleIndexOffset;
+		uint32 ParticleCount;
+	};
+}
diff --git a/Engine/Source/Runtime/Experimental/Chaos/Public/Chaos/TriangleMesh.h b/Engine/Source/Runtime/Experimental/Chaos/Public/Chaos/TriangleMesh.h
--- a/Engine/Source/Runtime/Experimental/Chaos/Public/Chaos/TriangleMesh.h
+++ b/Engine/Source/Runtime/Experimental/Chaos/Public/Chaos/TriangleMesh.h
@@ -81,12 +81,12 @@ namespace Chaos
 		int32 GetNumElements() const { return MElements.Num(); }
 
 		const TMap<int32, TSet<uint32>>& GetPointToNeighborsMap() const;
-		const TSet<uint32>& GetNeighbors(const int32 Element) { return GetPointToNeighborsMap()[Element]; }
+		const TSet<uint32>& GetNeighbors(const int32 Element) const { return GetPointToNeighborsMap()[Element]; }
 
 		const TMap<int32, TArray<int32>>& GetPointToTriangleMap() const;
-		const TArray<int32>& GetCoincidentTriangles(const int32 Element) { return GetPointToTriangleMap()[Element]; }
+		const TArray<int32>& GetCoincidentTriangles(const int32 Element) const { return GetPointToTriangleMap()[Element]; }
 
-		TSet<int32> GetNRing(const int32 Element, const int32 N)
+		TSet<int32> GetNRing(const int32 Element, const int32 N) const
 		{
 			TSet<int32> Neighbors;
 			TSet<uint32> LevelNeighbors, PrevLevelNeighbors;
