Chaos Cloth - Add a solver frequency parameter to control the number of subdivisions and add solver substepping to prevent wild jittering during large game hitches.

diff --git a/Engine/Plugins/Experimental/ChaosCloth/Source/Chaos/Private/ChaosCloth/ChaosClothConfig.cpp b/Engine/Plugins/Experimental/ChaosCloth/Source/Chaos/Private/ChaosCloth/ChaosClothConfig.cpp
--- a/Engine/Plugins/Experimental/ChaosCloth/Source/Chaos/Private/ChaosCloth/ChaosClothConfig.cpp
+++ b/Engine/Plugins/Experimental/ChaosCloth/Source/Chaos/Private/ChaosCloth/ChaosClothConfig.cpp
@@ -24,7 +24,6 @@
 //  LinearInertiaScale
 //  AngularInertiaScale
 //  CentrifugalInertiaScale
-//  SolverFrequency
 //  StiffnessFrequency
 //  TetherLimit
 //  AnimDriveSpringStiffness
@@ -75,6 +74,8 @@ void UChaosClothSharedSimConfig::MigrateFrom(const FClothConfig_Legacy& ClothCon
 {
 	IterationCount = FMath::Clamp(int32(ClothConfig.SolverFrequency / 60.f), 1, 100);
 
+	SolverFrequency = ClothConfig.SolverFrequency;
+
 	SelfCollisionThickness = FMath::Clamp(ClothConfig.SelfCollisionRadius, 0.f, 1000.f);
 
 	CollisionThickness = FMath::Clamp(ClothConfig.CollisionThickness, 0.f, 1000.f);
diff --git a/Engine/Plugins/Experimental/ChaosCloth/Source/Chaos/Private/ChaosCloth/ChaosClothingSimulation.cpp b/Engine/Plugins/Experimental/ChaosCloth/Source/Chaos/Private/ChaosCloth/ChaosClothingSimulation.cpp
--- a/Engine/Plugins/Experimental/ChaosCloth/Source/Chaos/Private/ChaosCloth/ChaosClothingSimulation.cpp
+++ b/Engine/Plugins/Experimental/ChaosCloth/Source/Chaos/Private/ChaosCloth/ChaosClothingSimulation.cpp
@@ -40,11 +40,6 @@
 #include "PhysXIncludes.h"
 #endif
 
-#include <memory>
-#include <unordered_map>
-#include <unordered_set>
-#include "Chaos/ErrorReporter.h"
-
 using namespace Chaos;
 
 DECLARE_CYCLE_STAT(TEXT("Chaos Cloth Simulate"), STAT_ChaosClothSimulate, STATGROUP_ChaosCloth);
@@ -60,6 +55,7 @@ namespace ClothingSimulationDefault
 	static const float CollisionThickness = 1.2f;
 	static const float CoefficientOfFriction = 0.f;
 	static const float Damping = 0.01f;
+	static const float SolverFrequency = 60.f;
 }
 
 ClothingSimulation::ClothingSimulation()
@@ -122,7 +118,6 @@ void ClothingSimulation::Initialize()
 			ParticlesInput.R(Index) = NewR;
 		});
 
-    MaxDeltaTime = 1.0f;
     ClampDeltaTime = 0.f;
     Time = 0.f;
 }
@@ -1033,7 +1028,7 @@ void ClothingSimulation::FillContext(USkeletalMeshComponent* InComponent, float
 {
     ClothingSimulationContext* Context = static_cast<ClothingSimulationContext*>(InOutContext);
     Context->ComponentToWorld = InComponent->GetComponentToWorld();
-    Context->DeltaTime = ClampDeltaTime > 0 ? std::min(InDeltaTime, ClampDeltaTime) : InDeltaTime;
+    Context->DeltaTime = ClampDeltaTime > 0 ? FMath::Min(InDeltaTime, ClampDeltaTime) : InDeltaTime;
 
 	Context->WorldGravity = InComponent->GetWorld() ?
 		FVector(0.f, 0.f, InComponent->GetWorld()->GetGravityZ()) :
@@ -1148,12 +1143,20 @@ void ClothingSimulation::Simulate(IClothingSimulationContext* InContext)
 
 	// Advance Sim
 	DeltaTime = Context->DeltaTime;
-	while (Context->DeltaTime > MaxDeltaTime)
+
+	const float SolverFrequency = ClothSharedSimConfig ?
+		ClothSharedSimConfig->SolverFrequency :
+		ClothingSimulationDefault::SolverFrequency;
+
+	const int32 NumSubSteps = FMath::Max(1, FMath::RoundToInt(DeltaTime * SolverFrequency));  // Start substepping once deltatime is more than 50% above the expected simulation deltatime
+	const float SubDeltaTime = DeltaTime / float(NumSubSteps);
+
+	for (int32 i = 0; i < NumSubSteps; ++i)
 	{
-		Evolution->AdvanceOneTimeStep(MaxDeltaTime);
-		Context->DeltaTime -= MaxDeltaTime;
+		Evolution->AdvanceOneTimeStep(SubDeltaTime);
 	}
-	Evolution->AdvanceOneTimeStep(Context->DeltaTime);
+	UE_LOG(LogChaosCloth, VeryVerbose, TEXT("DeltaTime: %.6f, SubDeltaTime: %.6f, NumSubSteps = %d"), DeltaTime, SubDeltaTime, NumSubSteps);
+
 	Time += DeltaTime;
 }
 
diff --git a/Engine/Plugins/Experimental/ChaosCloth/Source/Chaos/Public/ChaosCloth/ChaosClothConfig.h b/Engine/Plugins/Experimental/ChaosCloth/Source/Chaos/Public/ChaosCloth/ChaosClothConfig.h
--- a/Engine/Plugins/Experimental/ChaosCloth/Source/Chaos/Public/ChaosCloth/ChaosClothConfig.h
+++ b/Engine/Plugins/Experimental/ChaosCloth/Source/Chaos/Public/ChaosCloth/ChaosClothConfig.h
@@ -136,6 +136,11 @@ public:
 	UPROPERTY(EditAnywhere, Category = Simulation, meta = (UIMin = "0", UIMax = "20", ClampMin = "0", ClampMax = "100"))
 	int32 IterationCount = 1;
 
+	// The expected simulation frequency in Hz that, if not closely matching the game tick frequency, will trigger substepping to help keep a more consistent time step.
+	// Use higher values to avoid jitter, and smaller value to avoid the extra cost of substepping.
+	UPROPERTY(EditAnywhere, Category = Simulation, meta = (UIMin = "30", UIMax = "240", ClampMin = "1", ClampMax = "1080"))
+	float SolverFrequency = 60.f;
+
 	// The radius of the spheres used in self collision 
 	UPROPERTY(EditAnywhere, Category = Collision, meta = (UIMin = "0", UIMax = "100", ClampMin = "0", ClampMax = "1000"))
 	float SelfCollisionThickness = 2.0f;
diff --git a/Engine/Plugins/Experimental/ChaosCloth/Source/Chaos/Public/ChaosCloth/ChaosClothingSimulation.h b/Engine/Plugins/Experimental/ChaosCloth/Source/Chaos/Public/ChaosCloth/ChaosClothingSimulation.h
--- a/Engine/Plugins/Experimental/ChaosCloth/Source/Chaos/Public/ChaosCloth/ChaosClothingSimulation.h
+++ b/Engine/Plugins/Experimental/ChaosCloth/Source/Chaos/Public/ChaosCloth/ChaosClothingSimulation.h
@@ -150,7 +150,6 @@ namespace Chaos
 
 		float Time;
 		float DeltaTime;
-		float MaxDeltaTime;
 		float ClampDeltaTime;
 
 		FVector Gravity;
