Chaos Cloth - Add XPBD Spring, XPBD Axial Spring, and XPBD Long Range constraints.
- Add new Init constraint rule function hooks to the PBD Evolution solver to allow for clearing the XPBD Lambdas.
- Add a new experimental property toggle to the Chaos Cloth shared config to enable XPBD constraints.

diff --git a/Engine/Plugins/Experimental/ChaosCloth/Source/Chaos/Private/ChaosCloth/ChaosClothingSimulation.cpp b/Engine/Plugins/Experimental/ChaosCloth/Source/Chaos/Private/ChaosCloth/ChaosClothingSimulation.cpp
--- a/Engine/Plugins/Experimental/ChaosCloth/Source/Chaos/Private/ChaosCloth/ChaosClothingSimulation.cpp
+++ b/Engine/Plugins/Experimental/ChaosCloth/Source/Chaos/Private/ChaosCloth/ChaosClothingSimulation.cpp
@@ -36,6 +36,10 @@
 #include "ChaosCloth/ChaosWeightMapTarget.h"
 #include "ChaosCloth/ChaosClothPrivate.h"
 
+#include "Chaos/XPBDLongRangeConstraints.h"
+#include "Chaos/XPBDSpringConstraints.h"
+#include "Chaos/XPBDAxialSpringConstraints.h"
+
 #if WITH_PHYSX && !PLATFORM_LUMIN && !PLATFORM_ANDROID
 #include "PhysXIncludes.h"
 #endif
@@ -392,19 +396,40 @@ void ClothingSimulation::AddConstraints(const UChaosClothConfig* ChaosClothSimCo
 	const uint32 Offset = IndexToRangeMap[InSimDataIndex][0];
 	const uint32 ParticleCount = IndexToRangeMap[InSimDataIndex][1] - Offset;
 
+	const bool bUseXPBDConstraints = ClothSharedSimConfig && ClothSharedSimConfig->bUseXPBDConstraints;
+
 	if (ChaosClothSimConfig->ShapeTargetStiffness)
 	{
 		check(ChaosClothSimConfig->ShapeTargetStiffness > 0.f && ChaosClothSimConfig->ShapeTargetStiffness <= 1.f);
-		Evolution->AddPBDConstraintFunction([ShapeConstraints = Chaos::TPBDShapeConstraints<float, 3>(Evolution->Particles(), Offset, ParticleCount, AnimationPositions, ChaosClothSimConfig->ShapeTargetStiffness)](TPBDParticles<float, 3>& InParticles, const float Dt) {
+		Evolution->AddPBDConstraintFunction([ShapeConstraints = Chaos::TPBDShapeConstraints<float, 3>(Evolution->Particles(), Offset, ParticleCount, AnimationPositions, ChaosClothSimConfig->ShapeTargetStiffness)](TPBDParticles<float, 3>& InParticles, const float Dt)
+		{
 			ShapeConstraints.Apply(InParticles, Dt);
 		});
 	}
 	if (ChaosClothSimConfig->EdgeStiffness)
 	{
 		check(ChaosClothSimConfig->EdgeStiffness > 0.f && ChaosClothSimConfig->EdgeStiffness <= 1.f);
-		Evolution->AddPBDConstraintFunction([SpringConstraints = Chaos::TPBDSpringConstraints<float, 3>(Evolution->Particles(), SurfaceElements, ChaosClothSimConfig->EdgeStiffness)](TPBDParticles<float, 3>& InParticles, const float Dt) {
-			SpringConstraints.Apply(InParticles, Dt);
-		});
+		if (bUseXPBDConstraints)
+		{
+			const TSharedPtr<Chaos::TXPBDSpringConstraints<float, 3>> SpringConstraints =
+				MakeShared<Chaos::TXPBDSpringConstraints<float, 3>>(Evolution->Particles(), SurfaceElements, ChaosClothSimConfig->EdgeStiffness);
+			Evolution->AddXPBDConstraintFunctions(
+				[SpringConstraints]()
+				{
+					SpringConstraints->Init();
+				},
+				[SpringConstraints](TPBDParticles<float, 3>& InParticles, const float Dt)
+				{
+					SpringConstraints->Apply(InParticles, Dt);
+				});
+		}
+		else
+		{
+			Evolution->AddPBDConstraintFunction([SpringConstraints = Chaos::TPBDSpringConstraints<float, 3>(Evolution->Particles(), SurfaceElements, ChaosClothSimConfig->EdgeStiffness)](TPBDParticles<float, 3>& InParticles, const float Dt)
+			{
+				SpringConstraints.Apply(InParticles, Dt);
+			});
+		}
 	}
 	if (ChaosClothSimConfig->BendingStiffness)
 	{
@@ -412,24 +437,61 @@ void ClothingSimulation::AddConstraints(const UChaosClothConfig* ChaosClothSimCo
 		if (ChaosClothSimConfig->bUseBendingElements)
 		{
 			TArray<Chaos::TVector<int32, 4>> BendingConstraints = Mesh.GetUniqueAdjacentElements();
-			Evolution->AddPBDConstraintFunction([BendConstraints = Chaos::TPBDBendingConstraints<float>(Evolution->Particles(), MoveTemp(BendingConstraints), ChaosClothSimConfig->BendingStiffness)](TPBDParticles<float, 3>& InParticles, const float Dt) {
+			Evolution->AddPBDConstraintFunction([BendConstraints = Chaos::TPBDBendingConstraints<float>(Evolution->Particles(), MoveTemp(BendingConstraints), ChaosClothSimConfig->BendingStiffness)](TPBDParticles<float, 3>& InParticles, const float Dt)
+			{
 				BendConstraints.Apply(InParticles, Dt);
 			});
 		}
 		else
 		{
 			TArray<Chaos::TVector<int32, 2>> BendingConstraints = Mesh.GetUniqueAdjacentPoints();
-			Evolution->AddPBDConstraintFunction([SpringConstraints = Chaos::TPBDSpringConstraints<float, 3>(Evolution->Particles(), MoveTemp(BendingConstraints), ChaosClothSimConfig->BendingStiffness)](TPBDParticles<float, 3>& InParticles, const float Dt) {
-				SpringConstraints.Apply(InParticles, Dt);
-			});
+			if (bUseXPBDConstraints)
+			{
+				const TSharedPtr<Chaos::TXPBDSpringConstraints<float, 3>> SpringConstraints =
+					MakeShared<Chaos::TXPBDSpringConstraints<float, 3>>(Evolution->Particles(), MoveTemp(BendingConstraints), ChaosClothSimConfig->BendingStiffness);
+				Evolution->AddXPBDConstraintFunctions(
+					[SpringConstraints]()
+					{
+						SpringConstraints->Init();
+					},
+					[SpringConstraints](TPBDParticles<float, 3>& InParticles, const float Dt)
+					{
+						SpringConstraints->Apply(InParticles, Dt);
+					});
+			}
+			else
+			{
+				Evolution->AddPBDConstraintFunction([SpringConstraints = Chaos::TPBDSpringConstraints<float, 3>(Evolution->Particles(), MoveTemp(BendingConstraints), ChaosClothSimConfig->BendingStiffness)](TPBDParticles<float, 3>& InParticles, const float Dt)
+				{
+					SpringConstraints.Apply(InParticles, Dt);
+				});
+			}
 		}
 	}
 	if (ChaosClothSimConfig->AreaStiffness)
 	{
 		TArray<Chaos::TVector<int32, 3>> SurfaceConstraints = SurfaceElements;
-		Evolution->AddPBDConstraintFunction([SurfConstraints = Chaos::TPBDAxialSpringConstraints<float, 3>(Evolution->Particles(), MoveTemp(SurfaceConstraints), ChaosClothSimConfig->AreaStiffness)](TPBDParticles<float, 3>& InParticles, const float Dt) {
-			SurfConstraints.Apply(InParticles, Dt);
-		});
+		if (bUseXPBDConstraints)
+		{
+			const TSharedPtr<Chaos::TXPBDAxialSpringConstraints<float, 3>> AxialSpringConstraints =
+				MakeShared<Chaos::TXPBDAxialSpringConstraints<float, 3>>(Evolution->Particles(), MoveTemp(SurfaceConstraints), ChaosClothSimConfig->AreaStiffness);
+			Evolution->AddXPBDConstraintFunctions(
+				[AxialSpringConstraints]()
+				{
+					AxialSpringConstraints->Init();
+				},
+				[AxialSpringConstraints](TPBDParticles<float, 3>& InParticles, const float Dt)
+				{
+					AxialSpringConstraints->Apply(InParticles, Dt);
+				});
+		}
+		else
+		{
+			Evolution->AddPBDConstraintFunction([AxialSpringConstraints = Chaos::TPBDAxialSpringConstraints<float, 3>(Evolution->Particles(), MoveTemp(SurfaceConstraints), ChaosClothSimConfig->AreaStiffness)](TPBDParticles<float, 3>& InParticles, const float Dt)
+			{
+				AxialSpringConstraints.Apply(InParticles, Dt);
+			});
+		}
 	}
 	if (ChaosClothSimConfig->VolumeStiffness)
 	{
@@ -489,16 +551,37 @@ void ClothingSimulation::AddConstraints(const UChaosClothConfig* ChaosClothSimCo
 		check(Mesh.GetNumElements() > 0);
 		// PerFormance note: The Per constraint version of this function is quite a bit faster for smaller assets
 		// There might be a cross-over point where the PerParticle version is faster: To be determined
-		Chaos::TPBDLongRangeConstraints<float, 3> PBDLongRangeConstraints(
-			Evolution->Particles(),
-			Mesh.GetPointToNeighborsMap(),
-			10, // The max number of connected neighbors per particle.  ryan - What should this be?  Was k...
-			ChaosClothSimConfig->StrainLimitingStiffness);
-
-		Evolution->AddPBDConstraintFunction([PBDLongRangeConstraints = MoveTemp(PBDLongRangeConstraints)](TPBDParticles<float, 3>& InParticles, const float Dt)
+		if (bUseXPBDConstraints)
 		{
-			PBDLongRangeConstraints.Apply(InParticles, Dt);
-		});
+			const TSharedPtr<Chaos::TXPBDLongRangeConstraints<float, 3>> LongRangeConstraints = MakeShared<Chaos::TXPBDLongRangeConstraints<float, 3>>(
+				Evolution->Particles(),
+				Mesh.GetPointToNeighborsMap(),
+				10, // The max number of connected neighbors per particle.  ryan - What should this be?  Was k...
+				ChaosClothSimConfig->StrainLimitingStiffness);
+
+			Evolution->AddXPBDConstraintFunctions(
+				[LongRangeConstraints]()
+				{
+					LongRangeConstraints->Init();
+				},
+				[LongRangeConstraints](TPBDParticles<float, 3>& InParticles, const float Dt)
+				{
+					LongRangeConstraints->Apply(InParticles, Dt);
+				});
+		}
+		else
+		{
+			Chaos::TPBDLongRangeConstraints<float, 3> LongRangeConstraints(
+				Evolution->Particles(),
+				Mesh.GetPointToNeighborsMap(),
+				10, // The max number of connected neighbors per particle.  ryan - What should this be?  Was k...
+				ChaosClothSimConfig->StrainLimitingStiffness);
+
+			Evolution->AddPBDConstraintFunction([LongRangeConstraints = MoveTemp(LongRangeConstraints)](TPBDParticles<float, 3>& InParticles, const float Dt)
+			{
+				LongRangeConstraints.Apply(InParticles, Dt);
+			});
+		}
 	}
 
 	// Maximum Distance Constraints
diff --git a/Engine/Plugins/Experimental/ChaosCloth/Source/Chaos/Public/ChaosCloth/ChaosClothConfig.h b/Engine/Plugins/Experimental/ChaosCloth/Source/Chaos/Public/ChaosCloth/ChaosClothConfig.h
--- a/Engine/Plugins/Experimental/ChaosCloth/Source/Chaos/Public/ChaosCloth/ChaosClothConfig.h
+++ b/Engine/Plugins/Experimental/ChaosCloth/Source/Chaos/Public/ChaosCloth/ChaosClothConfig.h
@@ -164,5 +164,9 @@ public:
 	// The gravitational acceleration vector [cm/s^2]
 	UPROPERTY(EditAnywhere, Category = Simulation, meta = (EditCondition = "bUseGravityOverride"))
 	FVector Gravity = { 0.f, 0.f, -980.665f };
-};
 
+	// Enable the XPBD constraints that resolve stiffness independently from the number of iterations
+	// Experimental, this feature might be removed without warning, not for production use
+	UPROPERTY(EditAnywhere, Category = "Experimental", meta = (DisplayName="Use XPBD Constraints (Experimental)"))
+	bool bUseXPBDConstraints = false;
+};
diff --git a/Engine/Source/Runtime/Experimental/Chaos/Private/Chaos/PBDEvolution.cpp b/Engine/Source/Runtime/Experimental/Chaos/Private/Chaos/PBDEvolution.cpp
--- a/Engine/Source/Runtime/Experimental/Chaos/Private/Chaos/PBDEvolution.cpp
+++ b/Engine/Source/Runtime/Experimental/Chaos/Private/Chaos/PBDEvolution.cpp
@@ -107,6 +107,12 @@ void TPBDEvolution<T, d>::AdvanceOneTimeStep(const T Dt)
 #if !COMPILE_WITHOUT_UNREAL_SUPPORT
 	TPBDCollisionSpringConstraints<T, d> SelfCollisionRule(MParticles, MCollisionTriangles, MDisabledCollisionElements, Dt, MSelfCollisionThickness, 1.5f);
 #endif
+
+	for (TFunction<void()>& InitConstraintRule : MInitConstraintRules)
+	{
+		InitConstraintRule();  // Clear XPBD's Lambdas
+	}
+
 	for (int i = 0; i < MNumIterations; ++i)
 	{
 		for (TFunction<void(TPBDParticles<T, d>&, const T)>& ConstraintRule : MConstraintRules)
diff --git a/Engine/Source/Runtime/Experimental/Chaos/Public/Chaos/PBDAxialSpringConstraintsBase.h b/Engine/Source/Runtime/Experimental/Chaos/Public/Chaos/PBDAxialSpringConstraintsBase.h
--- a/Engine/Source/Runtime/Experimental/Chaos/Public/Chaos/PBDAxialSpringConstraintsBase.h
+++ b/Engine/Source/Runtime/Experimental/Chaos/Public/Chaos/PBDAxialSpringConstraintsBase.h
@@ -100,8 +100,6 @@ class TPBDAxialSpringConstraintsBase
   protected:
 	TArray<TVector<int32, 3>> MConstraints;
 	TArray<T> MBarys;
-
-  private:
 	TArray<T> MDists;
 	T MStiffness;
 };
diff --git a/Engine/Source/Runtime/Experimental/Chaos/Public/Chaos/PBDEvolution.h b/Engine/Source/Runtime/Experimental/Chaos/Public/Chaos/PBDEvolution.h
--- a/Engine/Source/Runtime/Experimental/Chaos/Public/Chaos/PBDEvolution.h
+++ b/Engine/Source/Runtime/Experimental/Chaos/Public/Chaos/PBDEvolution.h
@@ -24,6 +24,7 @@ class CHAOS_API TPBDEvolution
 	void SetCollisionKinematicUpdateFunction(TFunction<void(TKinematicGeometryClothParticles<T, d>&, const T, const T, const int32)> KinematicUpdate) { MCollisionKinematicUpdate = KinematicUpdate; }
 	void SetParticleUpdateFunction(TFunction<void(TPBDParticles<T, d>&, const T)> ParticleUpdate) { MParticleUpdate = ParticleUpdate; }
 	void AddPBDConstraintFunction(TFunction<void(TPBDParticles<T, d>&, const T)> ConstraintFunction) { MConstraintRules.Add(ConstraintFunction); }
+	void AddXPBDConstraintFunctions(TFunction<void()> InitConstraintFunction, TFunction<void(TPBDParticles<T, d>&, const T)> ConstraintFunction) { MInitConstraintRules.Add(InitConstraintFunction); MConstraintRules.Add(ConstraintFunction); }
 	void AddForceFunction(TFunction<void(TPBDParticles<T, d>&, const T, const int32)> ForceFunction) { MForceRules.Add(ForceFunction); }
 
 	const TPBDParticles<T, d>& Particles() const { return MParticles; }
@@ -56,7 +57,7 @@ class CHAOS_API TPBDEvolution
 
 	T GetTime() const { return MTime; }
 
-	void ResetConstraintRules() { MConstraintRules.Reset(); };
+	void ResetConstraintRules() { MInitConstraintRules.Reset(); MConstraintRules.Reset(); };
 	void ResetSelfCollision() { MCollisionTriangles.Reset(); MDisabledCollisionElements.Reset(); };
 
   private:
@@ -75,6 +76,7 @@ class CHAOS_API TPBDEvolution
 	FGravityForces GravityForces;
 
 	TArray<TFunction<void(TPBDParticles<T, d>&, const T, const int32)>> MForceRules;
+	TArray<TFunction<void()>> MInitConstraintRules;
 	TArray<TFunction<void(TPBDParticles<T, d>&, const T)>> MConstraintRules;
 	TFunction<void(TPBDParticles<T, d>&, const T)> MParticleUpdate;
 	TFunction<void(TPBDParticles<T, d>&, const T, const T, const int32)> MKinematicUpdate;
diff --git a/Engine/Source/Runtime/Experimental/Chaos/Public/Chaos/PBDLongRangeConstraintsBase.h b/Engine/Source/Runtime/Experimental/Chaos/Public/Chaos/PBDLongRangeConstraintsBase.h
--- a/Engine/Source/Runtime/Experimental/Chaos/Public/Chaos/PBDLongRangeConstraintsBase.h
+++ b/Engine/Source/Runtime/Experimental/Chaos/Public/Chaos/PBDLongRangeConstraintsBase.h
@@ -49,7 +49,7 @@ class CHAOS_API TPBDLongRangeConstraintsBase
 		return Delta;
 	};
 
-  private:
+  protected:
 	static TArray<TArray<uint32>> ComputeIslands(const TDynamicParticles<T, d>& InParticles, const TMap<int32, TSet<uint32>>& PointToNeighbors,/*const TTriangleMesh<T>& Mesh,*/ const TArray<uint32>& KinematicParticles);
 	void ComputeEuclidianConstraints(const TDynamicParticles<T, d>& InParticles, const TMap<int32, TSet<uint32>>& PointToNeighbors,/*const TTriangleMesh<T>& Mesh,*/ const int32 NumberOfAttachments);
 	void ComputeGeodesicConstraints(const TDynamicParticles<T, d>& InParticles, const TMap<int32, TSet<uint32>>& PointToNeighbors,/*const TTriangleMesh<T>& Mesh,*/ const int32 NumberOfAttachments);
@@ -77,8 +77,6 @@ class CHAOS_API TPBDLongRangeConstraintsBase
 
   protected:
 	TArray<TArray<uint32>> MConstraints;
-
-  private:
 	TArray<T> MDists;
 	T MStiffness;
 };
diff --git a/Engine/Source/Runtime/Experimental/Chaos/Public/Chaos/PBDSpringConstraintsBase.h b/Engine/Source/Runtime/Experimental/Chaos/Public/Chaos/PBDSpringConstraintsBase.h
--- a/Engine/Source/Runtime/Experimental/Chaos/Public/Chaos/PBDSpringConstraintsBase.h
+++ b/Engine/Source/Runtime/Experimental/Chaos/Public/Chaos/PBDSpringConstraintsBase.h
@@ -160,8 +160,6 @@ namespace Chaos
 
 	protected:
 		TArray<TVector<int32, 2>> MConstraints;
-
-	private:
 		TArray<T> MDists;
 		T MStiffness;
 	};
diff --git a/Engine/Source/Runtime/Experimental/Chaos/Public/Chaos/XPBDAxialSpringConstraints.h b/Engine/Source/Runtime/Experimental/Chaos/Public/Chaos/XPBDAxialSpringConstraints.h
new file mode 100644
--- /dev/null
+++ b/Engine/Source/Runtime/Experimental/Chaos/Public/Chaos/XPBDAxialSpringConstraints.h
@@ -0,0 +1,104 @@
+// Copyright 1998-2019 Epic Games, Inc. All Rights Reserved.
+#pragma once
+
+#include "Chaos/Array.h"
+#include "Chaos/PBDAxialSpringConstraintsBase.h"
+#include "Chaos/PBDParticles.h"
+#include "Chaos/PerParticleRule.h"
+#include "ChaosStats.h"
+
+#include <algorithm>
+
+DECLARE_CYCLE_STAT(TEXT("Chaos XPBD Axial Spring Constraint"), STAT_XPBD_AxialSpring, STATGROUP_Chaos);
+namespace Chaos
+{
+// Stiffness is in N/CM^2, so it needs to be adjusted from the PBD stiffness ranging between [0,1]
+static const double XPBDAxialSpringMaxCompliance = 1e-7;  // Max stiffness: 1e+11 N/M^2 = 1e+7 N/CM^2 -> Max compliance: 1e-7 CM^2/N
+
+template<class T, int d>
+class TXPBDAxialSpringConstraints : public TParticleRule<T, d>, public TPBDAxialSpringConstraintsBase<T, d>
+{
+	typedef TPBDAxialSpringConstraintsBase<T, d> Base;
+	using Base::MBarys;
+	using Base::MConstraints;
+
+public:
+	TXPBDAxialSpringConstraints(const TDynamicParticles<T, d>& InParticles, TArray<TVector<int32, 3>>&& Constraints, const T Stiffness = (T)1.)
+	    : TPBDAxialSpringConstraintsBase<T, d>(InParticles, MoveTemp(Constraints), Stiffness)
+	{ MLambdas.Init(0.f, MConstraints.Num()); }
+
+	virtual ~TXPBDAxialSpringConstraints() {}
+
+	void Init() const { for (T& Lambda : MLambdas) { Lambda = (T)0.; } }
+
+	virtual void Apply(TPBDParticles<T, d>& InParticles, const T Dt) const override //-V762
+	{
+		SCOPE_CYCLE_COUNTER(STAT_XPBD_AxialSpring);
+		for (int32 i = 0; i < MConstraints.Num(); ++i)
+		{
+			const TVector<int32, 3>& constraint = MConstraints[i];
+			const int32 i1 = constraint[0];
+			const int32 i2 = constraint[1];
+			const int32 i3 = constraint[2];
+			const TVector<T, d> Delta = GetDelta(InParticles, Dt, i);
+			const T Multiplier = (T)2. / (FMath::Max(MBarys[i], (T)1. - MBarys[i]) + (T)1.);
+			if (InParticles.InvM(i1) > 0)
+			{
+				InParticles.P(i1) -= Multiplier * InParticles.InvM(i1) * Delta;
+			}
+			if (InParticles.InvM(i2))
+			{
+				InParticles.P(i2) += Multiplier * InParticles.InvM(i2) * MBarys[i] * Delta;
+			}
+			if (InParticles.InvM(i3))
+			{
+				InParticles.P(i3) += Multiplier * InParticles.InvM(i3) * (1 - MBarys[i]) * Delta;
+			}
+		}
+	}
+
+private:
+	inline TVector<T, d> GetDelta(const TPBDParticles<T, d>& InParticles, const T Dt, const int32 InConstraintIndex) const
+	{
+		const TVector<int32, 3>& Constraint = MConstraints[InConstraintIndex];
+		const int32 i1 = Constraint[0];
+		const int32 i2 = Constraint[1];
+		const int32 i3 = Constraint[2];
+
+		const T Bary = MBarys[InConstraintIndex];
+		const T PInvMass = InParticles.InvM(i3) * ((T)1. - Bary) + InParticles.InvM(i2) * Bary;
+		if (InParticles.InvM(i1) == (T)0. && PInvMass == (T)0.)
+		{
+			return TVector<T, d>((T)0.);
+		}
+		const T CombinedInvMass = PInvMass + InParticles.InvM(i1);
+		ensure(CombinedInvMass > (T)SMALL_NUMBER);
+
+		const TVector<T, d>& P1 = InParticles.P(i1);
+		const TVector<T, d>& P2 = InParticles.P(i2);
+		const TVector<T, d>& P3 = InParticles.P(i3);
+		const TVector<T, d> P = (P2 - P3) * Bary + P3;
+
+		const TVector<T, d> Difference = P1 - P;
+		const T Distance = Difference.Size();
+		if (UNLIKELY(Distance <= SMALL_NUMBER))
+		{
+			return TVector<T, d>((T)0.);
+		}
+		const TVector<T, d> Direction = Difference / Distance;
+		const T Offset = (Distance - MDists[InConstraintIndex]);
+
+		T& Lambda = MLambdas[InConstraintIndex];
+		const T Alpha = (T)XPBDAxialSpringMaxCompliance / (MStiffness * Dt * Dt);
+
+		const T DLambda = (Offset - Alpha * Lambda) / (CombinedInvMass + Alpha);
+		const TVector<T, d> Delta = DLambda * Direction;
+		Lambda += DLambda;
+
+		return Delta;
+	}
+
+private:
+	mutable TArray<T> MLambdas;
+};
+}
diff --git a/Engine/Source/Runtime/Experimental/Chaos/Public/Chaos/XPBDLongRangeConstraints.h b/Engine/Source/Runtime/Experimental/Chaos/Public/Chaos/XPBDLongRangeConstraints.h
new file mode 100644
--- /dev/null
+++ b/Engine/Source/Runtime/Experimental/Chaos/Public/Chaos/XPBDLongRangeConstraints.h
@@ -0,0 +1,87 @@
+// Copyright 1998-2019 Epic Games, Inc. All Rights Reserved.
+#pragma once
+
+#include "Chaos/PBDLongRangeConstraintsBase.h"
+#include "Chaos/PBDParticles.h"
+#include "Chaos/PBDConstraintContainer.h"
+#include "ChaosStats.h"
+
+DECLARE_CYCLE_STAT(TEXT("Chaos XPBD Long Range Constraint"), STAT_XPBD_LongRange, STATGROUP_Chaos);
+
+namespace Chaos
+{
+// Stiffness is in N/CM^2, so it needs to be adjusted from the PBD stiffness ranging between [0,1]
+static const double XPBDLongRangeMaxCompliance = 1.e-3;
+
+template<class T, int d>
+class TXPBDLongRangeConstraints : public TPBDLongRangeConstraintsBase<T, d>, public FPBDConstraintContainer
+{
+	typedef TPBDLongRangeConstraintsBase<T, d> Base;
+	using Base::MConstraints;
+
+public:
+	TXPBDLongRangeConstraints(const TDynamicParticles<T, d>& InParticles, const TMap<int32, TSet<uint32>>& PointToNeighbors, const int32 NumberOfAttachments = 1, const T Stiffness = (T)1)
+	    : TPBDLongRangeConstraintsBase<T, d>(InParticles, PointToNeighbors, NumberOfAttachments, Stiffness)
+	{ MLambdas.Init(0.f, MConstraints.Num()); }
+
+	virtual ~TXPBDLongRangeConstraints() {}
+
+	void Init() const { for (T& Lambda : MLambdas) { Lambda = (T)0.; } }
+
+	void Apply(TPBDParticles<T, d>& InParticles, const T Dt, int32 Index) const
+	{
+		const auto& Constraint = MConstraints[Index];
+		int32 i2 = Constraint[Constraint.Num() - 1];
+		check(InParticles.InvM(i2) > 0);
+		InParticles.P(i2) += GetDelta(InParticles, Dt, Index);
+	}
+
+	void Apply(TPBDParticles<T, d>& InParticles, const T Dt) const 
+	{
+		SCOPE_CYCLE_COUNTER(STAT_XPBD_LongRange);
+		for (int32 i = 0; i < MConstraints.Num(); ++i)
+		{
+			Apply(InParticles, Dt, i);
+		}
+	}
+
+	void Apply(TPBDParticles<T, d>& InParticles, const T Dt, const TArray<int32>& InConstraintIndices) const
+	{
+		SCOPE_CYCLE_COUNTER(STAT_XPBD_LongRange);
+		for (int32 i : InConstraintIndices)
+		{
+			Apply(InParticles, Dt, i);
+		}
+	}
+
+private:
+	inline TVector<T, d> GetDelta(const TPBDParticles<T, d>& InParticles, const T Dt, const int32 InConstraintIndices) const
+	{
+		const TArray<uint32>& Constraint = MConstraints[InConstraintIndices];
+		check(Constraint.Num() > 1);
+		const uint32 i1 = Constraint[0];
+		const uint32 i2 = Constraint[Constraint.Num() - 1];
+		const uint32 i2m1 = Constraint[Constraint.Num() - 2];
+		check(InParticles.InvM(i1) == 0);
+		check(InParticles.InvM(i2) > 0);
+		const T Distance = ComputeGeodesicDistance(InParticles, Constraint); // This function is used for either Euclidean or Geodesic distances
+		if (Distance < MDists[InConstraintIndices]) { return TVector<T, d>(0); }
+		const T Offset = Distance - MDists[InConstraintIndices];
+
+		TVector<T, d> Direction = InParticles.P(i2m1) - InParticles.P(i2);
+		Direction.SafeNormalize();
+
+		T& Lambda = MLambdas[InConstraintIndices];
+		const T Alpha = (T)XPBDLongRangeMaxCompliance / (MStiffness * Dt * Dt);
+
+		const T DLambda = (Offset - Alpha * Lambda) / ((T)1. + Alpha);
+		const TVector<T, d> Delta = DLambda * Direction;
+		Lambda += DLambda;
+
+		return Delta;
+	};
+
+private:
+	mutable TArray<T> MLambdas;
+};
+}
diff --git a/Engine/Source/Runtime/Experimental/Chaos/Public/Chaos/XPBDSpringConstraints.h b/Engine/Source/Runtime/Experimental/Chaos/Public/Chaos/XPBDSpringConstraints.h
new file mode 100644
--- /dev/null
+++ b/Engine/Source/Runtime/Experimental/Chaos/Public/Chaos/XPBDSpringConstraints.h
@@ -0,0 +1,126 @@
+// Copyright 1998-2019 Epic Games, Inc. All Rights Reserved.
+#pragma once
+
+#include "Chaos/Array.h"
+#include "Chaos/PBDParticles.h"
+#include "Chaos/PBDSpringConstraintsBase.h"
+#include "Chaos/PBDConstraintContainer.h"
+#include "ChaosStats.h"
+
+#include "Templates/EnableIf.h"
+
+DECLARE_CYCLE_STAT(TEXT("Chaos XPBD Spring Constraint"), STAT_XPBD_Spring, STATGROUP_Chaos);
+
+namespace Chaos
+{
+// Stiffness is in N/CM^2, so it needs to be adjusted from the PBD stiffness ranging between [0,1]
+static const double XPBDSpringMaxCompliance = 1e-7;  // Max stiffness: 1e+11 N/M^2 = 1e+7 N/CM^2 -> Max compliance: 1e-7 CM^2/N
+
+template<class T, int32 d>
+class TXPBDSpringConstraints : public TPBDSpringConstraintsBase<T, d>, public FPBDConstraintContainer
+{
+	typedef TPBDSpringConstraintsBase<T, d> Base;
+	using Base::MConstraints;
+
+public:
+	TXPBDSpringConstraints(const T Stiffness = (T)1)
+	    : TPBDSpringConstraintsBase<T, d>(Stiffness * Stiffness * (T)XPBDSpringMaxStiffness)
+	{}
+
+	TXPBDSpringConstraints(const TDynamicParticles<T, d>& InParticles, TArray<TVector<int32, 2>>&& Constraints, const T Stiffness = (T)1.)
+	    : TPBDSpringConstraintsBase<T, d>(InParticles, MoveTemp(Constraints), Stiffness)
+	{ MLambdas.Init(0.f, MConstraints.Num()); }
+
+	TXPBDSpringConstraints(const TRigidParticles<T, d>& InParticles, TArray<TVector<int32, 2>>&& Constraints, const T Stiffness = (T)1.)
+	    : TPBDSpringConstraintsBase<T, d>(InParticles, MoveTemp(Constraints), Stiffness)
+	{ MLambdas.Init(0.f, MConstraints.Num()); }
+
+	TXPBDSpringConstraints(const TDynamicParticles<T, d>& InParticles, const TArray<TVector<int32, 3>>& Constraints, const T Stiffness = (T)1.)
+	    : TPBDSpringConstraintsBase<T, d>(InParticles, Constraints, Stiffness)
+	{ MLambdas.Init(0.f, MConstraints.Num()); }
+
+	TXPBDSpringConstraints(const TDynamicParticles<T, d>& InParticles, const TArray<TVector<int32, 4>>& Constraints, const T Stiffness = (T)1.)
+	    : TPBDSpringConstraintsBase<T, d>(InParticles, Constraints, Stiffness)
+	{ MLambdas.Init(0.f, MConstraints.Num()); }
+
+	virtual ~TXPBDSpringConstraints() {}
+
+	const TArray<TVector<int32, 2>>& GetConstraints() const { return MConstraints; }
+	TArray<TVector<int32, 2>>& GetConstraints() { return MConstraints; }
+	TArray<TVector<int32, 2>>& Constraints() { return MConstraints; }
+
+	void Init() const { for (T& Lambda : MLambdas) { Lambda = (T)0.; } }
+
+	void Apply(TPBDParticles<T, d>& InParticles, const T Dt, const int32 InConstraintIndex) const
+	{
+		const int32 i = InConstraintIndex;
+		{
+			const auto& Constraint = MConstraints[i];
+			const int32 i1 = Constraint[0];
+			const int32 i2 = Constraint[1];
+			const TVector<T, d> Delta = GetDelta(InParticles, Dt, i);
+			if (InParticles.InvM(i1) > 0)
+			{
+				InParticles.P(i1) -= InParticles.InvM(i1) * Delta;
+			}
+			if (InParticles.InvM(i2) > 0)
+			{
+				InParticles.P(i2) += InParticles.InvM(i2) * Delta;
+			}
+		}
+	}
+
+	void Apply(TPBDParticles<T, d>& InParticles, const T Dt) const
+	{
+		SCOPE_CYCLE_COUNTER(STAT_XPBD_Spring);
+		for (int32 i = 0; i < MConstraints.Num(); ++i)
+		{
+			Apply(InParticles, Dt, i);
+		}
+	}
+
+	void Apply(TPBDRigidParticles<T, d>& InParticles, const T Dt, const TArray<int32>& InConstraintIndices) const
+	{
+		SCOPE_CYCLE_COUNTER(STAT_XPBD_Spring);
+		for (int32 i : InConstraintIndices)
+		{
+			const auto& Constraint = MConstraints[i];
+			const int32 i1 = Constraint[0];
+			const int32 i2 = Constraint[1];
+			check(InParticles.Island(i1) == InParticles.Island(i2) || InParticles.Island(i1) == INDEX_NONE || InParticles.Island(i2) == INDEX_NONE);
+			Apply(InParticles, Dt, i);
+		}
+	}
+
+private:
+	inline TVector<T, d> GetDelta(const TPBDParticles<T, d>& InParticles, const T Dt, const int32 InConstraintIndex) const
+	{
+		const TVector<int32, 2>& Constraint = MConstraints[InConstraintIndex];
+
+		const int32 i1 = Constraint[0];
+		const int32 i2 = Constraint[1];
+
+		if (InParticles.InvM(i2) == 0 && InParticles.InvM(i1) == 0)
+			return TVector<T, d>(0);
+		const T CombinedInvMass = InParticles.InvM(i2) + InParticles.InvM(i1);
+
+		const TVector<T, d>& P1 = InParticles.P(i1);
+		const TVector<T, d>& P2 = InParticles.P(i2);
+		TVector<T, d> Direction = P1 - P2;
+		const T Distance = Direction.SafeNormalize();
+		const T Offset = Distance - MDists[InConstraintIndex];
+
+		T& Lambda = MLambdas[InConstraintIndex];
+		const T Alpha = (T)XPBDSpringMaxCompliance / (MStiffness * Dt * Dt);
+
+		const T DLambda = (Offset - Alpha * Lambda) / (CombinedInvMass + Alpha);
+		const TVector<T, d> Delta = DLambda * Direction;
+		Lambda += DLambda;
+
+		return Delta;
+	}
+
+private:
+	mutable TArray<T> MLambdas;
+};
+}
