Chaos Cloth - Fix time related issues while removing redundant solver substepping and SolverFrequency parameter.
Note: the previous context refactor actually fixed the hitch/jiggling related issue.

diff --git a/Engine/Plugins/Experimental/ChaosCloth/Source/Chaos/Private/ChaosCloth/ChaosClothConfig.cpp b/Engine/Plugins/Experimental/ChaosCloth/Source/Chaos/Private/ChaosCloth/ChaosClothConfig.cpp
--- a/Engine/Plugins/Experimental/ChaosCloth/Source/Chaos/Private/ChaosCloth/ChaosClothConfig.cpp
+++ b/Engine/Plugins/Experimental/ChaosCloth/Source/Chaos/Private/ChaosCloth/ChaosClothConfig.cpp
@@ -73,8 +73,6 @@ void UChaosClothSharedSimConfig::MigrateFrom(const FClothConfig_Legacy& ClothCon
 {
 	IterationCount = FMath::Clamp(int32(ClothConfig.SolverFrequency / 60.f), 1, 100);
 
-	SolverFrequency = ClothConfig.SolverFrequency;
-
 	SelfCollisionThickness = FMath::Clamp(ClothConfig.SelfCollisionRadius, 0.f, 1000.f);
 
 	CollisionThickness = FMath::Clamp(ClothConfig.CollisionThickness, 0.f, 1000.f);
diff --git a/Engine/Plugins/Experimental/ChaosCloth/Source/Chaos/Private/ChaosCloth/ChaosClothingSimulation.cpp b/Engine/Plugins/Experimental/ChaosCloth/Source/Chaos/Private/ChaosCloth/ChaosClothingSimulation.cpp
--- a/Engine/Plugins/Experimental/ChaosCloth/Source/Chaos/Private/ChaosCloth/ChaosClothingSimulation.cpp
+++ b/Engine/Plugins/Experimental/ChaosCloth/Source/Chaos/Private/ChaosCloth/ChaosClothingSimulation.cpp
@@ -59,7 +59,6 @@ namespace ClothingSimulationDefault
 	static const float CollisionThickness = 1.2f;
 	static const float CoefficientOfFriction = 0.f;
 	static const float Damping = 0.01f;
-	static const float SolverFrequency = 60.f;
 	static const float WindDrag = 0.f;
 }
 
@@ -1170,21 +1169,9 @@ void ClothingSimulation::Simulate(IClothingSimulationContext* InContext)
 
 	// Advance Sim
 	DeltaTime = Context->DeltaSeconds;
-
-	const float SolverFrequency = ClothSharedSimConfig ?
-		ClothSharedSimConfig->SolverFrequency :
-		ClothingSimulationDefault::SolverFrequency;
-
-	const int32 NumSubSteps = FMath::Max(1, FMath::RoundToInt(DeltaTime * SolverFrequency));  // Start substepping once deltatime is more than 50% above the expected simulation deltatime
-	const float SubDeltaTime = DeltaTime / float(NumSubSteps);
-
-	for (int32 i = 0; i < NumSubSteps; ++i)
-	{
-		Evolution->AdvanceOneTimeStep(SubDeltaTime);
-	}
-	UE_LOG(LogChaosCloth, VeryVerbose, TEXT("DeltaTime: %.6f, SubDeltaTime: %.6f, NumSubSteps = %d"), DeltaTime, SubDeltaTime, NumSubSteps);
-
+	Evolution->AdvanceOneTimeStep(DeltaTime);
 	Time += DeltaTime;
+	UE_LOG(LogChaosCloth, VeryVerbose, TEXT("DeltaTime: %.6f, Time = %.6f,  MaxPhysicsDelta = %.6f"), DeltaTime, Time, FClothingSimulationCommon::MaxPhysicsDelta);
 }
 
 void ClothingSimulation::GetSimulationData(
diff --git a/Engine/Plugins/Experimental/ChaosCloth/Source/Chaos/Public/ChaosCloth/ChaosClothConfig.h b/Engine/Plugins/Experimental/ChaosCloth/Source/Chaos/Public/ChaosCloth/ChaosClothConfig.h
--- a/Engine/Plugins/Experimental/ChaosCloth/Source/Chaos/Public/ChaosCloth/ChaosClothConfig.h
+++ b/Engine/Plugins/Experimental/ChaosCloth/Source/Chaos/Public/ChaosCloth/ChaosClothConfig.h
@@ -136,11 +136,6 @@ public:
 	UPROPERTY(EditAnywhere, Category = Simulation, meta = (UIMin = "0", UIMax = "20", ClampMin = "0", ClampMax = "100"))
 	int32 IterationCount = 1;
 
-	// The expected simulation frequency in Hz that, if not closely matching the game tick frequency, will trigger substepping to help keep a more consistent time step.
-	// Use higher values to avoid jitter, and smaller value to avoid the extra cost of substepping.
-	UPROPERTY(EditAnywhere, Category = Simulation, meta = (UIMin = "30", UIMax = "240", ClampMin = "1", ClampMax = "1080"))
-	float SolverFrequency = 60.f;
-
 	// The radius of the spheres used in self collision 
 	UPROPERTY(EditAnywhere, Category = Collision, meta = (UIMin = "0", UIMax = "100", ClampMin = "0", ClampMax = "1000"))
 	float SelfCollisionThickness = 2.0f;
diff --git a/Engine/Source/Runtime/ClothingSystemRuntimeCommon/Public/ClothingSimulation.h b/Engine/Source/Runtime/ClothingSystemRuntimeCommon/Public/ClothingSimulation.h
--- a/Engine/Source/Runtime/ClothingSystemRuntimeCommon/Public/ClothingSimulation.h
+++ b/Engine/Source/Runtime/ClothingSystemRuntimeCommon/Public/ClothingSimulation.h
@@ -61,7 +61,7 @@ protected:
 	/** Fills in the base data for a clothing simulation */
 	virtual void FillContext(USkeletalMeshComponent* InComponent, float InDeltaTime, IClothingSimulationContext* InOutContext) override;
 
-private:
+protected:
 	/** Maximum physics time, incoming deltas will be clamped down to this value on long frames */
 	float MaxPhysicsDelta;
 };
