Chaos Cloth - Add missing Chaos Cloth teleport mode.
- Move teleport property from the nv cloth context to the common cloth context.
- Fix Chaos Cloth particle and collision particle initializations.

diff --git a/Engine/Plugins/Experimental/ChaosCloth/Source/Chaos/Private/ChaosCloth/ChaosClothingSimulation.cpp b/Engine/Plugins/Experimental/ChaosCloth/Source/Chaos/Private/ChaosCloth/ChaosClothingSimulation.cpp
--- a/Engine/Plugins/Experimental/ChaosCloth/Source/Chaos/Private/ChaosCloth/ChaosClothingSimulation.cpp
+++ b/Engine/Plugins/Experimental/ChaosCloth/Source/Chaos/Private/ChaosCloth/ChaosClothingSimulation.cpp
@@ -139,6 +139,7 @@ void ClothingSimulation::Shutdown()
 	AnimationPositions.Reset();
 	AnimationNormals.Reset();
 	IndexToRangeMap.Reset();
+	RootBoneWorldTransforms.Reset();
 	Meshes.Reset();
 	FaceNormals.Reset();
 	PointNormals.Reset();
@@ -200,8 +201,10 @@ void ClothingSimulation::CreateActor(USkeletalMeshComponent* InOwnerComponent, U
 	if (IndexToRangeMap.Num() <= InSimDataIndex)
 	{
 		IndexToRangeMap.SetNum(InSimDataIndex + 1);
+		RootBoneWorldTransforms.SetNum(InSimDataIndex + 1);
 	}
 	IndexToRangeMap[InSimDataIndex] = Chaos::TVector<uint32, 2>(Offset, Particles.Size());
+	RootBoneWorldTransforms[InSimDataIndex] = Context.BoneTransforms[Asset->ReferenceBoneIndex] * Context.ComponentToWorld;
 
 	// ClothSharedSimConfig should either be a nullptr, or point to an object common to the whole skeletal mesh
 	if (ClothSharedSimConfig == nullptr)
@@ -262,6 +265,9 @@ void ClothingSimulation::CreateActor(USkeletalMeshComponent* InOwnerComponent, U
 		ExternalCollisions.Convexes.Num() == 0 &&
 		ExternalCollisions.Boxes.Num() == 0, TEXT("There cannot be any external collisions added before all the cloth assets collisions are processed."));
 	ExtractCollisions(Asset);
+
+	// Update collision transforms, including initial state for particles' X & R
+	UpdateCollisionTransforms(Context, /*bReinit =*/ false);  // Reinit is false, since we only need to update the new collision particles
 }
 
 void ClothingSimulation::ExtractCollisions(const UClothingAssetCommon* Asset)
@@ -325,7 +331,7 @@ void ClothingSimulation::ResetParticles(int32 InSimDataIndex)
 
 	for (uint32 i = Offset; i < Range; ++i)
 	{
-		Evolution->Particles().X(i) = AnimationPositions[i];
+		Evolution->Particles().P(i) = Evolution->Particles().X(i) = AnimationPositions[i];
 		Evolution->Particles().V(i) = Chaos::TVector<float, 3>(0.f);
 		Evolution->Particles().M(i) = 0.f;
 	}
@@ -649,7 +655,7 @@ void ClothingSimulation::AddSelfCollisions(int32 InSimDataIndex)
 	}
 }
 
-void ClothingSimulation::UpdateCollisionTransforms(const ClothingSimulationContext& Context)
+void ClothingSimulation::UpdateCollisionTransforms(const ClothingSimulationContext& Context, bool bReinit)
 {
 	SCOPE_CYCLE_COUNTER(STAT_ChaosClothUpdateCollisionTransforms);
 
@@ -685,13 +691,28 @@ void ClothingSimulation::UpdateCollisionTransforms(const ClothingSimulationConte
 		}
 	}
 
-	// Set the new collision particles and transforms initial state
-	for (int32 Index = PrevNumCollisions; Index < NumCollisions; ++Index)
+	// External collisions are reinitialized at every frame, but don't affect the size of the arrays if the same amount has been added/removed
+	const int32 Offset = FMath::Min((int32)ExternalCollisionsOffset, PrevNumCollisions);
+
+	// Reinit old collisions particles and transforms
+	if (bReinit)
 	{
-		CollisionParticles.X(Index) = CollisionTransforms[Index].GetTranslation();
-		CollisionParticles.R(Index) = CollisionTransforms[Index].GetRotation();
+		for (int32 Index = 0; Index < Offset; ++Index)
+		{
+			const Chaos::TRigidTransform<float, 3>& CollisionTransform = CollisionTransforms[Index];
+			CollisionParticles.X(Index) = CollisionTransform.GetTranslation();
+			CollisionParticles.R(Index) = CollisionTransform.GetRotation();
+			OldCollisionTransforms[Index] = CollisionTransform;
+		}
+	}
 
-		OldCollisionTransforms[Index] = CollisionTransforms[Index];
+	// Set the new collision particles and transforms initial state
+	for (int32 Index = Offset; Index < NumCollisions; ++Index)
+	{
+		const Chaos::TRigidTransform<float, 3>& CollisionTransform = CollisionTransforms[Index];
+		CollisionParticles.X(Index) = CollisionTransform.GetTranslation();
+		CollisionParticles.R(Index) = CollisionTransform.GetRotation();
+		OldCollisionTransforms[Index] = CollisionTransform;
 	}
 }
 
@@ -1093,8 +1114,6 @@ void ClothingSimulation::AddCollisions(const FClothCollisionData& ClothCollision
 		for (uint32 i = Offset; i < CollisionParticles.Size(); ++i)
 		{
 			const FClothCollisionPrim_Box& Box = ClothCollisionData.Boxes[i - Offset];
-			CollisionParticles.X(i) = Chaos::TVector<float, 3>(0.f);
-			CollisionParticles.R(i) = Chaos::TRotation<float, 3>::FromIdentity();
 			
 			BaseTransforms[i] = Chaos::TRigidTransform<float, 3>(Box.LocalPosition, Box.LocalRotation);
 			
@@ -1127,13 +1146,20 @@ void ClothingSimulation::Simulate(IClothingSimulationContext* InContext)
 	static const float UnitConversionScale = 100.f;  // It would seem that in game wind velocity is in m/s and not cm/s as seen in NvClothSupport::Constants::UnitConversionScale
 	Evolution->GetVelocityField().SetVelocity(Context->WindVelocity * UnitConversionScale);
 
-	// Get New Animation Positions and Normals
+	// Check teleport modes
+	const bool bTeleport = (Context->TeleportMode > EClothingTeleportMode::None);
+	const bool bTeleportAndReset = (Context->TeleportMode == EClothingTeleportMode::TeleportAndReset);
+	UE_CLOG(bTeleport, LogChaosCloth, Verbose, TEXT("Teleport, reset: %d"), bTeleportAndReset);
+
+	// Get New Animation Positions and Normals + deal with teleportation
 	{
 		SCOPE_CYCLE_COUNTER(STAT_ChaosClothGetAnimationData);
 
 		Swap(OldAnimationPositions, AnimationPositions);
 		check(OldAnimationPositions.Num() == AnimationPositions.Num());
-		
+
+		TPBDParticles<float, 3>& Particles = Evolution->Particles();
+
 		for (int32 Index = 0; Index < IndexToRangeMap.Num(); ++Index)
 		{
 			const UClothingAssetCommon* const Asset = Assets[Index];
@@ -1161,11 +1187,42 @@ void ClothingSimulation::Simulate(IClothingSimulationContext* InContext)
 				Context->RefToLocals.Num(),
 				reinterpret_cast<TArray<FVector>&>(AnimationPositions),
 				reinterpret_cast<TArray<FVector>&>(AnimationNormals), Offset);
+
+			// Teleport
+			const FTransform RootBoneTransform = Context->BoneTransforms[Asset->ReferenceBoneIndex];
+			const FTransform RootBoneWorldTransform = RootBoneTransform * Context->ComponentToWorld;
+
+			if (bTeleport)
+			{
+				UE_LOG(LogChaosCloth, Verbose, TEXT("Teleport before: %s, after: %s"), *RootBoneWorldTransforms[Index].ToString(), *RootBoneWorldTransform.ToString());
+				const uint32 Range = IndexToRangeMap[Index][1];
+				if (bTeleportAndReset)
+				{
+					// Teleport & reset
+					for (uint32 i = Offset; i < Range; ++i)
+					{
+						Particles.P(i) = Particles.X(i) = AnimationPositions[i];
+						Particles.V(i) = Chaos::TVector<float, 3>(0.f);
+					}
+				}
+				else
+				{
+					// Teleport only
+					const FTransform Delta = RootBoneWorldTransforms[Index].GetRelativeTransformReverse(RootBoneWorldTransform);
+					for (uint32 i = Offset; i < Range; ++i)
+					{
+						Particles.X(i) = Delta.TransformPositionNoScale(Particles.X(i));
+						Particles.V(i) = Delta.TransformVectorNoScale(Particles.V(i));
+					}
+				}
+			}
+
+			RootBoneWorldTransforms[Index] = RootBoneWorldTransform;
 		}
 	}
 
 	// Update collision transforms
-	UpdateCollisionTransforms(*Context);
+	UpdateCollisionTransforms(*Context, bTeleport);
 
 	// Advance Sim
 	DeltaTime = Context->DeltaSeconds;
diff --git a/Engine/Plugins/Experimental/ChaosCloth/Source/Chaos/Public/ChaosCloth/ChaosClothingSimulation.h b/Engine/Plugins/Experimental/ChaosCloth/Source/Chaos/Public/ChaosCloth/ChaosClothingSimulation.h
--- a/Engine/Plugins/Experimental/ChaosCloth/Source/Chaos/Public/ChaosCloth/ChaosClothingSimulation.h
+++ b/Engine/Plugins/Experimental/ChaosCloth/Source/Chaos/Public/ChaosCloth/ChaosClothingSimulation.h
@@ -94,8 +94,8 @@ namespace Chaos
 		void ExtractLegacyAssetCollisions(const UClothingAssetCommon* Asset);
 		// Add collisions from a ClothCollisionData structure
 		void AddCollisions(const FClothCollisionData& ClothCollisionData, const TArray<int32>& UsedBoneIndices);
-		// Update the collision transforms using the specified context
-		void UpdateCollisionTransforms(const ClothingSimulationContext& Context);
+		// Update the collision transforms using the specified context, also update old collision if bReinit is true
+		void UpdateCollisionTransforms(const ClothingSimulationContext& Context, bool bReinit);
 		// Return the correct bone index based on the asset used bone index array
 		FORCEINLINE int32 GetMappedBoneIndex(const TArray<int32>& UsedBoneIndices, int32 BoneIndex)
 		{
@@ -126,6 +126,7 @@ namespace Chaos
 
 		// Sim Data
 		TArray<Chaos::TVector<uint32, 2>> IndexToRangeMap;
+		TArray<FTransform> RootBoneWorldTransforms;  // Used for teleportation
 
 		TArray<TUniquePtr<Chaos::TTriangleMesh<float>>> Meshes;
 		mutable TArray<TArray<Chaos::TVector<float, 3>>> FaceNormals;
diff --git a/Engine/Source/Runtime/ClothingSystemRuntimeCommon/Private/ClothingSimulation.cpp b/Engine/Source/Runtime/ClothingSystemRuntimeCommon/Private/ClothingSimulation.cpp
--- a/Engine/Source/Runtime/ClothingSystemRuntimeCommon/Private/ClothingSimulation.cpp
+++ b/Engine/Source/Runtime/ClothingSystemRuntimeCommon/Private/ClothingSimulation.cpp
@@ -8,11 +8,18 @@
 // FClothingSimulationContextCommon
 //==============================================================================
 
+static TAutoConsoleVariable<float> GClothMaxDeltaTimeTeleportMultiplier(
+	TEXT("p.Cloth.MaxDeltaTimeTeleportMultiplier"),
+	1.5f,
+	TEXT("A multiplier of the MaxPhysicsDelta time at which we will automatically just teleport cloth to its new location\n")
+	TEXT(" default: 1.5"));
+
 FClothingSimulationContextCommon::FClothingSimulationContextCommon()
 	: ComponentToWorld(FTransform::Identity)
 	, WorldGravity(FVector::ZeroVector)
 	, WindVelocity(FVector::ZeroVector)
 	, DeltaSeconds(0.f)
+	, TeleportMode(EClothingTeleportMode::None)
 {}
 
 FClothingSimulationContextCommon::~FClothingSimulationContextCommon()
@@ -27,6 +34,7 @@ void FClothingSimulationContextCommon::Fill(const USkeletalMeshComponent* InComp
 	FillWorldGravity(InComponent);
 	FillWindVelocity(InComponent);
 	FillDeltaSeconds(InDeltaSeconds, InMaxPhysicsDelta);
+	FillTeleportMode(InComponent, InDeltaSeconds, InMaxPhysicsDelta);
 }
 
 void FClothingSimulationContextCommon::FillBoneTransforms(const USkeletalMeshComponent* InComponent)
@@ -114,6 +122,13 @@ void FClothingSimulationContextCommon::FillDeltaSeconds(float InDeltaSeconds, fl
 	DeltaSeconds = FMath::Min(InDeltaSeconds, InMaxPhysicsDelta);
 }
 
+void FClothingSimulationContextCommon::FillTeleportMode(const USkeletalMeshComponent* InComponent, float InDeltaSeconds, float InMaxPhysicsDelta)
+{
+	TeleportMode = (InDeltaSeconds > InMaxPhysicsDelta * GClothMaxDeltaTimeTeleportMultiplier.GetValueOnGameThread()) ?
+		EClothingTeleportMode::Teleport :
+		InComponent->ClothTeleportMode;
+}
+
 float FClothingSimulationContextCommon::SetWindFromComponent(const USkeletalMeshComponent* Component)
 {
 	float WindAdaption;
diff --git a/Engine/Source/Runtime/ClothingSystemRuntimeCommon/Public/ClothingSimulation.h b/Engine/Source/Runtime/ClothingSystemRuntimeCommon/Public/ClothingSimulation.h
--- a/Engine/Source/Runtime/ClothingSystemRuntimeCommon/Public/ClothingSimulation.h
+++ b/Engine/Source/Runtime/ClothingSystemRuntimeCommon/Public/ClothingSimulation.h
@@ -26,6 +26,7 @@ protected:
 	virtual void FillWorldGravity(const USkeletalMeshComponent* InComponent);
 	virtual void FillWindVelocity(const USkeletalMeshComponent* InComponent);
 	virtual void FillDeltaSeconds(float InDeltaSeconds, float InMaxPhysicsDelta);
+	virtual void FillTeleportMode(const USkeletalMeshComponent* InComponent, float InDeltaSeconds, float InMaxPhysicsDelta);
 
 	// Set the wind velocity and return the wind adaptation if required
 	float SetWindFromComponent(const USkeletalMeshComponent* Component);
@@ -48,6 +49,9 @@ public:
 
 	// Delta for this tick
 	float DeltaSeconds;
+
+	// Whether and how we should teleport the simulation this tick
+	EClothingTeleportMode TeleportMode;
 };
 
 // Base simulation to fill in common data for the base context
diff --git a/Engine/Source/Runtime/ClothingSystemRuntimeNv/Private/ClothingSimulationNv.cpp b/Engine/Source/Runtime/ClothingSystemRuntimeNv/Private/ClothingSimulationNv.cpp
--- a/Engine/Source/Runtime/ClothingSystemRuntimeNv/Private/ClothingSimulationNv.cpp
+++ b/Engine/Source/Runtime/ClothingSystemRuntimeNv/Private/ClothingSimulationNv.cpp
@@ -34,12 +34,6 @@ DECLARE_CYCLE_STAT(TEXT("Update Collisions"), STAT_NvClothUpdateCollisions, STAT
 DECLARE_CYCLE_STAT(TEXT("Fill Context"), STAT_NvClothFillContext, STATGROUP_Physics);
 DECLARE_CYCLE_STAT(TEXT("Update Anim Drive"), STAT_NvClothUpdateAnimDrive, STATGROUP_Physics);
 
-static TAutoConsoleVariable<float> GClothMaxDeltaTimeTeleportMultiplier(
-	TEXT("p.Cloth.MaxDeltaTimeTeleportMultiplier"),
-	1.5f,
-	TEXT("A multiplier of the MaxPhysicsDelta time at which we will automatically just teleport cloth to its new location\n")
-	TEXT(" default: 1.5"));
-
 //=============================================================================
 // FClothingSimulationContextNv
 //=============================================================================
@@ -47,7 +41,6 @@ static TAutoConsoleVariable<float> GClothMaxDeltaTimeTeleportMultiplier(
 FClothingSimulationContextNv::FClothingSimulationContextNv()
 	: PredictedLod(0)
 	, WindAdaption(0.0f)
-	, TeleportMode(EClothingTeleportMode::None)
 	, MaxDistanceScale(1.f)
 {
 }
@@ -65,10 +58,6 @@ void FClothingSimulationContextNv::Fill(const USkeletalMeshComponent* InComponen
 
 	PredictedLod = InComponent->PredictedLODLevel;
 
-	TeleportMode = (InDeltaSeconds > InMaxPhysicsDelta * GClothMaxDeltaTimeTeleportMultiplier.GetValueOnGameThread()) ?
-		EClothingTeleportMode::Teleport :
-		InComponent->ClothTeleportMode;
-
 	MaxDistanceScale = InComponent->GetClothMaxDistanceScale();
 }
 
diff --git a/Engine/Source/Runtime/ClothingSystemRuntimeNv/Public/ClothingSimulationNv.h b/Engine/Source/Runtime/ClothingSystemRuntimeNv/Public/ClothingSimulationNv.h
--- a/Engine/Source/Runtime/ClothingSystemRuntimeNv/Public/ClothingSimulationNv.h
+++ b/Engine/Source/Runtime/ClothingSystemRuntimeNv/Public/ClothingSimulationNv.h
@@ -54,9 +54,6 @@ public:
 	// when using the legacy wind calculation mode
 	float WindAdaption;
 
-	// Whether and how we should teleport the simulation this tick
-	EClothingTeleportMode TeleportMode;
-
 	// Scale for the max distance constraints of the simulation mesh
 	float MaxDistanceScale;
 };
