Chaos Cloth - Add new long range constraint's geodesic and limit scale parameters.
- Add long range attrachment debug draw in editor.
- Fix TMap syntax in the geodesic setup.
- Change geodesic constraint setup to only use first and last element of the path.

diff --git a/Engine/Plugins/Experimental/ChaosCloth/Source/Chaos/Private/ChaosCloth/ChaosClothConfig.cpp b/Engine/Plugins/Experimental/ChaosCloth/Source/Chaos/Private/ChaosCloth/ChaosClothConfig.cpp
--- a/Engine/Plugins/Experimental/ChaosCloth/Source/Chaos/Private/ChaosCloth/ChaosClothConfig.cpp
+++ b/Engine/Plugins/Experimental/ChaosCloth/Source/Chaos/Private/ChaosCloth/ChaosClothConfig.cpp
@@ -60,6 +60,7 @@ void UChaosClothConfig::MigrateFrom(const FClothConfig_Legacy& ClothConfig)
 	bUseSelfCollisions = (ClothConfig.SelfCollisionRadius > 0.f && ClothConfig.SelfCollisionStiffness > 0.f);
 
 	StrainLimitingStiffness = FMath::Clamp(ClothConfig.TetherStiffness, 0.f, 1.f);
+	LimitScale = FMath::Clamp(ClothConfig.TetherLimit, 0.01f, 10.f);
 	ShapeTargetStiffness = 0.f;
 }
 
diff --git a/Engine/Plugins/Experimental/ChaosCloth/Source/Chaos/Private/ChaosCloth/ChaosClothingSimulation.cpp b/Engine/Plugins/Experimental/ChaosCloth/Source/Chaos/Private/ChaosCloth/ChaosClothingSimulation.cpp
--- a/Engine/Plugins/Experimental/ChaosCloth/Source/Chaos/Private/ChaosCloth/ChaosClothingSimulation.cpp
+++ b/Engine/Plugins/Experimental/ChaosCloth/Source/Chaos/Private/ChaosCloth/ChaosClothingSimulation.cpp
@@ -146,6 +146,7 @@ void ClothingSimulation::Shutdown()
 	Evolution.Reset();
 	ExternalCollisionsOffset = 0;
 	ClothSharedSimConfig = nullptr;
+	LongRangeConstraints.Reset();
 }
 
 void ClothingSimulation::DestroyActors()
@@ -560,33 +561,35 @@ void ClothingSimulation::AddConstraints(const UChaosClothConfig* ChaosClothSimCo
 		// There might be a cross-over point where the PerParticle version is faster: To be determined
 		if (bUseXPBDConstraints)
 		{
-			const TSharedPtr<Chaos::TXPBDLongRangeConstraints<float, 3>> LongRangeConstraints = MakeShared<Chaos::TXPBDLongRangeConstraints<float, 3>>(
+			LongRangeConstraints = MakeShared<Chaos::TXPBDLongRangeConstraints<float, 3>>(
 				Evolution->Particles(),
 				Mesh.GetPointToNeighborsMap(),
 				10, // The max number of connected neighbors per particle.  ryan - What should this be?  Was k...
-				ChaosClothSimConfig->StrainLimitingStiffness);
+				ChaosClothSimConfig->StrainLimitingStiffness);  // TODO(Kriss.Gossart): Add LimitScale and Geodesic mode if ever of use
 
 			Evolution->AddXPBDConstraintFunctions(
-				[LongRangeConstraints]()
+				[this]()
 				{
-					LongRangeConstraints->Init();
+					static_cast<Chaos::TXPBDLongRangeConstraints<float, 3>&>(*LongRangeConstraints).Init();
 				},
-				[LongRangeConstraints](TPBDParticles<float, 3>& InParticles, const float Dt)
+				[this](TPBDParticles<float, 3>& InParticles, const float Dt)
 				{
-					LongRangeConstraints->Apply(InParticles, Dt);
+					static_cast<Chaos::TXPBDLongRangeConstraints<float, 3>&>(*LongRangeConstraints).Apply(InParticles, Dt);
 				});
 		}
 		else
 		{
-			Chaos::TPBDLongRangeConstraints<float, 3> LongRangeConstraints(
+			LongRangeConstraints = MakeShared<Chaos::TPBDLongRangeConstraints<float, 3>>(
 				Evolution->Particles(),
 				Mesh.GetPointToNeighborsMap(),
 				10, // The max number of connected neighbors per particle.  ryan - What should this be?  Was k...
-				ChaosClothSimConfig->StrainLimitingStiffness);
+				ChaosClothSimConfig->StrainLimitingStiffness,
+				ChaosClothSimConfig->LimitScale,
+				ChaosClothSimConfig->bUseGeodesicDistance);
 
-			Evolution->AddPBDConstraintFunction([LongRangeConstraints = MoveTemp(LongRangeConstraints)](TPBDParticles<float, 3>& InParticles, const float Dt)
+			Evolution->AddPBDConstraintFunction([this](TPBDParticles<float, 3>& InParticles, const float Dt)
 			{
-				LongRangeConstraints.Apply(InParticles, Dt);
+				static_cast<Chaos::TPBDLongRangeConstraints<float, 3>&>(*LongRangeConstraints).Apply(InParticles, Dt);
 			});
 		}
 	}
@@ -1853,4 +1856,29 @@ void ClothingSimulation::DebugDrawAnimDrive(USkeletalMeshComponent* OwnerCompone
 		}
 	}
 }
+
+void ClothingSimulation::DebugDrawLongRangeConstraint(USkeletalMeshComponent* OwnerComponent, FPrimitiveDrawInterface* PDI) const
+{
+	const TPBDParticles<float, 3>& Particles = Evolution->Particles();
+
+	const TArray<TArray<uint32>>& Constraints = LongRangeConstraints->GetConstraints();
+	const TArray<float>& Dists = LongRangeConstraints->GetDists();
+
+	for (int32 i = 0; i < Constraints.Num(); ++i)
+	{
+		const TArray<uint32>& Path = Constraints[i];
+		const float RefDist = Dists[i];
+		const float CurDist = TPBDLongRangeConstraintsBase<float, 3>::ComputeGeodesicDistance(Particles, Path);
+		const float Offset = CurDist - RefDist;
+
+		const Chaos::TVector<float, 3> P0 = Particles.X(Path[0]);  // Kinematic particle
+		const Chaos::TVector<float, 3> P1 = Particles.X(Path[Path.Num() - 1]);  // Target particle
+
+		const Chaos::TVector<float, 3> Direction = (Particles.X(Path[Path.Num() - 2]) - P1).GetSafeNormal();
+		const Chaos::TVector<float, 3> P2 = P1 + Direction * Offset;
+
+		PDI->DrawLine(P0, P1, FColor::Purple, SDPG_World, 0.0f, 0.001f);
+		PDI->DrawLine(P1, P2, FColor::Black, SDPG_World, 0.0f, 0.001f);
+	}
+}
 #endif  // #if WITH_EDITOR
diff --git a/Engine/Plugins/Experimental/ChaosCloth/Source/Chaos/Public/ChaosCloth/ChaosClothConfig.h b/Engine/Plugins/Experimental/ChaosCloth/Source/Chaos/Public/ChaosCloth/ChaosClothConfig.h
--- a/Engine/Plugins/Experimental/ChaosCloth/Source/Chaos/Public/ChaosCloth/ChaosClothConfig.h
+++ b/Engine/Plugins/Experimental/ChaosCloth/Source/Chaos/Public/ChaosCloth/ChaosClothConfig.h
@@ -25,15 +25,15 @@ public:
 	 * -	Total Mass: The total mass is distributed equally over all the particles
 	 * -	Density: A constant mass density is used	 	 
 	 */
-	UPROPERTY(EditAnywhere, Category = MassConfig)
+	UPROPERTY(EditAnywhere, Category = "Mass Config")
 	EClothMassMode MassMode = EClothMassMode::Density;
 
 	// The value used when the Mass Mode is set to Uniform Mass
-	UPROPERTY(EditAnywhere, Category = MassConfig)
+	UPROPERTY(EditAnywhere, Category = "Mass Config")
 	float UniformMass = 1.0f;
 
 	// The value used when Mass Mode is set to TotalMass
-	UPROPERTY(EditAnywhere, Category = MassConfig)
+	UPROPERTY(EditAnywhere, Category = "Mass Config")
 	float TotalMass = 100.0f;
 
 	/**
@@ -43,11 +43,11 @@ public:
 	 * Wool: 0.13
 	 * Silk: 0.133
 	 */
-	UPROPERTY(EditAnywhere, Category = MassConfig)
+	UPROPERTY(EditAnywhere, Category = "Mass Config")
 	float Density = 0.1f;
 
 	// This is a lower bound to cloth particle masses
-	UPROPERTY(EditAnywhere, Category = MassConfig)
+	UPROPERTY(EditAnywhere, Category = "Mass Config")
 	float MinPerParticleMass = 0.0001f;	
 
 	// The Stiffness of the Edge constraints
@@ -67,9 +67,18 @@ public:
 	float VolumeStiffness = 0.f;
 
 	// The strain limiting stiffness of the long range attachment constraints (aka tether stiffness)
-	UPROPERTY(EditAnywhere, Category = Stiffness, meta = (UIMin = "0", UIMax = "1", ClampMin = "0", ClampMax = "1"))
+	UPROPERTY(EditAnywhere, Category = "Long Range Attachment", meta = (UIMin = "0", UIMax = "1", ClampMin = "0", ClampMax = "1"))
 	float StrainLimitingStiffness = 0.5f;
 
+	// The limit scale of the long range attachment constraints (aka tether limit)
+	UPROPERTY(EditAnywhere, Category = "Long Range Attachment", meta = (UIMin = "0.01", UIMax = "2", ClampMin = "0.01", ClampMax = "10"))
+	float LimitScale = 1.f;
+
+	// Use geodesic instead of euclidean distance calculations in the long range attachment constraint,
+	// which is slower at setup but less prone to artifacts during simulation
+	UPROPERTY(EditAnywhere, Category = "Long Range Attachment")
+	bool bUseGeodesicDistance = false;
+
 	// The stiffness of the shape target constraints
 	UPROPERTY()
 	float ShapeTargetStiffness = 0.f;
@@ -83,7 +92,7 @@ public:
 	float AnimDriveSpringStiffness = 0.001f;
 
 	// Enable the more accurate bending element constraints instead of the faster cross-edge spring constraints used for controlling bending stiffness.
-	UPROPERTY(EditAnywhere, Category = ClothEnableFlags)
+	UPROPERTY(EditAnywhere, Category = "Cloth Enable Flags")
 	bool bUseBendingElements = false;
 
 	// Enable tetrahedral constraints
@@ -95,7 +104,7 @@ public:
 	bool bUseThinShellVolumeConstraints = false;
 
 	// Enable self collision
-	UPROPERTY(EditAnywhere, Category = ClothEnableFlags)
+	UPROPERTY(EditAnywhere, Category = "Cloth Enable Flags")
 	bool bUseSelfCollisions = false;
 
 	// Enable continuous collision detection
diff --git a/Engine/Plugins/Experimental/ChaosCloth/Source/Chaos/Public/ChaosCloth/ChaosClothingSimulation.h b/Engine/Plugins/Experimental/ChaosCloth/Source/Chaos/Public/ChaosCloth/ChaosClothingSimulation.h
--- a/Engine/Plugins/Experimental/ChaosCloth/Source/Chaos/Public/ChaosCloth/ChaosClothingSimulation.h
+++ b/Engine/Plugins/Experimental/ChaosCloth/Source/Chaos/Public/ChaosCloth/ChaosClothingSimulation.h
@@ -13,6 +13,7 @@
 namespace Chaos
 {
 	typedef FClothingSimulationContextCommon ClothingSimulationContext;
+	template<class T, int d> class TPBDLongRangeConstraintsBase;
 
 	class ClothingSimulation : public FClothingSimulationCommon
 #if WITH_EDITOR
@@ -50,6 +51,7 @@ namespace Chaos
 		CHAOSCLOTH_API void DebugDrawBackstops(USkeletalMeshComponent* OwnerComponent, FPrimitiveDrawInterface* PDI) const;
 		CHAOSCLOTH_API void DebugDrawMaxDistances(USkeletalMeshComponent* OwnerComponent, FPrimitiveDrawInterface* PDI) const;
 		CHAOSCLOTH_API void DebugDrawAnimDrive(USkeletalMeshComponent* OwnerComponent, FPrimitiveDrawInterface* PDI) const;
+		CHAOSCLOTH_API void DebugDrawLongRangeConstraint(USkeletalMeshComponent* OwnerComponent, FPrimitiveDrawInterface* PDI) const;
 #endif  // #if WITH_EDITOR
 
 	protected:
@@ -141,6 +143,8 @@ namespace Chaos
 
 		FVector Gravity;
 
+		TSharedPtr<TPBDLongRangeConstraintsBase<float, 3>> LongRangeConstraints;
+
 #if WITH_EDITOR
 		// Visualization material
 		UMaterial* DebugClothMaterial;
diff --git a/Engine/Plugins/Experimental/ChaosClothEditor/Source/Chaos/Private/ChaosClothEditor/ChaosSimulationEditorExtender.cpp b/Engine/Plugins/Experimental/ChaosClothEditor/Source/Chaos/Private/ChaosClothEditor/ChaosSimulationEditorExtender.cpp
--- a/Engine/Plugins/Experimental/ChaosClothEditor/Source/Chaos/Private/ChaosClothEditor/ChaosSimulationEditorExtender.cpp
+++ b/Engine/Plugins/Experimental/ChaosClothEditor/Source/Chaos/Private/ChaosClothEditor/ChaosSimulationEditorExtender.cpp
@@ -54,7 +54,8 @@ const FVisualizationOption FVisualizationOption::OptionData[] =
 	FVisualizationOption(&Chaos::ClothingSimulation::DebugDrawCollision           , LOCTEXT("ChaosVisName_Collision"           , "Collisions"                 ), LOCTEXT("ChaosVisName_Collision_ToolTip"           , "Draws the collision bodies the simulation is currently using")),
 	FVisualizationOption(&Chaos::ClothingSimulation::DebugDrawBackstops           , LOCTEXT("ChaosVisName_Backstop"            , "Backstops"                  ), LOCTEXT("ChaosVisName_Backstop_ToolTip"            , "Draws the backstop offset for each simulation particle"), /*bDisablesSimulation =*/true),
 	FVisualizationOption(&Chaos::ClothingSimulation::DebugDrawMaxDistances        , LOCTEXT("ChaosVisName_MaxDistance"         , "Max Distances"              ), LOCTEXT("ChaosVisName_MaxDistance_ToolTip"         , "Draws the current max distances for the sim particles as a line along its normal"), true),
-	FVisualizationOption(&Chaos::ClothingSimulation::DebugDrawAnimDrive           , LOCTEXT("ChaosVisName_AnimDrive"           , "Anim Drive"                 ), LOCTEXT("ChaosVisName_AnimDrive_Tooltip"           , "Draws the current skinned reference mesh for the simulation which anim drive will attempt to reach if enabled"))
+	FVisualizationOption(&Chaos::ClothingSimulation::DebugDrawAnimDrive           , LOCTEXT("ChaosVisName_AnimDrive"           , "Anim Drive"                 ), LOCTEXT("ChaosVisName_AnimDrive_Tooltip"           , "Draws the current skinned reference mesh for the simulation which anim drive will attempt to reach if enabled")),
+	FVisualizationOption(&Chaos::ClothingSimulation::DebugDrawLongRangeConstraint , LOCTEXT("ChaosVisName_LongRangeConstraint" , "Long Range Constraint"      ), LOCTEXT("ChaosVisName_LongRangeConstraint_Tooltip" , "Draws the long range attachment constraint distances"))
 };
 const uint32 FVisualizationOption::Count = sizeof(OptionData) / sizeof(FVisualizationOption);
 
diff --git a/Engine/Source/Runtime/Experimental/Chaos/Private/Chaos/PBDLongRangeConstraintsBase.cpp b/Engine/Source/Runtime/Experimental/Chaos/Private/Chaos/PBDLongRangeConstraintsBase.cpp
--- a/Engine/Source/Runtime/Experimental/Chaos/Private/Chaos/PBDLongRangeConstraintsBase.cpp
+++ b/Engine/Source/Runtime/Experimental/Chaos/Private/Chaos/PBDLongRangeConstraintsBase.cpp
@@ -11,10 +11,28 @@
 using namespace Chaos;
 
 template<class T, int d>
-TPBDLongRangeConstraintsBase<T, d>::TPBDLongRangeConstraintsBase(const TDynamicParticles<T, d>& InParticles, const TMap<int32, TSet<uint32>>& PointToNeighbors, const int32 NumberOfAttachments, const T Stiffness)
+TPBDLongRangeConstraintsBase<T, d>::TPBDLongRangeConstraintsBase(
+	const TDynamicParticles<T, d>& InParticles,
+	const TMap<int32, TSet<uint32>>& PointToNeighbors,
+	const int32 NumberOfAttachments,
+	const T Stiffness,
+	const T LimitScale,
+	const bool bUseGeodesicDistance)
     : MStiffness(Stiffness)
 {
-	ComputeEuclidianConstraints(InParticles, PointToNeighbors, NumberOfAttachments);
+	if (bUseGeodesicDistance)
+	{
+		ComputeGeodesicConstraints(InParticles, PointToNeighbors, NumberOfAttachments);
+	}
+	else
+	{
+		ComputeEuclidianConstraints(InParticles, PointToNeighbors, NumberOfAttachments);
+	}
+	// Scale distance limits
+	for (float& Dist : MDists)
+	{
+		Dist *= LimitScale;
+	}
 }
 
 template<class T, int d>
@@ -110,6 +128,7 @@ void TPBDLongRangeConstraintsBase<T, d>::ComputeEuclidianConstraints(
 		}
 	}
 
+	// Compute the islands of kinematic particles
 	const TArray<TArray<uint32>> IslandElements = ComputeIslands(InParticles, PointToNeighbors, KinematicParticles);
 	int32 NumTotalIslandElements = 0;
 	for(const auto& Elements : IslandElements)
@@ -126,7 +145,7 @@ void TPBDLongRangeConstraintsBase<T, d>::ComputeEuclidianConstraints(
 		if (InParticles.InvM(i) == 0.0)
 			continue;
 
-		// Measure the distance to all particles in all islands...
+		// Measure the distance to all kinematic particles in all islands...
 		ClosestElements.Reset();
 		for (const TArray<uint32>& Elements : IslandElements)
 		{
@@ -192,7 +211,7 @@ void TPBDLongRangeConstraintsBase<T, d>::ComputeGeodesicConstraints(
 		auto Neighbors = PointToNeighbors[i];
 		for (auto Neighbor : Neighbors)
 		{
-			Distances[TVector<uint32, 2>(i, Neighbor)] = ComputeDistance(InParticles, Neighbor, i);
+			Distances.Add(TVector<uint32, 2>(i, Neighbor), ComputeDistance(InParticles, Neighbor, i));
 		}
 	}
 	// Start and End Points to path and geodesic distance
@@ -200,12 +219,12 @@ void TPBDLongRangeConstraintsBase<T, d>::ComputeGeodesicConstraints(
 	// Dijkstra for each Kinematic Particle (assume a small number of kinematic points) - note this is N^2 log N with N kinematic points
 	for (auto Element : KinematicParticles)
 	{
-		GeodesicPaths[TVector<uint32, 2>(Element, Element)] = {0, {Element}};
+		GeodesicPaths.Add(TVector<uint32, 2>(Element, Element), {0, {Element}});
 		for (uint32 i = 0; i < InParticles.Size(); ++i)
 		{
 			if (i != Element)
 			{
-				GeodesicPaths[TVector<uint32, 2>(Element, i)] = {FLT_MAX, {}};
+				GeodesicPaths.Add(TVector<uint32, 2>(Element, i), {FLT_MAX, {}});
 			}
 		}
 	}
@@ -307,8 +326,11 @@ void TPBDLongRangeConstraintsBase<T, d>::ComputeGeodesicConstraints(
 		{
 			Dist += (InParticles.X(MConstraints[i][j]) - InParticles.X(MConstraints[i][j - 1])).Size();
 			Path.Add(MConstraints[i][j]);
-			NewConstraints.Add(Path);
-			NewDists.Add(Dist);
+			// TODO(Kriss.Gossart): We might want to add a new mode for differentiating between these two cases
+			//NewConstraints.Add(Path);  // This would save the whole geodesic path, which is too expensive
+			//NewDists.Add(Dist);
+			NewConstraints.Add({Path[0], Path[Path.Num() - 1]});  // Use shortest distance instead 
+			NewDists.Add(ComputeDistance(InParticles, Path[0], Path[Path.Num() - 1]));
 			ProcessedPairs.Add(MConstraints[i][j], Path);
 		}
 	}
diff --git a/Engine/Source/Runtime/Experimental/Chaos/Public/Chaos/PBDLongRangeConstraints.h b/Engine/Source/Runtime/Experimental/Chaos/Public/Chaos/PBDLongRangeConstraints.h
--- a/Engine/Source/Runtime/Experimental/Chaos/Public/Chaos/PBDLongRangeConstraints.h
+++ b/Engine/Source/Runtime/Experimental/Chaos/Public/Chaos/PBDLongRangeConstraints.h
@@ -17,8 +17,14 @@ class TPBDLongRangeConstraints : public TPBDLongRangeConstraintsBase<T, d>, publ
 	using Base::MConstraints;
 
   public:
-	TPBDLongRangeConstraints(const TDynamicParticles<T, d>& InParticles, const TMap<int32, TSet<uint32>>& PointToNeighbors, const int32 NumberOfAttachments = 1, const T Stiffness = (T)1)
-	    : TPBDLongRangeConstraintsBase<T, d>(InParticles, PointToNeighbors, NumberOfAttachments, Stiffness) {}
+	TPBDLongRangeConstraints(
+		const TDynamicParticles<T, d>& InParticles,
+		const TMap<int32, TSet<uint32>>& PointToNeighbors,
+		const int32 NumberOfAttachments = 1,
+		const T Stiffness = (T)1,
+		const T LimitScale = (T)1,
+		const bool bUseGeodesicDistance = false)
+		: TPBDLongRangeConstraintsBase<T, d>(InParticles, PointToNeighbors, NumberOfAttachments, Stiffness, LimitScale, bUseGeodesicDistance) {}
 	virtual ~TPBDLongRangeConstraints() {}
 
 
diff --git a/Engine/Source/Runtime/Experimental/Chaos/Public/Chaos/PBDLongRangeConstraintsBase.h b/Engine/Source/Runtime/Experimental/Chaos/Public/Chaos/PBDLongRangeConstraintsBase.h
--- a/Engine/Source/Runtime/Experimental/Chaos/Public/Chaos/PBDLongRangeConstraintsBase.h
+++ b/Engine/Source/Runtime/Experimental/Chaos/Public/Chaos/PBDLongRangeConstraintsBase.h
@@ -12,7 +12,14 @@ template<class T, int d>
 class CHAOS_API TPBDLongRangeConstraintsBase
 {
   public:
-	TPBDLongRangeConstraintsBase(const TDynamicParticles<T, d>& InParticles, const TMap<int32, TSet<uint32>>& PointToNeighbors, const int32 NumberOfAttachments = 1, const T Stiffness = (T)1);
+	TPBDLongRangeConstraintsBase(
+		const TDynamicParticles<T, d>& InParticles,
+		const TMap<int32, TSet<uint32>>& PointToNeighbors,
+		const int32 NumberOfAttachments = 1,
+		const T Stiffness = (T)1,
+		const T LimitScale = (T)1,
+		const bool bUseGeodesicDistance = false);
+
 	virtual ~TPBDLongRangeConstraintsBase() {}
 
 	inline TVector<T, d> GetDelta(const TPBDParticles<T, d>& InParticles, const int32 i) const
@@ -49,13 +56,9 @@ class CHAOS_API TPBDLongRangeConstraintsBase
 		return Delta;
 	};
 
-  protected:
-	static TArray<TArray<uint32>> ComputeIslands(const TDynamicParticles<T, d>& InParticles, const TMap<int32, TSet<uint32>>& PointToNeighbors,/*const TTriangleMesh<T>& Mesh,*/ const TArray<uint32>& KinematicParticles);
-	void ComputeEuclidianConstraints(const TDynamicParticles<T, d>& InParticles, const TMap<int32, TSet<uint32>>& PointToNeighbors,/*const TTriangleMesh<T>& Mesh,*/ const int32 NumberOfAttachments);
-	void ComputeGeodesicConstraints(const TDynamicParticles<T, d>& InParticles, const TMap<int32, TSet<uint32>>& PointToNeighbors,/*const TTriangleMesh<T>& Mesh,*/ const int32 NumberOfAttachments);
+	const TArray<TArray<uint32>>& GetConstraints() const { return MConstraints; }
+	const TArray<T>& GetDists() const { return MDists; }
 
-	static T ComputeDistance(const TParticles<T, d>& InParticles, const uint32 i, const uint32 j) { return (InParticles.X(i) - InParticles.X(j)).Size(); }
-	static T ComputeDistance(const TPBDParticles<T, d>& InParticles, const uint32 i, const uint32 j) { return (InParticles.P(i) - InParticles.P(j)).Size(); }
 	static T ComputeGeodesicDistance(const TParticles<T, d>& InParticles, const TArray<uint32>& Path)
 	{
 		T distance = 0;
@@ -75,6 +78,14 @@ class CHAOS_API TPBDLongRangeConstraintsBase
 		return distance;
 	}
 
+  protected:
+	static TArray<TArray<uint32>> ComputeIslands(const TDynamicParticles<T, d>& InParticles, const TMap<int32, TSet<uint32>>& PointToNeighbors,/*const TTriangleMesh<T>& Mesh,*/ const TArray<uint32>& KinematicParticles);
+	void ComputeEuclidianConstraints(const TDynamicParticles<T, d>& InParticles, const TMap<int32, TSet<uint32>>& PointToNeighbors,/*const TTriangleMesh<T>& Mesh,*/ const int32 NumberOfAttachments);
+	void ComputeGeodesicConstraints(const TDynamicParticles<T, d>& InParticles, const TMap<int32, TSet<uint32>>& PointToNeighbors,/*const TTriangleMesh<T>& Mesh,*/ const int32 NumberOfAttachments);
+
+	static T ComputeDistance(const TParticles<T, d>& InParticles, const uint32 i, const uint32 j) { return (InParticles.X(i) - InParticles.X(j)).Size(); }
+	static T ComputeDistance(const TPBDParticles<T, d>& InParticles, const uint32 i, const uint32 j) { return (InParticles.P(i) - InParticles.P(j)).Size(); }
+
   protected:
 	TArray<TArray<uint32>> MConstraints;
 	TArray<T> MDists;
