Chaos Cloth - Fix simulation blow up and damping issues.
- Rework wind and evolution solver to allow drag forces from triangle elements to apply to particles.
- Fix cloth config density scale that was using a square meter values with a square cm unit in the engine.

diff --git a/Engine/Plugins/Experimental/ChaosCloth/Source/Chaos/Private/ChaosCloth/ChaosClothConfig.cpp b/Engine/Plugins/Experimental/ChaosCloth/Source/Chaos/Private/ChaosCloth/ChaosClothConfig.cpp
--- a/Engine/Plugins/Experimental/ChaosCloth/Source/Chaos/Private/ChaosCloth/ChaosClothConfig.cpp
+++ b/Engine/Plugins/Experimental/ChaosCloth/Source/Chaos/Private/ChaosCloth/ChaosClothConfig.cpp
@@ -62,6 +62,8 @@ void UChaosClothConfig::MigrateFrom(const FClothConfig_Legacy& ClothConfig)
 	StrainLimitingStiffness = FMath::Clamp(ClothConfig.TetherStiffness, 0.f, 1.f);
 	LimitScale = FMath::Clamp(ClothConfig.TetherLimit, 0.01f, 10.f);
 	ShapeTargetStiffness = 0.f;
+
+	DragCoefficient = (ClothConfig.WindMethod == EClothingWindMethod_Legacy::Accurate) ? ClothConfig.WindDragCoefficient: 0.5f;  // Only Accurate wind uses the WindDragCoefficient
 }
 
 UChaosClothSharedSimConfig::UChaosClothSharedSimConfig()
@@ -86,8 +88,6 @@ void UChaosClothSharedSimConfig::MigrateFrom(const FClothConfig_Legacy& ClothCon
 	GravityScale = ClothConfig.GravityScale;
 
 	Gravity = ClothConfig.GravityOverride;
-
-	WindDrag = (ClothConfig.WindMethod == EClothingWindMethod_Legacy::Accurate) ? ClothConfig.WindDragCoefficient: 0.5f;  // Only Accurate wind uses the WindDragCoefficient
 }
 
 void UChaosClothSharedSimConfig::Serialize(FArchive& Ar)
diff --git a/Engine/Plugins/Experimental/ChaosCloth/Source/Chaos/Private/ChaosCloth/ChaosClothingSimulation.cpp b/Engine/Plugins/Experimental/ChaosCloth/Source/Chaos/Private/ChaosCloth/ChaosClothingSimulation.cpp
--- a/Engine/Plugins/Experimental/ChaosCloth/Source/Chaos/Private/ChaosCloth/ChaosClothingSimulation.cpp
+++ b/Engine/Plugins/Experimental/ChaosCloth/Source/Chaos/Private/ChaosCloth/ChaosClothingSimulation.cpp
@@ -59,13 +59,14 @@ namespace ClothingSimulationDefault
 	static const float CollisionThickness = 1.2f;
 	static const float CoefficientOfFriction = 0.f;
 	static const float Damping = 0.01f;
-	static const float WindDrag = 0.f;
+	static const float WorldScale = 100.f;  // World is in cm, but values like wind speed and density are in SI unit and relates to m.
 }
 
 ClothingSimulation::ClothingSimulation()
 	: ClothSharedSimConfig(nullptr)
 	, ExternalCollisionsOffset(0)
 	, Gravity(ClothingSimulationDefault::Gravity)
+	, WindVelocity(FVector::ZeroVector)
 {
 #if WITH_EDITOR
 	DebugClothMaterial = LoadObject<UMaterial>(nullptr, TEXT("/Engine/EditorMaterials/Cloth/CameraLitDoubleSided.CameraLitDoubleSided"), nullptr, LOAD_None, nullptr);  // LOAD_EditorOnly
@@ -92,7 +93,6 @@ void ClothingSimulation::Initialize()
 			ClothingSimulationDefault::Damping));
     Evolution->CollisionParticles().AddArray(&BoneIndices);
 	Evolution->CollisionParticles().AddArray(&BaseTransforms);
-	Evolution->GetVelocityField().SetDrag(ClothingSimulationDefault::WindDrag);  // TODO: Per particle drag
     Evolution->GetGravityForces().SetAcceleration(Gravity);
 
     Evolution->SetKinematicUpdateFunction(
@@ -240,6 +240,17 @@ void ClothingSimulation::CreateActor(USkeletalMeshComponent* InOwnerComponent, U
 
 	AddConstraints(ChaosClothSimConfig, PhysMesh, InSimDataIndex);
 
+	// Add Velocity field
+	auto GetVelocity = [this](const TVector<float, 3>&)->TVector<float, 3>
+	{
+		return WindVelocity;
+	};
+	Evolution->GetVelocityFields().Emplace(
+		*Meshes[InSimDataIndex],
+		GetVelocity,
+		/*bInIsUniform =*/ true,
+		ChaosClothSimConfig->DragCoefficient);
+
 	// Add Self Collisions
 	if (ChaosClothSimConfig->bUseSelfCollisions)
 	{
@@ -297,7 +308,6 @@ void ClothingSimulation::UpdateSimulationFromSharedSimConfig()
 		Evolution->SetSelfCollisionThickness(ClothSharedSimConfig->SelfCollisionThickness);
 		Evolution->SetCollisionThickness(ClothSharedSimConfig->CollisionThickness);
 		Evolution->SetDamping(ClothSharedSimConfig->Damping);
-		Evolution->GetVelocityField().SetDrag(ClothSharedSimConfig->WindDrag);  // TODO: Per particle drag
 		Gravity = ClothSharedSimConfig->Gravity;
 	}
 }
@@ -345,43 +355,52 @@ void ClothingSimulation::SetParticleMasses(const UChaosClothConfig* ChaosClothCo
 
 	// Assign per particle mass proportional to connected area.
 	const TArray<TVector<int32, 3>>& SurfaceElements = Mesh.GetSurfaceElements();
-	float TotalArea = 0.0;
+	float TotalArea = 0.f;
 	for (const Chaos::TVector<int32, 3>& Tri : SurfaceElements)
 	{
-		const float TriArea = 0.5 * Chaos::TVector<float, 3>::CrossProduct(
+		const float TriArea = 0.5f * Chaos::TVector<float, 3>::CrossProduct(
 			Particles.X(Tri[1]) - Particles.X(Tri[0]),
 			Particles.X(Tri[2]) - Particles.X(Tri[0])).Size();
 		TotalArea += TriArea;
-		const float ThirdTriArea = TriArea / 3.0;
+		const float ThirdTriArea = TriArea / 3.f;
 		Particles.M(Tri[0]) += ThirdTriArea;
 		Particles.M(Tri[1]) += ThirdTriArea;
 		Particles.M(Tri[2]) += ThirdTriArea;
 	}
 	const TSet<int32> Vertices = Mesh.GetVertices();
+	float TotalMass = 0.f;
 	switch (ChaosClothConfig->MassMode)
 	{
 	case EClothMassMode::UniformMass:
 		for (const int32 Vertex : Vertices)
 		{
 			Particles.M(Vertex) = ChaosClothConfig->UniformMass;
+			TotalMass += Particles.M(Vertex);
 		}
 		break;
 	case EClothMassMode::TotalMass:
 	{
-		const float MassPerUnitArea = TotalArea > 0.0 ? ChaosClothConfig->TotalMass / TotalArea : 1.0;
+		const float MassPerUnitArea = TotalArea > 0.f ? ChaosClothConfig->TotalMass / TotalArea : 1.f;
 		for (const int32 Vertex : Vertices)
 		{
 			Particles.M(Vertex) *= MassPerUnitArea;
+			TotalMass += Particles.M(Vertex);
 		}
 		break;
 	}
 	case EClothMassMode::Density:
+	{
+		const float Density = ChaosClothConfig->Density / FMath::Square(ClothingSimulationDefault::WorldScale);
 		for (const int32 Vertex : Vertices)
 		{
-			Particles.M(Vertex) *= ChaosClothConfig->Density;
+			Particles.M(Vertex) *= Density;
+			TotalMass += Particles.M(Vertex);
 		}
 		break;
+	}
 	};
+	UE_LOG(LogChaosCloth, Verbose, TEXT("Density: %f, Total surface: %f, Total mass: %f, "), TotalArea > 0.f ? TotalMass / TotalArea : 1.f, TotalArea, TotalMass);
+	UE_LOG(LogChaosCloth, Verbose, TEXT("SI Density: %f, SI Total surface: %f, SI Total mass: %f, "), (TotalArea > 0.f ? TotalMass / TotalArea : 1.f) * FMath::Square(ClothingSimulationDefault::WorldScale), TotalArea / FMath::Square(ClothingSimulationDefault::WorldScale), TotalMass);
 
 	// Clamp and enslave
 	const FPointWeightMap& MaxDistances = PhysMesh.GetWeightMap(EChaosWeightMapTarget::MaxDistance);
@@ -1145,9 +1164,8 @@ void ClothingSimulation::Simulate(IClothingSimulationContext* InContext)
 			Context->WorldGravity * ClothSharedSimConfig->GravityScale:
 			Gravity));
 
-	// Set wind velocity
-	static const float UnitConversionScale = 100.f;  // It would seem that in game wind velocity is in m/s and not cm/s as seen in NvClothSupport::Constants::UnitConversionScale
-	Evolution->GetVelocityField().SetVelocity(Context->WindVelocity * UnitConversionScale);
+	// Set wind velocity, used by the velocity field lambda
+	WindVelocity = Context->WindVelocity * ClothingSimulationDefault::WorldScale;  // Wind speed is set in m/s and need to be converted to cm/s
 
 	// Check teleport modes
 	const bool bTeleport = (Context->TeleportMode > EClothingTeleportMode::None);
@@ -1342,6 +1360,7 @@ void ClothingSimulation::RefreshClothConfig()
 
 	Evolution->ResetConstraintRules();
 	Evolution->ResetSelfCollision();
+	Evolution->ResetVelocityFields();
 
 	for (int32 SimDataIndex = 0; SimDataIndex < Assets.Num(); ++SimDataIndex)
 	{
@@ -1358,6 +1377,17 @@ void ClothingSimulation::RefreshClothConfig()
 
 				AddConstraints(ChaosClothConfig, PhysMesh, SimDataIndex);
 
+				// Add Velocity field
+				auto GetVelocity = [this](const TVector<float, 3>&)->TVector<float, 3>
+				{
+					return WindVelocity;
+				};
+				Evolution->GetVelocityFields().Emplace(
+					*Meshes[SimDataIndex],
+					GetVelocity,
+					/*bInIsUniform =*/ true,
+					ChaosClothConfig->DragCoefficient);
+
 				// Add Self Collisions
 				if (ChaosClothConfig->bUseSelfCollisions)
 				{
@@ -1881,4 +1911,27 @@ void ClothingSimulation::DebugDrawLongRangeConstraint(USkeletalMeshComponent* Ow
 		PDI->DrawLine(P1, P2, FColor::Black, SDPG_World, 0.0f, 0.001f);
 	}
 }
+
+void ClothingSimulation::DebugDrawWindDragForces(USkeletalMeshComponent* OwnerComponent, FPrimitiveDrawInterface* PDI) const
+{
+	const TPBDParticles<float, 3>& Particles = Evolution->Particles();
+	const TArray<TVelocityField<float, 3>>& VelocityFields = Evolution->GetVelocityFields();
+
+	for (const TVelocityField<float, 3>& VelocityField : VelocityFields)
+	{
+		const TArray<TVector<int32, 3>>& Elements = VelocityField.GetElements();
+		const TArray<TVector<float, 3>>& Forces = VelocityField.GetForces();
+
+		for (int32 ElementIndex = 0; ElementIndex < Elements.Num(); ++ElementIndex)
+		{
+			const TVector<int32, 3>& Element = Elements[ElementIndex];
+			const TVector<float, 3> Position = (1.f / 3.f ) * (
+				Particles.X(Element[0]) +
+				Particles.X(Element[1]) +
+				Particles.X(Element[2]));
+			const TVector<float, 3>& Force = Forces[ElementIndex];
+			PDI->DrawLine(Position, Position + Force, FColor::Green, SDPG_World, 0.0f, 0.001f);
+		}
+	}
+}
 #endif  // #if WITH_EDITOR
diff --git a/Engine/Plugins/Experimental/ChaosCloth/Source/Chaos/Public/ChaosCloth/ChaosClothConfig.h b/Engine/Plugins/Experimental/ChaosCloth/Source/Chaos/Public/ChaosCloth/ChaosClothConfig.h
--- a/Engine/Plugins/Experimental/ChaosCloth/Source/Chaos/Public/ChaosCloth/ChaosClothConfig.h
+++ b/Engine/Plugins/Experimental/ChaosCloth/Source/Chaos/Public/ChaosCloth/ChaosClothConfig.h
@@ -29,22 +29,25 @@ public:
 	EClothMassMode MassMode = EClothMassMode::Density;
 
 	// The value used when the Mass Mode is set to Uniform Mass
-	UPROPERTY(EditAnywhere, Category = "Mass Config")
-	float UniformMass = 1.0f;
+	UPROPERTY(EditAnywhere, Category = "Mass Config", meta = (UIMin = "0.000001", UIMax = "0.001", ClampMin = "0"))
+	float UniformMass = 0.00015f;
 
 	// The value used when Mass Mode is set to TotalMass
-	UPROPERTY(EditAnywhere, Category = "Mass Config")
-	float TotalMass = 100.0f;
+	UPROPERTY(EditAnywhere, Category = "Mass Config", meta = (UIMin = "0.001", UIMax = "10", ClampMin = "0"))
+	float TotalMass = 0.5f;
 
 	/**
 	 * The value used when Mass Mode is set to Density 
-	 * Water: 1.0
-	 * Cotton: 0.155
-	 * Wool: 0.13
-	 * Silk: 0.133
+	 * Melton Wool: 0.7
+	 * Heavy leather: 0.6
+	 * Polyurethane: 0.5
+	 * Denim: 0.4
+	 * Light leather: 0.3
+	 * Cotton: 0.2
+	 * Silk: 0.1
 	 */
-	UPROPERTY(EditAnywhere, Category = "Mass Config")
-	float Density = 0.1f;
+	UPROPERTY(EditAnywhere, Category = "Mass Config", meta = (UIMin = "0.001", UIMax = "1", ClampMin = "0"))
+	float Density = 0.35f;
 
 	// This is a lower bound to cloth particle masses
 	UPROPERTY(EditAnywhere, Category = "Mass Config")
@@ -83,6 +86,10 @@ public:
 	UPROPERTY()
 	float ShapeTargetStiffness = 0.f;
 
+	// The drag coefficient applying on each particle
+	UPROPERTY(EditAnywhere, Category = Wind, meta = (UIMin = "0", UIMax = "1", ClampMin = "0", ClampMax = "10"))
+	float DragCoefficient = 0.5f;
+
 	// Friction coefficient for cloth - collider interaction
 	UPROPERTY(EditAnywhere, Category = Collision, meta = (UIMin = "0", UIMax = "1", ClampMin = "0", ClampMax = "10"))
 	float CoefficientOfFriction = 0.f;
@@ -169,10 +176,6 @@ public:
 	UPROPERTY(EditAnywhere, Category = Simulation, meta = (EditCondition = "bUseGravityOverride"))
 	FVector Gravity = { 0.f, 0.f, -980.665f };
 
-	// The wind drag coefficient applying on each particle
-	UPROPERTY(EditAnywhere, Category = Simulation, meta = (UIMin = "0", UIMax = "1", ClampMin = "0", ClampMax = "10"))
-	float WindDrag = 0.5f;
-
 	// Enable the XPBD constraints that resolve stiffness independently from the number of iterations
 	// Experimental, this feature might be removed without warning, not for production use
 	UPROPERTY(EditAnywhere, Category = "Experimental", meta = (DisplayName="Use XPBD Constraints (Experimental)"))
diff --git a/Engine/Plugins/Experimental/ChaosCloth/Source/Chaos/Public/ChaosCloth/ChaosClothingSimulation.h b/Engine/Plugins/Experimental/ChaosCloth/Source/Chaos/Public/ChaosCloth/ChaosClothingSimulation.h
--- a/Engine/Plugins/Experimental/ChaosCloth/Source/Chaos/Public/ChaosCloth/ChaosClothingSimulation.h
+++ b/Engine/Plugins/Experimental/ChaosCloth/Source/Chaos/Public/ChaosCloth/ChaosClothingSimulation.h
@@ -52,6 +52,7 @@ namespace Chaos
 		CHAOSCLOTH_API void DebugDrawMaxDistances(USkeletalMeshComponent* OwnerComponent, FPrimitiveDrawInterface* PDI) const;
 		CHAOSCLOTH_API void DebugDrawAnimDrive(USkeletalMeshComponent* OwnerComponent, FPrimitiveDrawInterface* PDI) const;
 		CHAOSCLOTH_API void DebugDrawLongRangeConstraint(USkeletalMeshComponent* OwnerComponent, FPrimitiveDrawInterface* PDI) const;
+		CHAOSCLOTH_API void DebugDrawWindDragForces(USkeletalMeshComponent* OwnerComponent, FPrimitiveDrawInterface* PDI) const;
 #endif  // #if WITH_EDITOR
 
 	protected:
@@ -142,6 +143,7 @@ namespace Chaos
 		float DeltaTime;
 
 		FVector Gravity;
+		FVector WindVelocity;
 
 		TSharedPtr<TPBDLongRangeConstraintsBase<float, 3>> LongRangeConstraints;
 
diff --git a/Engine/Plugins/Experimental/ChaosClothEditor/Source/Chaos/Private/ChaosClothEditor/ChaosSimulationEditorExtender.cpp b/Engine/Plugins/Experimental/ChaosClothEditor/Source/Chaos/Private/ChaosClothEditor/ChaosSimulationEditorExtender.cpp
--- a/Engine/Plugins/Experimental/ChaosClothEditor/Source/Chaos/Private/ChaosClothEditor/ChaosSimulationEditorExtender.cpp
+++ b/Engine/Plugins/Experimental/ChaosClothEditor/Source/Chaos/Private/ChaosClothEditor/ChaosSimulationEditorExtender.cpp
@@ -55,7 +55,8 @@ const FVisualizationOption FVisualizationOption::OptionData[] =
 	FVisualizationOption(&Chaos::ClothingSimulation::DebugDrawBackstops           , LOCTEXT("ChaosVisName_Backstop"            , "Backstops"                  ), LOCTEXT("ChaosVisName_Backstop_ToolTip"            , "Draws the backstop offset for each simulation particle"), /*bDisablesSimulation =*/true),
 	FVisualizationOption(&Chaos::ClothingSimulation::DebugDrawMaxDistances        , LOCTEXT("ChaosVisName_MaxDistance"         , "Max Distances"              ), LOCTEXT("ChaosVisName_MaxDistance_ToolTip"         , "Draws the current max distances for the sim particles as a line along its normal"), true),
 	FVisualizationOption(&Chaos::ClothingSimulation::DebugDrawAnimDrive           , LOCTEXT("ChaosVisName_AnimDrive"           , "Anim Drive"                 ), LOCTEXT("ChaosVisName_AnimDrive_Tooltip"           , "Draws the current skinned reference mesh for the simulation which anim drive will attempt to reach if enabled")),
-	FVisualizationOption(&Chaos::ClothingSimulation::DebugDrawLongRangeConstraint , LOCTEXT("ChaosVisName_LongRangeConstraint" , "Long Range Constraint"      ), LOCTEXT("ChaosVisName_LongRangeConstraint_Tooltip" , "Draws the long range attachment constraint distances"))
+	FVisualizationOption(&Chaos::ClothingSimulation::DebugDrawLongRangeConstraint , LOCTEXT("ChaosVisName_LongRangeConstraint" , "Long Range Constraint"      ), LOCTEXT("ChaosVisName_LongRangeConstraint_Tooltip" , "Draws the long range attachment constraint distances")),
+	FVisualizationOption(&Chaos::ClothingSimulation::DebugDrawWindDragForces      , LOCTEXT("ChaosVisName_WindDragForces"      , "Wind drag forces"           ), LOCTEXT("ChaosVisName_Wind_Tooltip"                , "Draws the Wind drag forces"))
 };
 const uint32 FVisualizationOption::Count = sizeof(OptionData) / sizeof(FVisualizationOption);
 
diff --git a/Engine/Source/Runtime/Experimental/Chaos/Private/Chaos/PBDEvolution.cpp b/Engine/Source/Runtime/Experimental/Chaos/Private/Chaos/PBDEvolution.cpp
--- a/Engine/Source/Runtime/Experimental/Chaos/Private/Chaos/PBDEvolution.cpp
+++ b/Engine/Source/Runtime/Experimental/Chaos/Private/Chaos/PBDEvolution.cpp
@@ -17,6 +17,7 @@
 
 DECLARE_CYCLE_STAT(TEXT("Chaos PBD Advance Time"), STAT_ChaosPBDVAdvanceTime, STATGROUP_Chaos);
 DECLARE_CYCLE_STAT(TEXT("Chaos PBD Velocity Damping State Update"), STAT_ChaosPBDVelocityDampUpdateState, STATGROUP_Chaos);
+DECLARE_CYCLE_STAT(TEXT("Chaos PBD Velocity Field Update Forces"), STAT_ChaosPBDVelocityFieldUpdateForces, STATGROUP_Chaos);
 DECLARE_CYCLE_STAT(TEXT("Chaos PBD Velocity Damping"), STAT_ChaosPBDVelocityDampUpdate, STATGROUP_Chaos);
 DECLARE_CYCLE_STAT(TEXT("Chaos PBD Pre Iteration Updates"), STAT_ChaosPBDPreIterationUpdates, STATGROUP_Chaos);
 DECLARE_CYCLE_STAT(TEXT("Chaos PBD Constraint Rule"), STAT_ChaosPBDConstraintRule, STATGROUP_Chaos);
@@ -69,6 +70,14 @@ void TPBDEvolution<T, d>::AdvanceOneTimeStep(const T Dt)
 		DampVelocityRule.UpdatePositionBasedState(MParticles);
 	}	
 
+	{
+		SCOPE_CYCLE_COUNTER(STAT_ChaosPBDVelocityFieldUpdateForces);
+		for (FVelocityField& VelocityField : VelocityFields)
+		{
+			VelocityField.UpdateForces(MParticles, Dt);
+		}
+	}	
+
 	// Don't bother with threaded execution if we don't have enough work to make it worth while.
 	const bool NonParallelUpdate = MParticles.Size() < 1000; // TODO: 1000 is a guess, tune this!
 
@@ -83,11 +92,14 @@ void TPBDEvolution<T, d>::AdvanceOneTimeStep(const T Dt)
 		{
 			InitForceRule.Apply(MParticles, Dt, Index); // F = TV(0)
 			GravityForces.Apply(MParticles, Dt, Index);
-			VelocityField.Apply(MParticles, Dt, Index);
 			for (TFunction<void(TPBDParticles<T, d>&, const T, const int32)>& ForceRule : MForceRules)
 			{
 				ForceRule(MParticles, Dt, Index); // F += M * A
 			}
+			for (FVelocityField& VelocityField : VelocityFields)
+			{
+				VelocityField.Apply(MParticles, Dt, Index);
+			}
 			if (MKinematicUpdate)
 			{
 				MKinematicUpdate(MParticles, Dt, MTime + Dt, Index); // X = ...
diff --git a/Engine/Source/Runtime/Experimental/Chaos/Public/Chaos/PBDEvolution.h b/Engine/Source/Runtime/Experimental/Chaos/Public/Chaos/PBDEvolution.h
--- a/Engine/Source/Runtime/Experimental/Chaos/Public/Chaos/PBDEvolution.h
+++ b/Engine/Source/Runtime/Experimental/Chaos/Public/Chaos/PBDEvolution.h
@@ -3,7 +3,7 @@
 
 #include "Chaos/KinematicGeometryParticles.h"
 #include "Chaos/PerParticleGravity.h"
-#include "Chaos/PerParticleVelocityField.h"
+#include "Chaos/VelocityField.h"
 #include "Chaos/PBDParticles.h"
 #include "Chaos/Vector.h"
 
@@ -14,7 +14,7 @@ class CHAOS_API TPBDEvolution
 {
   public:
 	using FGravityForces = TPerParticleGravity<T, d>;
-	using FVelocityField = TPerParticleVelocityField<T, d>;
+	using FVelocityField = TVelocityField<T, d>;
 
 	// TODO(mlentine): Init particles from some type of input
 	TPBDEvolution(TPBDParticles<T, d>&& InParticles, TKinematicGeometryClothParticles<T, d>&& InGeometryParticles, TArray<TVector<int32, 3>>&& CollisionTriangles, int32 NumIterations = 1, T CollisionThickness = 0, T SelfCollisionsThickness = 0, T CoefficientOfFriction = 0, T Damping = 0.04);
@@ -35,8 +35,8 @@ class CHAOS_API TPBDEvolution
 	FGravityForces& GetGravityForces() { return GravityForces; }
 	const FGravityForces& GetGravityForces() const { return GravityForces; }
 
-	FVelocityField& GetVelocityField() { return VelocityField; }
-	const FVelocityField& GetVelocityField() const { return VelocityField; }
+	TArray<FVelocityField>& GetVelocityFields() { return VelocityFields; }
+	const TArray<FVelocityField>& GetVelocityFields() const { return VelocityFields; }
 
 	const TGeometryClothParticles<T, d>& CollisionParticles() const { return MCollisionParticles; }
 	TGeometryClothParticles<T, d>& CollisionParticles() { return MCollisionParticles; }
@@ -64,6 +64,7 @@ class CHAOS_API TPBDEvolution
 
 	void ResetConstraintRules() { MInitConstraintRules.Reset(); MConstraintRules.Reset(); };
 	void ResetSelfCollision() { MCollisionTriangles.Reset(); MDisabledCollisionElements.Reset(); };
+	void ResetVelocityFields() { VelocityFields.Reset(); };
 
   private:
 	TPBDParticles<T, d> MParticles;
@@ -79,7 +80,7 @@ class CHAOS_API TPBDEvolution
 	T MTime;
 
 	FGravityForces GravityForces;
-	FVelocityField VelocityField;
+	TArray<FVelocityField> VelocityFields;
 
 	TArray<TFunction<void(TPBDParticles<T, d>&, const T, const int32)>> MForceRules;
 	TArray<TFunction<void()>> MInitConstraintRules;
diff --git a/Engine/Source/Runtime/Experimental/Chaos/Public/Chaos/PerParticleVelocityField.h b/Engine/Source/Runtime/Experimental/Chaos/Public/Chaos/PerParticleVelocityField.h
deleted file mode 100644
diff --git a/Engine/Source/Runtime/Experimental/Chaos/Public/Chaos/VelocityField.h b/Engine/Source/Runtime/Experimental/Chaos/Public/Chaos/VelocityField.h
new file mode 100644
--- /dev/null
+++ b/Engine/Source/Runtime/Experimental/Chaos/Public/Chaos/VelocityField.h
@@ -0,0 +1,139 @@
+// Copyright Epic Games, Inc. All Rights Reserved.
+#pragma once
+
+#include "Chaos/DynamicParticles.h"
+#include "Chaos/TriangleMesh.h"
+
+namespace Chaos
+{
+	// Velocity field basic implementation
+	// TODO: Add lift
+	template<class T, int d>
+	class TVelocityField
+	{
+	public:
+		TVelocityField()
+			: QuarterRhoDragCoefficient((T)0.), bIsUniform(true) {}
+
+		TVelocityField(
+			const TTriangleMesh<T>& TriangleMesh,
+			TFunction<TVector<T, d>(const TVector<T, d>&)> InGetVelocity,
+			const bool bInIsUniform,
+			const T DragCoefficient = (T)0.5,
+			const T FluidDensity = (T)1.225e-6)
+			: PointToTriangleMap(TriangleMesh.GetPointToTriangleMap())  // TODO(Kriss.Gossart): Replace map lookup by safer array lookup
+			, Elements(TriangleMesh.GetElements())
+			, GetVelocity(InGetVelocity)
+			, bIsUniform(bInIsUniform)
+		{
+			Forces.SetNumUninitialized(Elements.Num());
+			SetProperties(DragCoefficient, FluidDensity);
+		}
+
+		virtual ~TVelocityField() {}
+
+		inline void UpdateForces(const TPBDParticles<T, d>& InParticles, const T Dt) const
+		{
+			if (GetVelocity)
+			{
+				if (bIsUniform)
+				{
+					const TVector<T, d> Velocity = GetVelocity(TVector<T, d>(0.));
+					for (int32 ElementIndex = 0; ElementIndex < Elements.Num(); ++ElementIndex)
+					{
+						UpdateUniformField(InParticles, Dt, ElementIndex, Velocity);
+					}
+				}
+				else
+				{
+					for (int32 ElementIndex = 0; ElementIndex < Elements.Num(); ++ElementIndex)
+					{
+						UpdateVectorField(InParticles, Dt, ElementIndex);
+					}
+				}
+			}
+		}
+
+		inline void Apply(TPBDParticles<T, d>& InParticles, const T Dt, const int32 Index) const
+		{
+			if (const TArray<int32>* ElementIndices = PointToTriangleMap.Find(Index))  // TODO(Kriss.Gossart): Replace map lookup by faster array lookup
+			{
+				for (const int32 ElementIndex : *ElementIndices)
+				{
+					InParticles.F(Index) += Forces[ElementIndex];
+				}
+			}
+		}
+
+		void SetProperties(const T DragCoefficient, const T FluidDensity)
+		{
+			QuarterRhoDragCoefficient = (T)0.25 * FluidDensity * DragCoefficient;
+		}
+
+		const TArray<TVector<int32, 3>>& GetElements() const { return Elements; }
+		const TArray<TVector<T, d>>& GetForces() const { return Forces; }
+
+	private:
+		inline void UpdateUniformField(const TPBDParticles<T, d>& InParticles, const T Dt, int32 ElementIndex, const TVector<T, d>& Velocity) const
+		{
+			const TVector<int32, 3>& Element = Elements[ElementIndex];
+
+			// Calculate direction and the relative velocity of the triangle to the flow
+			const TVector<T, d>& SurfaceVelocity = (T)(1. / 3.) * (
+				InParticles.V(Element[0]) + 
+				InParticles.V(Element[1]) + 
+				InParticles.V(Element[2]));
+			TVector<T, d> Direction = Velocity - SurfaceVelocity;
+			const T RelVelocity = Direction.SafeNormalize();
+
+			// Calculate the cross sectional area of the surface exposed to the flow
+			TVector<float, 3> Normal = TVector<float, 3>::CrossProduct(
+				InParticles.X(Element[1]) - InParticles.X(Element[0]),
+				InParticles.X(Element[2]) - InParticles.X(Element[0]));
+			const T DoubleArea = Normal.SafeNormalize();
+			const T DoubleCrossSectionalArea = DoubleArea * FMath::Abs(TVector<float, 3>::DotProduct(Direction, Normal));
+
+			// Calculate the drag force
+			const T Drag = QuarterRhoDragCoefficient * FMath::Square(RelVelocity) * DoubleCrossSectionalArea;
+			Forces[ElementIndex] = Direction * Drag;
+		}
+
+		inline void UpdateVectorField(const TPBDParticles<T, d>& InParticles, const T Dt, int32 ElementIndex) const
+		{
+			const TVector<int32, 3>& Element = Elements[ElementIndex];
+
+			// Get the triangle's position
+			const TVector<T, d>& SurfacePosition = (T)(1. / 3.) * (
+				InParticles.X(Element[0]) +
+				InParticles.X(Element[1]) +
+				InParticles.X(Element[2]));
+
+			// Calculate direction and the relative velocity of the triangle to the flow at the field's sampled position
+			const TVector<T, d>& SurfaceVelocity = (T)(1. / 3.) * (
+				InParticles.V(Element[0]) + 
+				InParticles.V(Element[1]) + 
+				InParticles.V(Element[2]));
+			TVector<T, d> Direction = GetVelocity(SurfacePosition) - SurfaceVelocity;
+			const T RelVelocity = Direction.SafeNormalize();
+
+			// Calculate the cross sectional area of the surface exposed to the flow
+			 TVector<float, 3> Normal = TVector<float, 3>::CrossProduct(
+				InParticles.X(Element[1]) - InParticles.X(Element[0]),
+				InParticles.X(Element[2]) - InParticles.X(Element[0]));
+			const T DoubleArea = Normal.SafeNormalize();
+			const T DoubleCrossSectionalArea = DoubleArea * FMath::Abs(TVector<float, 3>::DotProduct(Direction, Normal));
+
+			// Calculate the drag force
+			const T Drag = QuarterRhoDragCoefficient * FMath::Square(RelVelocity) * DoubleCrossSectionalArea;
+			Forces[ElementIndex] = Direction * Drag;
+		}
+
+	private:
+		const TMap<int32, TArray<int32>>& PointToTriangleMap;
+		const TArray<TVector<int32, 3>>& Elements;
+		TFunction<TVector<T, d>(const TVector<T, d>&)> GetVelocity;
+		mutable TArray<TVector<T, 3>> Forces;
+		float QuarterRhoDragCoefficient;
+		bool bIsUniform;
+	};
+}
