Chaos Cloth - Fix wrong index in log message causing a range check to fail.

diff --git a/Engine/Plugins/Experimental/ChaosCloth/Source/Chaos/Private/ChaosCloth/ChaosClothingSimulation.cpp b/Engine/Plugins/Experimental/ChaosCloth/Source/Chaos/Private/ChaosCloth/ChaosClothingSimulation.cpp
--- a/Engine/Plugins/Experimental/ChaosCloth/Source/Chaos/Private/ChaosCloth/ChaosClothingSimulation.cpp
+++ b/Engine/Plugins/Experimental/ChaosCloth/Source/Chaos/Private/ChaosCloth/ChaosClothingSimulation.cpp
@@ -1042,7 +1042,7 @@ void ClothingSimulation::AddCollisions(const FClothCollisionData& ClothCollision
 			const FClothCollisionPrim_Sphere& Sphere = ClothCollisionData.Spheres[i];
 
 			BoneIndices[j] = GetMappedBoneIndex(UsedBoneIndices, Sphere.BoneIndex);
-			UE_LOG(LogChaosCloth, VeryVerbose, TEXT("Found collision sphere on bone index %d."), BoneIndices[i]);
+			UE_LOG(LogChaosCloth, VeryVerbose, TEXT("Found collision sphere on bone index %d."), BoneIndices[j]);
 
 			BaseTransforms[j] = Chaos::TRigidTransform<float, 3>(FTransform::Identity);
 
