Chaos Cloth - Fix clothing interactor and Chaos Cloth simulation to allow for in game change of gravity.

diff --git a/Engine/Plugins/Experimental/ChaosCloth/Source/Chaos/Private/ChaosCloth/ChaosClothingSimulation.cpp b/Engine/Plugins/Experimental/ChaosCloth/Source/Chaos/Private/ChaosCloth/ChaosClothingSimulation.cpp
--- a/Engine/Plugins/Experimental/ChaosCloth/Source/Chaos/Private/ChaosCloth/ChaosClothingSimulation.cpp
+++ b/Engine/Plugins/Experimental/ChaosCloth/Source/Chaos/Private/ChaosCloth/ChaosClothingSimulation.cpp
@@ -65,6 +65,7 @@ namespace ClothingSimulationDefault
 ClothingSimulation::ClothingSimulation()
 	: ClothSharedSimConfig(nullptr)
 	, ExternalCollisionsOffset(0)
+	, bOverrideGravity(false)
 	, Gravity(ClothingSimulationDefault::Gravity)
 	, WindVelocity(FVector::ZeroVector)
 {
@@ -308,7 +309,6 @@ void ClothingSimulation::UpdateSimulationFromSharedSimConfig()
 		Evolution->SetSelfCollisionThickness(ClothSharedSimConfig->SelfCollisionThickness);
 		Evolution->SetCollisionThickness(ClothSharedSimConfig->CollisionThickness);
 		Evolution->SetDamping(ClothSharedSimConfig->Damping);
-		Gravity = ClothSharedSimConfig->Gravity;
 	}
 }
 
@@ -1158,11 +1158,12 @@ void ClothingSimulation::Simulate(IClothingSimulationContext* InContext)
 		return;
 	}
 
-	// Set gravity
+	// Set gravity, using the legacy priority: 1) config override, 2) game override, 3) world gravity
 	Evolution->GetGravityForces().SetAcceleration(Chaos::TVector<float, 3>(
-		(ClothSharedSimConfig && !ClothSharedSimConfig->bUseGravityOverride) ?
-			Context->WorldGravity * ClothSharedSimConfig->GravityScale:
-			Gravity));
+		(ClothSharedSimConfig && ClothSharedSimConfig->bUseGravityOverride) ? ClothSharedSimConfig->Gravity :
+		bOverrideGravity ? Gravity :
+		ClothSharedSimConfig ? Context->WorldGravity * ClothSharedSimConfig->GravityScale :
+		Context->WorldGravity));
 
 	// Set wind velocity, used by the velocity field lambda
 	WindVelocity = Context->WindVelocity * ClothingSimulationDefault::WorldScale;  // Wind speed is set in m/s and need to be converted to cm/s
@@ -1426,12 +1427,13 @@ void ClothingSimulation::SetAnimDriveSpringStiffness(float InStiffness)
 
 void ClothingSimulation::SetGravityOverride(const FVector& InGravityOverride)
 {
+	bOverrideGravity = true;
 	Gravity = InGravityOverride;
 }
 
 void ClothingSimulation::DisableGravityOverride()
 {
-	Gravity = !ClothSharedSimConfig ? ClothingSimulationDefault::Gravity : ClothSharedSimConfig->Gravity;
+	bOverrideGravity = false;
 }
 
 #if WITH_EDITOR
diff --git a/Engine/Plugins/Experimental/ChaosCloth/Source/Chaos/Public/ChaosCloth/ChaosClothingSimulation.h b/Engine/Plugins/Experimental/ChaosCloth/Source/Chaos/Public/ChaosCloth/ChaosClothingSimulation.h
--- a/Engine/Plugins/Experimental/ChaosCloth/Source/Chaos/Public/ChaosCloth/ChaosClothingSimulation.h
+++ b/Engine/Plugins/Experimental/ChaosCloth/Source/Chaos/Public/ChaosCloth/ChaosClothingSimulation.h
@@ -142,6 +142,7 @@ namespace Chaos
 		float Time;
 		float DeltaTime;
 
+		bool bOverrideGravity;
 		FVector Gravity;
 		FVector WindVelocity;
 
diff --git a/Engine/Plugins/Experimental/ChaosCloth/Source/Chaos/Public/ChaosCloth/ChaosClothingSimulationInteractor.h b/Engine/Plugins/Experimental/ChaosCloth/Source/Chaos/Public/ChaosCloth/ChaosClothingSimulationInteractor.h
--- a/Engine/Plugins/Experimental/ChaosCloth/Source/Chaos/Public/ChaosCloth/ChaosClothingSimulationInteractor.h
+++ b/Engine/Plugins/Experimental/ChaosCloth/Source/Chaos/Public/ChaosCloth/ChaosClothingSimulationInteractor.h
@@ -19,20 +19,11 @@ public:
 	virtual void PhysicsAssetUpdated() override;
 	virtual void ClothConfigUpdated() override;
 	virtual void Sync(IClothingSimulation* InSimulation, IClothingSimulationContext* InContext) override;
+	virtual void SetAnimDriveSpringStiffness(float InStiffness) override;
+	virtual void EnableGravityOverride(const FVector& InVector) override;
+	virtual void DisableGravityOverride() override;
 	//////////////////////////////////////////////////////////////////////////
 
-	// Set the stiffness of the spring force for the animation drive
-	UFUNCTION(BlueprintCallable, Category=ClothingSimulation)
-	void SetAnimDriveSpringStiffness(float InStiffness);
-
-	// Set a new gravity override and enable the override
-	UFUNCTION(BlueprintCallable, Category=ClothingSimulation)
-	void EnableGravityOverride(const FVector& InVector);
-
-	// Disable any currently set gravity override
-	UFUNCTION(BlueprintCallable, Category = ClothingSimulation)
-	void DisableGravityOverride();
-
 private:
 
 	// Command queue processed when we hit a sync
diff --git a/Engine/Source/Runtime/ClothingSystemRuntimeInterface/Public/ClothingSimulationInteractor.h b/Engine/Source/Runtime/ClothingSystemRuntimeInterface/Public/ClothingSimulationInteractor.h
--- a/Engine/Source/Runtime/ClothingSystemRuntimeInterface/Public/ClothingSimulationInteractor.h
+++ b/Engine/Source/Runtime/ClothingSystemRuntimeInterface/Public/ClothingSimulationInteractor.h
@@ -42,6 +42,21 @@ public:
 	virtual void Sync(IClothingSimulation* InSimulation, IClothingSimulationContext* InContext)
 	PURE_VIRTUAL(UClothingSimulationInteractor::Sync, );
 
+	// Set the stiffness of the spring force for the animation drive
+	UFUNCTION(BlueprintCallable, Category = ClothingSimulation)
+	virtual void SetAnimDriveSpringStiffness(float InStiffness)
+	PURE_VIRTUAL(UClothingSimulationInteractor::SetAnimDriveSpringStiffness, );
+
+	// Set a new gravity override and enable the override
+	UFUNCTION(BlueprintCallable, Category = ClothingSimulation)
+	virtual void EnableGravityOverride(const FVector& InVector)
+	PURE_VIRTUAL(UClothingSimulationInteractor::EnableGravityOverride, );
+
+	// Disable any currently set gravity override
+	UFUNCTION(BlueprintCallable, Category = ClothingSimulation)
+	virtual void DisableGravityOverride()
+	PURE_VIRTUAL(UClothingSimulationInteractor::DisableGravityOverride, );
+
 protected:
 
 	// Intended to be called by functions on the interactor to message to the 
diff --git a/Engine/Source/Runtime/ClothingSystemRuntimeNv/Public/ClothingSimulationInteractorNv.h b/Engine/Source/Runtime/ClothingSystemRuntimeNv/Public/ClothingSimulationInteractorNv.h
--- a/Engine/Source/Runtime/ClothingSystemRuntimeNv/Public/ClothingSimulationInteractorNv.h
+++ b/Engine/Source/Runtime/ClothingSystemRuntimeNv/Public/ClothingSimulationInteractorNv.h
@@ -23,24 +23,15 @@ public:
 	virtual void PhysicsAssetUpdated() override;
 	virtual void ClothConfigUpdated() override;
 	virtual void Sync(IClothingSimulation* InSimulation, IClothingSimulationContext* InContext) override;
+	virtual void SetAnimDriveSpringStiffness(float InStiffness) override;
+	virtual void EnableGravityOverride(const FVector& InVector) override;
+	virtual void DisableGravityOverride() override;
 	//////////////////////////////////////////////////////////////////////////
 
-	// Set the stiffness of the spring force for anim drive
-	UFUNCTION(BlueprintCallable, Category=ClothingSimulation)
-	void SetAnimDriveSpringStiffness(float InStiffness);
-
 	// Set the stiffness of the resistive damping force for anim drive
 	UFUNCTION(BlueprintCallable, Category=ClothingSimulation)
 	void SetAnimDriveDamperStiffness(float InStiffness);
 
-	// Set a new gravity override and enable the override
-	UFUNCTION(BlueprintCallable, Category=ClothingSimulation)
-	void EnableGravityOverride(const FVector& InVector);
-
-	// Disable any currently set gravity override
-	UFUNCTION(BlueprintCallable, Category = ClothingSimulation)
-	void DisableGravityOverride();
-
 private:
 
 	// Command queue processed when we hit a sync
