Chaos Cloth - Fix multiple cloth long range constraints.

diff --git a/Engine/Plugins/Experimental/ChaosCloth/Source/Chaos/Private/ChaosCloth/ChaosClothingSimulation.cpp b/Engine/Plugins/Experimental/ChaosCloth/Source/Chaos/Private/ChaosCloth/ChaosClothingSimulation.cpp
--- a/Engine/Plugins/Experimental/ChaosCloth/Source/Chaos/Private/ChaosCloth/ChaosClothingSimulation.cpp
+++ b/Engine/Plugins/Experimental/ChaosCloth/Source/Chaos/Private/ChaosCloth/ChaosClothingSimulation.cpp
@@ -186,6 +186,7 @@ void ClothingSimulation::CreateActor(USkeletalMeshComponent* InOwnerComponent, U
 		Meshes.SetNum(InSimDataIndex + 1);
 		FaceNormals.SetNum(InSimDataIndex + 1);
 		PointNormals.SetNum(InSimDataIndex + 1);
+		LongRangeConstraints.SetNum(InSimDataIndex + 1);
 	}
 
 	check(Asset->GetNumLods() > 0);
@@ -580,25 +581,25 @@ void ClothingSimulation::AddConstraints(const UChaosClothConfig* ChaosClothSimCo
 		// There might be a cross-over point where the PerParticle version is faster: To be determined
 		if (bUseXPBDConstraints)
 		{
-			LongRangeConstraints = MakeShared<Chaos::TXPBDLongRangeConstraints<float, 3>>(
+			LongRangeConstraints[InSimDataIndex] = MakeShared<Chaos::TXPBDLongRangeConstraints<float, 3>>(
 				Evolution->Particles(),
 				Mesh.GetPointToNeighborsMap(),
 				10, // The max number of connected neighbors per particle.  ryan - What should this be?  Was k...
 				ChaosClothSimConfig->StrainLimitingStiffness);  // TODO(Kriss.Gossart): Add LimitScale and Geodesic mode if ever of use
 
 			Evolution->AddXPBDConstraintFunctions(
-				[this]()
+				[this, InSimDataIndex]()
 				{
-					static_cast<Chaos::TXPBDLongRangeConstraints<float, 3>&>(*LongRangeConstraints).Init();
+					static_cast<Chaos::TXPBDLongRangeConstraints<float, 3>&>(*LongRangeConstraints[InSimDataIndex]).Init();
 				},
-				[this](TPBDParticles<float, 3>& InParticles, const float Dt)
+				[this, InSimDataIndex](TPBDParticles<float, 3>& InParticles, const float Dt)
 				{
-					static_cast<Chaos::TXPBDLongRangeConstraints<float, 3>&>(*LongRangeConstraints).Apply(InParticles, Dt);
+					static_cast<Chaos::TXPBDLongRangeConstraints<float, 3>&>(*LongRangeConstraints[InSimDataIndex]).Apply(InParticles, Dt);
 				});
 		}
 		else
 		{
-			LongRangeConstraints = MakeShared<Chaos::TPBDLongRangeConstraints<float, 3>>(
+			LongRangeConstraints[InSimDataIndex] = MakeShared<Chaos::TPBDLongRangeConstraints<float, 3>>(
 				Evolution->Particles(),
 				Mesh.GetPointToNeighborsMap(),
 				10, // The max number of connected neighbors per particle.  ryan - What should this be?  Was k...
@@ -606,9 +607,9 @@ void ClothingSimulation::AddConstraints(const UChaosClothConfig* ChaosClothSimCo
 				ChaosClothSimConfig->LimitScale,
 				ChaosClothSimConfig->bUseGeodesicDistance);
 
-			Evolution->AddPBDConstraintFunction([this](TPBDParticles<float, 3>& InParticles, const float Dt)
+			Evolution->AddPBDConstraintFunction([this, InSimDataIndex](TPBDParticles<float, 3>& InParticles, const float Dt)
 			{
-				static_cast<Chaos::TPBDLongRangeConstraints<float, 3>&>(*LongRangeConstraints).Apply(InParticles, Dt);
+				static_cast<Chaos::TPBDLongRangeConstraints<float, 3>&>(*LongRangeConstraints[InSimDataIndex]).Apply(InParticles, Dt);
 			});
 		}
 	}
@@ -1893,24 +1894,32 @@ void ClothingSimulation::DebugDrawLongRangeConstraint(USkeletalMeshComponent* Ow
 {
 	const TPBDParticles<float, 3>& Particles = Evolution->Particles();
 
-	const TArray<TArray<uint32>>& Constraints = LongRangeConstraints->GetConstraints();
-	const TArray<float>& Dists = LongRangeConstraints->GetDists();
-
-	for (int32 i = 0; i < Constraints.Num(); ++i)
+	for (int32 i = 0; i < IndexToRangeMap.Num(); ++i)
 	{
-		const TArray<uint32>& Path = Constraints[i];
-		const float RefDist = Dists[i];
-		const float CurDist = TPBDLongRangeConstraintsBase<float, 3>::ComputeGeodesicDistance(Particles, Path);
-		const float Offset = CurDist - RefDist;
+		if (!LongRangeConstraints[i])
+		{
+			continue;
+		}
 
-		const Chaos::TVector<float, 3> P0 = Particles.X(Path[0]);  // Kinematic particle
-		const Chaos::TVector<float, 3> P1 = Particles.X(Path[Path.Num() - 1]);  // Target particle
+		const TArray<TArray<uint32>>& Constraints = LongRangeConstraints[i]->GetConstraints();
+		const TArray<float>& Dists = LongRangeConstraints[i]->GetDists();
 
-		const Chaos::TVector<float, 3> Direction = (Particles.X(Path[Path.Num() - 2]) - P1).GetSafeNormal();
-		const Chaos::TVector<float, 3> P2 = P1 + Direction * Offset;
+		for (int32 j = 0; j < Constraints.Num(); ++j)
+		{
+			const TArray<uint32>& Path = Constraints[j];
+			const float RefDist = Dists[j];
+			const float CurDist = TPBDLongRangeConstraintsBase<float, 3>::ComputeGeodesicDistance(Particles, Path);
+			const float Offset = CurDist - RefDist;
+
+			const Chaos::TVector<float, 3> P0 = Particles.X(Path[0]);  // Kinematic particle
+			const Chaos::TVector<float, 3> P1 = Particles.X(Path[Path.Num() - 1]);  // Target particle
 
-		PDI->DrawLine(P0, P1, FColor::Purple, SDPG_World, 0.0f, 0.001f);
-		PDI->DrawLine(P1, P2, FColor::Black, SDPG_World, 0.0f, 0.001f);
+			const Chaos::TVector<float, 3> Direction = (Particles.X(Path[Path.Num() - 2]) - P1).GetSafeNormal();
+			const Chaos::TVector<float, 3> P2 = P1 + Direction * Offset;
+
+			PDI->DrawLine(P0, P1, FColor::Purple, SDPG_World, 0.0f, 0.001f);
+			PDI->DrawLine(P1, P2, FColor::Black, SDPG_World, 0.0f, 0.001f);
+		}
 	}
 }
 
diff --git a/Engine/Plugins/Experimental/ChaosCloth/Source/Chaos/Public/ChaosCloth/ChaosClothingSimulation.h b/Engine/Plugins/Experimental/ChaosCloth/Source/Chaos/Public/ChaosCloth/ChaosClothingSimulation.h
--- a/Engine/Plugins/Experimental/ChaosCloth/Source/Chaos/Public/ChaosCloth/ChaosClothingSimulation.h
+++ b/Engine/Plugins/Experimental/ChaosCloth/Source/Chaos/Public/ChaosCloth/ChaosClothingSimulation.h
@@ -146,7 +146,7 @@ namespace Chaos
 		FVector Gravity;
 		FVector WindVelocity;
 
-		TSharedPtr<TPBDLongRangeConstraintsBase<float, 3>> LongRangeConstraints;
+		TArray<TSharedPtr<TPBDLongRangeConstraintsBase<float, 3>>> LongRangeConstraints;
 
 #if WITH_EDITOR
 		// Visualization material
