Chaos Cloth - Add per cloth damping.
- Add shared config Damping Override parameter to prevent accidental override of a correct damping value.
- Rename and hide unused parameter CoefficienOfFriction to FrictionCoefficient for when friction will get enabled.

diff --git a/Engine/Plugins/Experimental/ChaosCloth/Source/Chaos/Private/ChaosCloth/ChaosClothConfig.cpp b/Engine/Plugins/Experimental/ChaosCloth/Source/Chaos/Private/ChaosCloth/ChaosClothConfig.cpp
--- a/Engine/Plugins/Experimental/ChaosCloth/Source/Chaos/Private/ChaosCloth/ChaosClothConfig.cpp
+++ b/Engine/Plugins/Experimental/ChaosCloth/Source/Chaos/Private/ChaosCloth/ChaosClothConfig.cpp
@@ -54,7 +54,7 @@ void UChaosClothConfig::MigrateFrom(const FClothConfig_Legacy& ClothConfig)
 
 	AnimDriveSpringStiffness = FMath::Clamp(ClothConfig.AnimDriveSpringStiffness, 0.f, 1.f);
 
-	CoefficientOfFriction = FMath::Clamp(ClothConfig.Friction, 0.f, 10.f);
+	FrictionCoefficient = FMath::Clamp(ClothConfig.Friction, 0.f, 10.f);
 
 	bUseBendingElements = false;
 	bUseSelfCollisions = (ClothConfig.SelfCollisionRadius > 0.f && ClothConfig.SelfCollisionStiffness > 0.f);
@@ -64,6 +64,9 @@ void UChaosClothConfig::MigrateFrom(const FClothConfig_Legacy& ClothConfig)
 	ShapeTargetStiffness = 0.f;
 
 	DragCoefficient = (ClothConfig.WindMethod == EClothingWindMethod_Legacy::Accurate) ? ClothConfig.WindDragCoefficient: 0.5f;  // Only Accurate wind uses the WindDragCoefficient
+
+	const float Damping = (ClothConfig.Damping.X + ClothConfig.Damping.Y + ClothConfig.Damping.Z) / 3.f;
+	DampingCoefficient = FMath::Clamp(Damping * Damping * 0.95f, 0.f, 1.f);  // Nv Cloth seems to have a different damping formulation.
 }
 
 UChaosClothSharedSimConfig::UChaosClothSharedSimConfig()
@@ -80,9 +83,6 @@ void UChaosClothSharedSimConfig::MigrateFrom(const FClothConfig_Legacy& ClothCon
 
 	CollisionThickness = FMath::Clamp(ClothConfig.CollisionThickness, 0.f, 1000.f);
 
-	const float InDamping = (ClothConfig.Damping.X + ClothConfig.Damping.Y + ClothConfig.Damping.Z) / 3.f;
-	Damping = FMath::Clamp(InDamping * InDamping * 0.95f, 0.f, 1.f);  // Nv Cloth seems to have a different damping formulation.
-
 	bUseGravityOverride = ClothConfig.bUseGravityOverride;
 
 	GravityScale = ClothConfig.GravityScale;
diff --git a/Engine/Plugins/Experimental/ChaosCloth/Source/Chaos/Private/ChaosCloth/ChaosClothingSimulation.cpp b/Engine/Plugins/Experimental/ChaosCloth/Source/Chaos/Private/ChaosCloth/ChaosClothingSimulation.cpp
--- a/Engine/Plugins/Experimental/ChaosCloth/Source/Chaos/Private/ChaosCloth/ChaosClothingSimulation.cpp
+++ b/Engine/Plugins/Experimental/ChaosCloth/Source/Chaos/Private/ChaosCloth/ChaosClothingSimulation.cpp
@@ -57,8 +57,8 @@ namespace ClothingSimulationDefault
 	static const int32 NumIterations = 1;
 	static const float SelfCollisionThickness = 2.f;
 	static const float CollisionThickness = 1.2f;
-	static const float CoefficientOfFriction = 0.f;
-	static const float Damping = 0.01f;
+	static const float FrictionCoefficient = 0.f;
+	static const float DampingCoefficient = 0.01f;
 	static const float WorldScale = 100.f;  // World is in cm, but values like wind speed and density are in SI unit and relates to m.
 }
 
@@ -90,8 +90,8 @@ void ClothingSimulation::Initialize()
 			ClothingSimulationDefault::NumIterations,
 			ClothingSimulationDefault::CollisionThickness,
 			ClothingSimulationDefault::SelfCollisionThickness,
-			ClothingSimulationDefault::CoefficientOfFriction,
-			ClothingSimulationDefault::Damping));
+			ClothingSimulationDefault::FrictionCoefficient,
+			ClothingSimulationDefault::DampingCoefficient));
     Evolution->CollisionParticles().AddArray(&BoneIndices);
 	Evolution->CollisionParticles().AddArray(&BaseTransforms);
     Evolution->GetGravityForces().SetAcceleration(Gravity);
@@ -199,7 +199,7 @@ void ClothingSimulation::CreateActor(USkeletalMeshComponent* InOwnerComponent, U
 	// Add particles
 	TPBDParticles<float, 3>& Particles = Evolution->Particles();
 	const uint32 Offset = Particles.Size();
-	Particles.AddParticles(PhysMesh.Vertices.Num());
+	Evolution->AddParticles(PhysMesh.Vertices.Num(), (uint32)InSimDataIndex);
 
 	if (IndexToRangeMap.Num() <= InSimDataIndex)
 	{
@@ -242,7 +242,17 @@ void ClothingSimulation::CreateActor(USkeletalMeshComponent* InOwnerComponent, U
 
 	AddConstraints(ChaosClothSimConfig, PhysMesh, InSimDataIndex);
 
-	// Add Velocity field
+	// Set damping
+	if (ClothSharedSimConfig && ClothSharedSimConfig->bUseDampingOverride)
+	{
+		Evolution->SetDamping(InSimDataIndex, ClothSharedSimConfig->Damping);
+	}
+	else
+	{
+		Evolution->SetDamping(InSimDataIndex, ChaosClothSimConfig->DampingCoefficient);
+	}
+
+	// Add velocity field
 	auto GetVelocity = [this](const TVector<float, 3>&)->TVector<float, 3>
 	{
 		return WindVelocity;
@@ -309,7 +319,6 @@ void ClothingSimulation::UpdateSimulationFromSharedSimConfig()
 		Evolution->SetIterations(ClothSharedSimConfig->IterationCount);
 		Evolution->SetSelfCollisionThickness(ClothSharedSimConfig->SelfCollisionThickness);
 		Evolution->SetCollisionThickness(ClothSharedSimConfig->CollisionThickness);
-		Evolution->SetDamping(ClothSharedSimConfig->Damping);
 	}
 }
 
@@ -1379,6 +1388,16 @@ void ClothingSimulation::RefreshClothConfig()
 
 				AddConstraints(ChaosClothConfig, PhysMesh, SimDataIndex);
 
+				// Set damping
+				if (ClothSharedSimConfig && ClothSharedSimConfig->bUseDampingOverride)
+				{
+					Evolution->SetDamping(SimDataIndex, ClothSharedSimConfig->Damping);
+				}
+				else
+				{
+					Evolution->SetDamping(SimDataIndex, ChaosClothConfig->DampingCoefficient);
+				}
+
 				// Add Velocity field
 				auto GetVelocity = [this](const TVector<float, 3>&)->TVector<float, 3>
 				{
diff --git a/Engine/Plugins/Experimental/ChaosCloth/Source/Chaos/Public/ChaosCloth/ChaosClothConfig.h b/Engine/Plugins/Experimental/ChaosCloth/Source/Chaos/Public/ChaosCloth/ChaosClothConfig.h
--- a/Engine/Plugins/Experimental/ChaosCloth/Source/Chaos/Public/ChaosCloth/ChaosClothConfig.h
+++ b/Engine/Plugins/Experimental/ChaosCloth/Source/Chaos/Public/ChaosCloth/ChaosClothConfig.h
@@ -86,13 +86,17 @@ public:
 	UPROPERTY()
 	float ShapeTargetStiffness = 0.f;
 
-	// The drag coefficient applying on each particle
-	UPROPERTY(EditAnywhere, Category = Wind, meta = (UIMin = "0", UIMax = "1", ClampMin = "0", ClampMax = "10"))
-	float DragCoefficient = 0.5f;
+	// The amount of damping applied to the cloth velocities
+	UPROPERTY(EditAnywhere, Category = "Cloth Dynamics", meta = (UIMin = "0", UIMax = "1", ClampMin = "0", ClampMax = "1"))
+	float DampingCoefficient = 0.01f;
 
 	// Friction coefficient for cloth - collider interaction
-	UPROPERTY(EditAnywhere, Category = Collision, meta = (UIMin = "0", UIMax = "1", ClampMin = "0", ClampMax = "10"))
-	float CoefficientOfFriction = 0.f;
+	UPROPERTY()
+	float FrictionCoefficient = 0.f;
+
+	// The drag coefficient applying on each particle
+	UPROPERTY(EditAnywhere, Category = "Cloth Dynamics", meta = (UIMin = "0", UIMax = "1", ClampMin = "0", ClampMax = "10"))
+	float DragCoefficient = 0.5f;
 
 	// Default spring stiffness for anim drive if an anim drive is in use
 	UPROPERTY(EditAnywhere, Category = Stiffness, meta = (UIMin = "0", UIMax = "1", ClampMin = "0", ClampMax = "1"))
@@ -160,8 +164,12 @@ public:
 	UPROPERTY(EditAnywhere, Category = Collision, meta = (UIMin = "0", UIMax = "100", ClampMin = "0", ClampMax = "1000"))
 	float CollisionThickness = 1.0f;
 
-	//The amount of cloth damping
-	UPROPERTY(EditAnywhere, Category = Simulation, meta = (UIMin = "0", UIMax = "1", ClampMin = "0", ClampMax = "1"))
+	// Use shared config damping rather than per cloth damping
+	UPROPERTY(EditAnywhere, Category = Simulation, meta = (InlineEditConditionToggle))
+	bool bUseDampingOverride = true;
+
+	// The amount of cloth damping. Override the per cloth damping coefficients.
+	UPROPERTY(EditAnywhere, Category = Simulation, meta = (EditCondition = "bUseDampingOverride", UIMin = "0", UIMax = "1", ClampMin = "0", ClampMax = "1"))
 	float Damping = 0.01f;
 
 	// Use gravity value vs world gravity
diff --git a/Engine/Source/Runtime/Experimental/Chaos/Private/Chaos/PBDEvolution.cpp b/Engine/Source/Runtime/Experimental/Chaos/Private/Chaos/PBDEvolution.cpp
--- a/Engine/Source/Runtime/Experimental/Chaos/Private/Chaos/PBDEvolution.cpp
+++ b/Engine/Source/Runtime/Experimental/Chaos/Private/Chaos/PBDEvolution.cpp
@@ -4,7 +4,7 @@
 #include "Chaos/Framework/Parallel.h"
 #include "Chaos/PBDCollisionSphereConstraints.h"
 #include "Chaos/PBDCollisionSpringConstraints.h"
-#include "Chaos/PerParticleDampVelocity.h"
+#include "Chaos/PerGroupDampVelocity.h"
 #include "Chaos/PerParticleEulerStepVelocity.h"
 #include "Chaos/PerParticleGravity.h"
 #include "Chaos/PerParticleInitForce.h"
@@ -42,7 +42,8 @@ TPBDEvolution<T, d>::TPBDEvolution(TPBDParticles<T, d>&& InParticles, TKinematic
 	, MTime(0)
 {
 	MCollisionParticles.AddArray(&MCollided);
-	
+	MParticles.AddArray(&MParticleGroupIds);
+
 	SetParticleUpdateFunction(
 		[PBDUpdateRule = 
 			TPerParticlePBDUpdateFromDeltaPosition<float, 3>()](TPBDParticles<T, d>& MParticlesInput, const T Dt) 
@@ -55,13 +56,38 @@ TPBDEvolution<T, d>::TPBDEvolution(TPBDParticles<T, d>&& InParticles, TKinematic
 			});
 }
 
+template<class T, int d>
+void TPBDEvolution<T, d>::AddParticles(uint32 Num, uint32 GroupId)
+{
+	// Add new particles
+	const uint32 Offset = MParticles.Size();
+	MParticles.AddParticles(Num);
+
+	// Initialize the new particle's group id
+	for (uint32 i = Offset; i < MParticles.Size(); ++i)
+	{
+		MParticleGroupIds[i] = GroupId;
+	}
+
+	// Resize group parameter arrays
+	const uint32 GroupNum = (uint32)MPerGroupDamping.Num();
+	if (GroupId >= GroupNum)
+	{
+		MPerGroupDamping.SetNum((int32)GroupId + 1);
+		for (uint32 i = GroupNum; i <= GroupId; ++i)
+		{
+			MPerGroupDamping[i] = MDamping;
+		}
+	}
+}
+
 template<class T, int d>
 void TPBDEvolution<T, d>::AdvanceOneTimeStep(const T Dt)
 {
 	SCOPE_CYCLE_COUNTER(STAT_ChaosPBDVAdvanceTime);
 	TPerParticleInitForce<T, d> InitForceRule;
 	TPerParticleEulerStepVelocity<T, d> EulerStepVelocityRule;
-	TPerParticleDampVelocity<T, d> DampVelocityRule(MDamping);
+	TPerGroupDampVelocity<T, d> DampVelocityRule(MParticleGroupIds, MPerGroupDamping);
 	TPerParticlePBDEulerStep<T, d> EulerStepRule;
 	TPerParticlePBDCollisionConstraint<T, d, EGeometryParticlesSimType::Other> CollisionRule(MCollisionParticles, MCollided, MCollisionThickness, MCoefficientOfFriction);
 
diff --git a/Engine/Source/Runtime/Experimental/Chaos/Public/Chaos/PBDEvolution.h b/Engine/Source/Runtime/Experimental/Chaos/Public/Chaos/PBDEvolution.h
--- a/Engine/Source/Runtime/Experimental/Chaos/Public/Chaos/PBDEvolution.h
+++ b/Engine/Source/Runtime/Experimental/Chaos/Public/Chaos/PBDEvolution.h
@@ -29,6 +29,8 @@ class CHAOS_API TPBDEvolution
 	void AddXPBDConstraintFunctions(TFunction<void()> InitConstraintFunction, TFunction<void(TPBDParticles<T, d>&, const T)> ConstraintFunction) { MInitConstraintRules.Add(InitConstraintFunction); MConstraintRules.Add(ConstraintFunction); }
 	void AddForceFunction(TFunction<void(TPBDParticles<T, d>&, const T, const int32)> ForceFunction) { MForceRules.Add(ForceFunction); }
 
+	void AddParticles(uint32 Num, uint32 GroupId = 0);
+
 	const TPBDParticles<T, d>& Particles() const { return MParticles; }
 	TPBDParticles<T, d>& Particles() { return MParticles; }
 
@@ -59,6 +61,7 @@ class CHAOS_API TPBDEvolution
 
 	T GetDamping() const { return MDamping; }
 	void SetDamping(const T Damping) { MDamping = Damping; }
+	void SetDamping(const uint32 GroupId, const T Damping) { MPerGroupDamping[GroupId] = Damping; }
 
 	T GetTime() const { return MTime; }
 
@@ -72,6 +75,8 @@ class CHAOS_API TPBDEvolution
 	TArray<TVector<int32, 3>> MCollisionTriangles;       // Used for self-collisions
 	TSet<TVector<int32, 2>> MDisabledCollisionElements;  // 
 	TArrayCollectionArray<bool> MCollided;
+	TArrayCollectionArray<uint32> MParticleGroupIds;  // Used for per group parameters
+	TArray<T> MPerGroupDamping;
 	int32 MNumIterations;
 	T MCollisionThickness;
 	T MSelfCollisionThickness;
diff --git a/Engine/Source/Runtime/Experimental/Chaos/Public/Chaos/PerGroupDampVelocity.h b/Engine/Source/Runtime/Experimental/Chaos/Public/Chaos/PerGroupDampVelocity.h
new file mode 100644
--- /dev/null
+++ b/Engine/Source/Runtime/Experimental/Chaos/Public/Chaos/PerGroupDampVelocity.h
@@ -0,0 +1,44 @@
+// Copyright 1998-2019 Epic Games, Inc. All Rights Reserved.
+#pragma once
+
+#include "Chaos/PerParticleDampVelocity.h"
+
+namespace Chaos
+{
+	template<class T, int d>
+	class TPerGroupDampVelocity : public TPerParticleDampVelocity<T, d>
+	{
+	public:
+		TPerGroupDampVelocity(const TArray<uint32>& ParticleGroupIds, const TArray<T>& PerGroupDamping)
+			: MParticleGroupIds(ParticleGroupIds)
+			, MPerGroupDamping(PerGroupDamping)
+		{
+		}
+		virtual ~TPerGroupDampVelocity() override {}
+
+		inline void Apply(TDynamicParticles<T, d>& InParticles, const T Dt, const int32 Index) const override
+		{
+			ApplyHelper(InParticles, Dt, Index);
+		}
+
+		inline void Apply(TRigidParticles<T, d>& InParticles, const T Dt, const int32 Index) const override
+		{
+			ApplyHelper(InParticles, Dt, Index);
+		}
+
+	private:
+		template<class T_PARTICLES>
+		inline void ApplyHelper(T_PARTICLES& InParticles, const T Dt, const int32 Index) const
+		{
+			if (InParticles.InvM(Index) != (T)0.)  // Only damp dynamic/rigid particles
+			{
+				TPerParticleDampVelocity<T, d>::MCoefficient = MPerGroupDamping[MParticleGroupIds[Index]];
+				TPerParticleDampVelocity<T, d>::ApplyHelper(InParticles, Dt, Index);
+			}
+		}
+
+	private:
+		const TArray<uint32>& MParticleGroupIds;
+		const TArray<T>& MPerGroupDamping;
+	};
+}
diff --git a/Engine/Source/Runtime/Experimental/Chaos/Public/Chaos/PerParticleDampVelocity.h b/Engine/Source/Runtime/Experimental/Chaos/Public/Chaos/PerParticleDampVelocity.h
--- a/Engine/Source/Runtime/Experimental/Chaos/Public/Chaos/PerParticleDampVelocity.h
+++ b/Engine/Source/Runtime/Experimental/Chaos/Public/Chaos/PerParticleDampVelocity.h
@@ -86,15 +86,6 @@ namespace Chaos
 
 			UpdatePositionBasedState(InParticles, ActiveIndices);
 		}
-
-		template<class T_PARTICLES>
-		inline void ApplyHelper(T_PARTICLES& InParticles, const T Dt, const int32 Index) const
-		{
-			TVector<T, d> R = InParticles.X(Index) - MXcm;
-			TVector<T, d> Dv = MVcm - InParticles.V(Index) + TVector<T, d>::CrossProduct(R, MOmega);
-			InParticles.V(Index) += MCoefficient * Dv;
-		}
-
 		inline void Apply(TDynamicParticles<T, d>& InParticles, const T Dt, const int32 Index) const override //-V762
 		{
 			if (InParticles.InvM(Index) == 0)
@@ -113,9 +104,20 @@ namespace Chaos
 			ApplyHelper(InParticles, Dt, Index);
 		}
 
+	protected:
+		template<class T_PARTICLES>
+		inline void ApplyHelper(T_PARTICLES& InParticles, const T Dt, const int32 Index) const
+		{
+			TVector<T, d> R = InParticles.X(Index) - MXcm;
+			TVector<T, d> Dv = MVcm - InParticles.V(Index) + TVector<T, d>::CrossProduct(R, MOmega);
+			InParticles.V(Index) += MCoefficient * Dv;
+		}
+
+	protected:
+		mutable T MCoefficient;  // The mutable allows to be modified in derived classes Apply const functions
+
 	private:
 		TArray<int32> ActiveIndices;
 		TVector<T, d> MXcm, MVcm, MOmega;
-		T MCoefficient;
 	};
 }
