Chaos Cloth - Improve the damping migration path so that the damping override is disabled when migration happens from an nv cloth asset, but enabled when loading a Chaos Cloth asset which had been authored with the shared damping value.

[RN] none

diff --git a/Engine/Plugins/Experimental/ChaosCloth/Source/Chaos/Private/ChaosCloth/ChaosClothConfig.cpp b/Engine/Plugins/Experimental/ChaosCloth/Source/Chaos/Private/ChaosCloth/ChaosClothConfig.cpp
--- a/Engine/Plugins/Experimental/ChaosCloth/Source/Chaos/Private/ChaosCloth/ChaosClothConfig.cpp
+++ b/Engine/Plugins/Experimental/ChaosCloth/Source/Chaos/Private/ChaosCloth/ChaosClothConfig.cpp
@@ -88,6 +88,8 @@ void UChaosClothSharedSimConfig::MigrateFrom(const FClothConfig_Legacy& ClothCon
 	GravityScale = ClothConfig.GravityScale;
 
 	Gravity = ClothConfig.GravityOverride;
+
+	bUseDampingOverride = false;  // Damping is migrated to per cloth configs
 }
 
 void UChaosClothSharedSimConfig::Serialize(FArchive& Ar)
