Chaos Cloth - Update the wind Drag Coefficient default value to 0.07.

diff --git a/Engine/Plugins/Experimental/ChaosCloth/Source/Chaos/Public/ChaosCloth/ChaosClothConfig.h b/Engine/Plugins/Experimental/ChaosCloth/Source/Chaos/Public/ChaosCloth/ChaosClothConfig.h
--- a/Engine/Plugins/Experimental/ChaosCloth/Source/Chaos/Public/ChaosCloth/ChaosClothConfig.h
+++ b/Engine/Plugins/Experimental/ChaosCloth/Source/Chaos/Public/ChaosCloth/ChaosClothConfig.h
@@ -96,7 +96,7 @@ public:
 
 	// The drag coefficient applying on each particle
 	UPROPERTY(EditAnywhere, Category = "Cloth Dynamics", meta = (UIMin = "0", UIMax = "1", ClampMin = "0", ClampMax = "10"))
-	float DragCoefficient = 0.5f;
+	float DragCoefficient = 0.07f;
 
 	// Default spring stiffness for anim drive if an anim drive is in use
 	UPROPERTY(EditAnywhere, Category = Stiffness, meta = (UIMin = "0", UIMax = "1", ClampMin = "0", ClampMax = "1"))
