Chaos Cloth - Add custom versioning to Chaos Cloth config, and force Drag coefficient to 0.07 for all previously saved assets.

diff --git a/Engine/Plugins/Experimental/ChaosCloth/Source/Chaos/Private/ChaosCloth/ChaosClothConfig.cpp b/Engine/Plugins/Experimental/ChaosCloth/Source/Chaos/Private/ChaosCloth/ChaosClothConfig.cpp
--- a/Engine/Plugins/Experimental/ChaosCloth/Source/Chaos/Private/ChaosCloth/ChaosClothConfig.cpp
+++ b/Engine/Plugins/Experimental/ChaosCloth/Source/Chaos/Private/ChaosCloth/ChaosClothConfig.cpp
@@ -2,6 +2,7 @@
 
 #include "ChaosCloth/ChaosClothConfig.h"
 #include "ClothConfig_Legacy.h"
+#include "ChaosClothConfigCustomVersion.h"
 #include "ChaosClothSharedConfigCustomVersion.h"
 #include "ClothingSimulationInteractor.h"
 
@@ -69,6 +70,23 @@ void UChaosClothConfig::MigrateFrom(const FClothConfig_Legacy& ClothConfig)
 	DampingCoefficient = FMath::Clamp(Damping * Damping * 0.95f, 0.f, 1.f);  // Nv Cloth seems to have a different damping formulation.
 }
 
+void UChaosClothConfig::Serialize(FArchive& Ar)
+{
+	Super::Serialize(Ar);
+	Ar.UsingCustomVersion(FChaosClothConfigCustomVersion::GUID);
+}
+
+void UChaosClothConfig::PostLoad()
+{
+	Super::PostLoad();
+	const int32 ChaosClothConfigCustomVersion = GetLinkerCustomVersion(FChaosClothConfigCustomVersion::GUID);
+
+	if (ChaosClothConfigCustomVersion < FChaosClothConfigCustomVersion::UpdateDragDefault)
+	{
+		DragCoefficient = 0.07f;  // Reset to a more appropirate default for chaos cloth assets saved before this custom version
+	}
+}
+
 UChaosClothSharedSimConfig::UChaosClothSharedSimConfig()
 {}
 
diff --git a/Engine/Plugins/Experimental/ChaosCloth/Source/Chaos/Private/ChaosCloth/ChaosClothConfigCustomVersion.cpp b/Engine/Plugins/Experimental/ChaosCloth/Source/Chaos/Private/ChaosCloth/ChaosClothConfigCustomVersion.cpp
new file mode 100644
--- /dev/null
+++ b/Engine/Plugins/Experimental/ChaosCloth/Source/Chaos/Private/ChaosCloth/ChaosClothConfigCustomVersion.cpp
@@ -0,0 +1,8 @@
+// Copyright Epic Games, Inc. All Rights Reserved.
+
+#include "ChaosClothConfigCustomVersion.h"
+#include "Serialization/CustomVersion.h"
+
+const FGuid FChaosClothConfigCustomVersion::GUID(0x76A52329, 0x092345B5, 0x98AED841, 0xCF2F6AD8);
+// Register the custom version with core
+FCustomVersionRegistration GRegisterChaosClothConfigCustomVersion(FChaosClothConfigCustomVersion::GUID, FChaosClothConfigCustomVersion::LatestVersion, TEXT("ChaosClothConfigVer"));
diff --git a/Engine/Plugins/Experimental/ChaosCloth/Source/Chaos/Private/ChaosCloth/ChaosClothConfigCustomVersion.h b/Engine/Plugins/Experimental/ChaosCloth/Source/Chaos/Private/ChaosCloth/ChaosClothConfigCustomVersion.h
new file mode 100644
--- /dev/null
+++ b/Engine/Plugins/Experimental/ChaosCloth/Source/Chaos/Private/ChaosCloth/ChaosClothConfigCustomVersion.h
@@ -0,0 +1,28 @@
+// Copyright Epic Games, Inc. All Rights Reserved.
+
+#pragma once
+
+#include "CoreMinimal.h"
+#include "Misc/Guid.h"
+
+// Custom serialization version for cloth config
+struct FChaosClothConfigCustomVersion
+{
+	enum Type
+	{
+		// Before any version changes were made
+		BeforeCustomVersionWasAdded = 0,
+		// Update drag default to better preserve legacy behavior
+		UpdateDragDefault = 1,
+
+		// -----<new versions can be added above this line>-------------------------------------------------
+		VersionPlusOne,
+		LatestVersion = VersionPlusOne - 1
+	};
+
+	// The GUID for this custom version number
+	const static FGuid GUID;
+
+private:
+	FChaosClothConfigCustomVersion() {}
+};
diff --git a/Engine/Plugins/Experimental/ChaosCloth/Source/Chaos/Public/ChaosCloth/ChaosClothConfig.h b/Engine/Plugins/Experimental/ChaosCloth/Source/Chaos/Public/ChaosCloth/ChaosClothConfig.h
--- a/Engine/Plugins/Experimental/ChaosCloth/Source/Chaos/Public/ChaosCloth/ChaosClothConfig.h
+++ b/Engine/Plugins/Experimental/ChaosCloth/Source/Chaos/Public/ChaosCloth/ChaosClothConfig.h
@@ -19,6 +19,12 @@ public:
 	/** Migrate from the legacy FClothConfig structure. */
 	virtual void MigrateFrom(const FClothConfig_Legacy&) override;
 
+	/** Serialize override used to set the current custom version. */
+	virtual void Serialize(FArchive& Ar) override;
+
+	/** PostLoad override used to deal with updates/changes in properties. */
+	virtual void PostLoad() override;
+
 	/**
 	 * How cloth particle mass is determined
 	 * -	Uniform Mass: Every particle's mass will be set to the value specified in the UniformMass setting
