Chaos Cloth - Fix legacy wind drag migration to the default 0.07 Drag Coefficient value instad of 0.5.

diff --git a/Engine/Plugins/Experimental/ChaosCloth/Source/Chaos/Private/ChaosCloth/ChaosClothConfig.cpp b/Engine/Plugins/Experimental/ChaosCloth/Source/Chaos/Private/ChaosCloth/ChaosClothConfig.cpp
--- a/Engine/Plugins/Experimental/ChaosCloth/Source/Chaos/Private/ChaosCloth/ChaosClothConfig.cpp
+++ b/Engine/Plugins/Experimental/ChaosCloth/Source/Chaos/Private/ChaosCloth/ChaosClothConfig.cpp
@@ -64,7 +64,7 @@ void UChaosClothConfig::MigrateFrom(const FClothConfig_Legacy& ClothConfig)
 	LimitScale = FMath::Clamp(ClothConfig.TetherLimit, 0.01f, 10.f);
 	ShapeTargetStiffness = 0.f;
 
-	DragCoefficient = (ClothConfig.WindMethod == EClothingWindMethod_Legacy::Accurate) ? ClothConfig.WindDragCoefficient: 0.5f;  // Only Accurate wind uses the WindDragCoefficient
+	DragCoefficient = (ClothConfig.WindMethod == EClothingWindMethod_Legacy::Accurate) ? ClothConfig.WindDragCoefficient: 0.07f;  // Only Accurate wind uses the WindDragCoefficient
 
 	const float Damping = (ClothConfig.Damping.X + ClothConfig.Damping.Y + ClothConfig.Damping.Z) / 3.f;
 	DampingCoefficient = FMath::Clamp(Damping * Damping * 0.95f, 0.f, 1.f);  // Nv Cloth seems to have a different damping formulation.
