CIS - Fix test that was loading cloth assets before Cloth plugins when Chaos is disabled.

diff --git a/Engine/Plugins/Experimental/ChaosCloth/ChaosCloth.uplugin b/Engine/Plugins/Experimental/ChaosCloth/ChaosCloth.uplugin
--- a/Engine/Plugins/Experimental/ChaosCloth/ChaosCloth.uplugin
+++ b/Engine/Plugins/Experimental/ChaosCloth/ChaosCloth.uplugin
@@ -18,7 +18,7 @@
       {
         "Name": "ChaosCloth",
         "Type": "Runtime",
-        "LoadingPhase": "PostEngineInit"
+        "LoadingPhase": "Default"
       }
     ]
 }
