Chaos Cloth - Filter delta time to allievate fps related jittering.

diff --git a/Engine/Plugins/Experimental/ChaosCloth/Source/Chaos/Private/ChaosCloth/ChaosClothingSimulation.cpp b/Engine/Plugins/Experimental/ChaosCloth/Source/Chaos/Private/ChaosCloth/ChaosClothingSimulation.cpp
--- a/Engine/Plugins/Experimental/ChaosCloth/Source/Chaos/Private/ChaosCloth/ChaosClothingSimulation.cpp
+++ b/Engine/Plugins/Experimental/ChaosCloth/Source/Chaos/Private/ChaosCloth/ChaosClothingSimulation.cpp
@@ -125,6 +125,7 @@ void ClothingSimulation::Initialize()
 		});
 
     Time = 0.f;
+	DeltaTime = 1.f / 30.f;  // Initialize filtered timestep at 30fps 
 }
 
 void ClothingSimulation::Shutdown()
@@ -1256,11 +1257,14 @@ void ClothingSimulation::Simulate(IClothingSimulationContext* InContext)
 	// Update collision transforms
 	UpdateCollisionTransforms(*Context, bTeleport);
 
+	// Filter delta time to smoothen time variations and prevent unwanted vibrations
+	static const float Decay = 0.1f;
+	DeltaTime = DeltaTime + (Context->DeltaSeconds - DeltaTime) * Decay;
+
 	// Advance Sim
-	DeltaTime = Context->DeltaSeconds;
 	Evolution->AdvanceOneTimeStep(DeltaTime);
 	Time += DeltaTime;
-	UE_LOG(LogChaosCloth, VeryVerbose, TEXT("DeltaTime: %.6f, Time = %.6f,  MaxPhysicsDelta = %.6f"), DeltaTime, Time, FClothingSimulationCommon::MaxPhysicsDelta);
+	UE_LOG(LogChaosCloth, VeryVerbose, TEXT("DeltaTime: %.6f, FilteredDeltaTime: %.6f, Time = %.6f,  MaxPhysicsDelta = %.6f"), Context->DeltaSeconds, DeltaTime, Time, FClothingSimulationCommon::MaxPhysicsDelta);
 }
 
 void ClothingSimulation::GetSimulationData(
