Local space offset added to cloth to
- Control input acceleration
- Good simulation when character is far from origin
- Preserves centrifugal and coriolis 'fictitious' forces when in rotational component space

diff --git a/Engine/Plugins/Experimental/ChaosCloth/Source/Chaos/Private/ChaosCloth/ChaosClothingSimulation.cpp b/Engine/Plugins/Experimental/ChaosCloth/Source/Chaos/Private/ChaosCloth/ChaosClothingSimulation.cpp
--- a/Engine/Plugins/Experimental/ChaosCloth/Source/Chaos/Private/ChaosCloth/ChaosClothingSimulation.cpp
+++ b/Engine/Plugins/Experimental/ChaosCloth/Source/Chaos/Private/ChaosCloth/ChaosClothingSimulation.cpp
@@ -68,6 +68,14 @@ ClothingSimulation::ClothingSimulation()
 	, bOverrideGravity(false)
 	, Gravity(ClothingSimulationDefault::Gravity)
 	, WindVelocity(FVector::ZeroVector)
+	, bLocalSimSpaceEnabled(false)
+	, LocalSimSpaceOffset(FVector(0.0f, 0.0f, 0.0f))
+	, PrevLocalSimSpaceOffset(FVector(0.0f, 0.0f, 0.0f))
+	, LocalSimSpaceVelocity(FVector(0.0f, 0.0f, 0.0f))
+	, LocalSimSpaceCappedVelocity(FVector(0.0f, 0.0f, 0.0f))
+	, PrevLocalSimSpaceVelocity(FVector(0.0f, 0.0f, 0.0f))
+	, ComponentLinearAccScale(FVector(0.0f, 0.0f, 0.0f))
+	, ComponentLinearAccClamp(FVector(0.0f, 0.0f, 0.0f))
 {
 #if WITH_EDITOR
 	DebugClothMaterial = LoadObject<UMaterial>(nullptr, TEXT("/Engine/EditorMaterials/Cloth/CameraLitDoubleSided.CameraLitDoubleSided"), nullptr, LOAD_None, nullptr);  // LOAD_EditorOnly
@@ -223,10 +231,19 @@ void ClothingSimulation::CreateActor(USkeletalMeshComponent* InOwnerComponent, U
 	AnimationPositions.SetNum(Particles.Size());
 	AnimationNormals.SetNum(Particles.Size());
 
+	//
+	FTransform ComponentToLocalSimSpace = Context.ComponentToWorld;
+	if (bLocalSimSpaceEnabled && Offset == 0)
+	{
+		LocalSimSpaceOffset = ComponentToLocalSimSpace.GetTranslation(); // Only initialize this once for all cloth instances
+		PrevLocalSimSpaceOffset = LocalSimSpaceOffset;
+	}
+	ComponentToLocalSimSpace.AddToTranslation(-LocalSimSpaceOffset);
+
 	ClothingMeshUtils::SkinPhysicsMesh<true, false>(
 		Asset->UsedBoneIndices,
 		PhysMesh, // curr pos and norm
-		Context.ComponentToWorld,
+		ComponentToLocalSimSpace,
 		Context.RefToLocals.GetData(),
 		Context.RefToLocals.Num(),
 		reinterpret_cast<TArray<FVector>&>(AnimationPositions),
@@ -316,6 +333,10 @@ void ClothingSimulation::UpdateSimulationFromSharedSimConfig()
 {
 	if (ClothSharedSimConfig) // ClothSharedSimConfig will be a null pointer if all cloth instances are disabled in which case we will use default Evolution parameters
 	{
+		bLocalSimSpaceEnabled = ClothSharedSimConfig->bUseLocalSpaceSimulation;
+		ComponentLinearAccScale = ClothSharedSimConfig->ComponentLinearAccScale;
+		ComponentLinearAccClamp = ClothSharedSimConfig->ComponentLinearAccClamp;
+
 		// Now set all the common parameters on the simulation
 		Evolution->SetIterations(ClothSharedSimConfig->IterationCount);
 		Evolution->SetSelfCollisionThickness(ClothSharedSimConfig->SelfCollisionThickness);
@@ -709,6 +730,9 @@ void ClothingSimulation::UpdateCollisionTransforms(const ClothingSimulationConte
 		OldCollisionTransforms.SetNum(NumCollisions);
 	}
 
+	FTransform ComponentToLocalSimulationSpace = Context.ComponentToWorld;
+	ComponentToLocalSimulationSpace.AddToTranslation(-LocalSimSpaceOffset);
+
 	// Update the collision transforms
 	for (int32 Index = 0; Index < NumCollisions; ++Index)
 	{
@@ -716,11 +740,11 @@ void ClothingSimulation::UpdateCollisionTransforms(const ClothingSimulationConte
 		if (Context.BoneTransforms.IsValidIndex(BoneIndex))
 		{
 			const FTransform& BoneTransform = Context.BoneTransforms[BoneIndex];
-			CollisionTransforms[Index] = BaseTransforms[Index] * BoneTransform * Context.ComponentToWorld;
+			CollisionTransforms[Index] = BaseTransforms[Index] * BoneTransform * ComponentToLocalSimulationSpace;
 		}
 		else
 		{
-			CollisionTransforms[Index] = BaseTransforms[Index] * Context.ComponentToWorld;  // External collisions often don't map to a bone
+			CollisionTransforms[Index] = BaseTransforms[Index] * ComponentToLocalSimulationSpace;  // External collisions often don't map to a bone
 		}
 	}
 
@@ -1169,6 +1193,47 @@ void ClothingSimulation::Simulate(IClothingSimulationContext* InContext)
 		return;
 	}
 
+	// New LocalSimSpaceOffset (use the component transform for now, but something else like the middle of the bounding volume can be used here)
+	if(bLocalSimSpaceEnabled)
+	{
+		PrevLocalSimSpaceOffset = LocalSimSpaceOffset;
+		LocalSimSpaceOffset = Context->ComponentToWorld.GetTranslation();
+
+		PrevLocalSimSpaceVelocity = LocalSimSpaceVelocity;
+		LocalSimSpaceVelocity = (LocalSimSpaceOffset - PrevLocalSimSpaceOffset) / Context->DeltaSeconds;  // Estimate simspace frame velocity
+		const FVector ScaledVelocityDelta = (LocalSimSpaceVelocity - PrevLocalSimSpaceVelocity) * ComponentLinearAccScale;
+
+		//const FVector ScaledVelocityDeltaMax = ComponentLinearAccClamp * Context->DeltaSeconds;
+		//FVector ScaledClampedVelocityDelta = ScaledVelocityDelta.ComponentMin(ScaledVelocityDeltaMax);
+		//ScaledClampedVelocityDelta = ScaledClampedVelocityDelta.ComponentMax(-ScaledVelocityDeltaMax);
+
+		const FVector SimSpaceDelta = (ScaledVelocityDelta + LocalSimSpaceCappedVelocity) * Context->DeltaSeconds;
+		LocalSimSpaceCappedVelocity = SimSpaceDelta / Context->DeltaSeconds;
+
+		// Particle positions updated to be relative to the new space
+		TPBDParticles<float, 3>& Particles = Evolution->Particles();
+		for (uint32 i = 0; i < Particles.Size(); ++i)
+		{
+			Particles.X(i) -= SimSpaceDelta;
+			Particles.P(i) -= SimSpaceDelta;
+		}
+		TGeometryClothParticles<float, 3>& CollisionParticles = Evolution->CollisionParticles();
+		for (uint32 i = 0; i < CollisionParticles.Size(); i++)
+		{
+			CollisionParticles.X(i) -= SimSpaceDelta;
+		}
+
+		for (int32 i = 0; i < OldAnimationPositions.Num(); i++)
+		{
+			OldAnimationPositions[i] -= SimSpaceDelta;
+		}
+
+		for (int32 i = 0; i < OldCollisionTransforms.Num(); i++)
+		{
+			OldCollisionTransforms[i].AddToTranslation(-SimSpaceDelta);
+		}
+	}
+
 	// Set gravity, using the legacy priority: 1) config override, 2) game override, 3) world gravity
 	Evolution->GetGravityForces().SetAcceleration(Chaos::TVector<float, 3>(
 		(ClothSharedSimConfig && ClothSharedSimConfig->bUseGravityOverride) ? ClothSharedSimConfig->Gravity :
@@ -1212,10 +1277,13 @@ void ClothingSimulation::Simulate(IClothingSimulationContext* InContext)
 			// This saves a lot of Matrix multiplication work later
 			const int32 Offset = IndexToRangeMap[Index][0];
 
+			FTransform ComponentToLocalSimSpace = Context->ComponentToWorld;
+			ComponentToLocalSimSpace.AddToTranslation(-LocalSimSpaceOffset);
+
 			ClothingMeshUtils::SkinPhysicsMesh<true, false>(
 				Asset->UsedBoneIndices,
 				PhysMesh,
-				Context->ComponentToWorld,
+				ComponentToLocalSimSpace,
 				Context->RefToLocals.GetData(),
 				Context->RefToLocals.Num(),
 				reinterpret_cast<TArray<FVector>&>(AnimationPositions),
@@ -1313,7 +1381,7 @@ void ClothingSimulation::GetSimulationData(
 		for (uint32 j = VertexDomain[0]; j < VertexDomain[1]; ++j)
         {
 			const uint32 LocalIndex = j - VertexDomain[0];
-            Data.Positions[LocalIndex] = Evolution->Particles().X(j);
+            Data.Positions[LocalIndex] = Evolution->Particles().X(j) + LocalSimSpaceOffset;
             Data.Normals[LocalIndex] = -PointNormals[i][LocalIndex]; // Note the Normals are inverted due to how barycentric coordinates are calculated (see GetPointBaryAndDist in ClothingMeshUtils.cpp)
 		}
     }
@@ -1484,9 +1552,9 @@ void ClothingSimulation::DebugDrawPhysMeshWired(USkeletalMeshComponent* OwnerCom
 				const FVector& Pos1 = Particles.X(Element.Y);
 				const FVector& Pos2 = Particles.X(Element.Z);
 
-				PDI->DrawLine(Pos0, Pos1, FLinearColor::White, SDPG_World, 0.0f, 0.001f);
-				PDI->DrawLine(Pos1, Pos2, FLinearColor::White, SDPG_World, 0.0f, 0.001f);
-				PDI->DrawLine(Pos2, Pos0, FLinearColor::White, SDPG_World, 0.0f, 0.001f);
+				PDI->DrawLine(Pos0 + LocalSimSpaceOffset, Pos1 + LocalSimSpaceOffset, FLinearColor::White, SDPG_World, 0.0f, 0.001f);
+				PDI->DrawLine(Pos1 + LocalSimSpaceOffset, Pos2 + LocalSimSpaceOffset, FLinearColor::White, SDPG_World, 0.0f, 0.001f);
+				PDI->DrawLine(Pos2 + LocalSimSpaceOffset, Pos0 + LocalSimSpaceOffset, FLinearColor::White, SDPG_World, 0.0f, 0.001f);
 			}
 		}
 	}
@@ -1525,7 +1593,9 @@ void ClothingSimulation::DebugDrawPhysMeshShaded(USkeletalMeshComponent* OwnerCo
 		}
 	}
 
-	MeshBuilder.Draw(PDI, FMatrix::Identity, DebugClothMaterial->GetRenderProxy(), SDPG_World, false, false);
+	FMatrix LocalSimSpaceToWorld(FMatrix::Identity);
+	LocalSimSpaceToWorld.SetOrigin(LocalSimSpaceOffset);
+	MeshBuilder.Draw(PDI, LocalSimSpaceToWorld, DebugClothMaterial->GetRenderProxy(), SDPG_World, false, false);
 }
 
 void ClothingSimulation::DebugDrawPointNormals(USkeletalMeshComponent* OwnerComponent, FPrimitiveDrawInterface* PDI) const
@@ -1546,7 +1616,7 @@ void ClothingSimulation::DebugDrawPointNormals(USkeletalMeshComponent* OwnerComp
 				const TVector<float, 3>& Pos = Particles.X(ParticleIndex);
 				const TVector<float, 3>& Normal = MeshPointNormals[ParticleIndex - Range[0]];
 
-				PDI->DrawLine(Pos, Pos + Normal * 20.0f, FLinearColor::White, SDPG_World, 0.0f, 0.001f);
+				PDI->DrawLine(Pos + LocalSimSpaceOffset, Pos + Normal * 20.0f + LocalSimSpaceOffset, FLinearColor::White, SDPG_World, 0.0f, 0.001f);
 			}
 		}
 	}
@@ -1570,7 +1640,7 @@ void ClothingSimulation::DebugDrawInversedPointNormals(USkeletalMeshComponent* O
 				const TVector<float, 3>& Pos = Particles.X(ParticleIndex);
 				const TVector<float, 3>& Normal = MeshPointNormals[ParticleIndex - Range[0]];
 
-				PDI->DrawLine(Pos, Pos - Normal * 20.0f, FLinearColor::White, SDPG_World, 0.0f, 0.001f);
+				PDI->DrawLine(Pos + LocalSimSpaceOffset, Pos - Normal * 20.0f + LocalSimSpaceOffset, FLinearColor::White, SDPG_World, 0.0f, 0.001f);
 			}
 		}
 	}
@@ -1599,7 +1669,7 @@ void ClothingSimulation::DebugDrawFaceNormals(USkeletalMeshComponent* OwnerCompo
 					Particles.X(Element.Z)) / 3.f;
 				const TVector<float, 3>& Normal = MeshFaceNormals[ElementIndex];
 
-				PDI->DrawLine(Pos, Pos + Normal * 20.0f, FLinearColor::Yellow, SDPG_World, 0.0f, 0.001f);
+				PDI->DrawLine(Pos + LocalSimSpaceOffset, Pos + Normal * 20.0f + LocalSimSpaceOffset, FLinearColor::Yellow, SDPG_World, 0.0f, 0.001f);
 			}
 		}
 	}
@@ -1628,7 +1698,7 @@ void ClothingSimulation::DebugDrawInversedFaceNormals(USkeletalMeshComponent* Ow
 					Particles.X(Element.Z)) / 3.f;
 				const TVector<float, 3>& Normal = MeshFaceNormals[ElementIndex];
 
-				PDI->DrawLine(Pos, Pos - Normal * 20.0f, FLinearColor::Yellow, SDPG_World, 0.0f, 0.001f);
+				PDI->DrawLine(Pos + LocalSimSpaceOffset, Pos - Normal * 20.0f + LocalSimSpaceOffset, FLinearColor::Yellow, SDPG_World, 0.0f, 0.001f);
 			}
 		}
 	}
@@ -1636,32 +1706,32 @@ void ClothingSimulation::DebugDrawInversedFaceNormals(USkeletalMeshComponent* Ow
 
 void ClothingSimulation::DebugDrawCollision(USkeletalMeshComponent* OwnerComponent, FPrimitiveDrawInterface* PDI) const
 {
-	auto DrawSphere = [&PDI](const Chaos::TSphere<float, 3>& Sphere, const TRotation<float, 3>& Rotation, const Chaos::TVector<float, 3>& Position, const FLinearColor& Color)
+	auto DrawSphere = [&PDI, this](const Chaos::TSphere<float, 3>& Sphere, const TRotation<float, 3>& Rotation, const Chaos::TVector<float, 3>& Position, const FLinearColor& Color)
 	{
 		const float Radius = Sphere.GetRadius();
 		const Chaos::TVector<float, 3> Center = Sphere.GetCenter();
-		const FTransform Transform(Rotation, Position + Rotation.RotateVector(Center));
+		const FTransform Transform(Rotation, Position + Rotation.RotateVector(Center) + LocalSimSpaceOffset);
 		DrawWireSphere(PDI, Transform, Color, Radius, 12, SDPG_World, 0.0f, 0.001f, false);
 	};
 
-	auto DrawBox = [&PDI](const Chaos::TBox<float, 3>& Box, const TRotation<float, 3>& Rotation, const Chaos::TVector<float, 3>& Position, const FLinearColor& Color)
+	auto DrawBox = [&PDI, this](const Chaos::TBox<float, 3>& Box, const TRotation<float, 3>& Rotation, const Chaos::TVector<float, 3>& Position, const FLinearColor& Color)
 	{
-		const FMatrix BoxToWorld = FTransform(Rotation, Position).ToMatrixNoScale();
+		const FMatrix BoxToWorld = FTransform(Rotation, Position + LocalSimSpaceOffset).ToMatrixNoScale();
 		const FVector Radii = Box.Extents() * 0.5f;
 		DrawWireBox(PDI, BoxToWorld, FBox(Box.Min(), Box.Max()), Color, SDPG_World, 0.0f, 0.001f, false);
 	};
 
-	auto DrawCapsule = [&PDI](const Chaos::TCapsule<float>& Capsule, const TRotation<float, 3>& Rotation, const Chaos::TVector<float, 3>& Position, const FLinearColor& Color)
+	auto DrawCapsule = [&PDI, this](const Chaos::TCapsule<float>& Capsule, const TRotation<float, 3>& Rotation, const Chaos::TVector<float, 3>& Position, const FLinearColor& Color)
 	{
 		const float HalfHeight = Capsule.GetHeight() * 0.5f;
 		const float Radius = Capsule.GetRadius();
 		const FVector X = Rotation.RotateVector(FVector::ForwardVector);
 		const FVector Y = Rotation.RotateVector(FVector::RightVector);
 		const FVector Z = Rotation.RotateVector(FVector::UpVector);
-		DrawWireCapsule(PDI, Position, X, Y, Z, Color, Radius, HalfHeight + Radius, 12, SDPG_World, 0.0f, 0.001f, false);
+		DrawWireCapsule(PDI, Position + LocalSimSpaceOffset, X, Y, Z, Color, Radius, HalfHeight + Radius, 12, SDPG_World, 0.0f, 0.001f, false);
 	};
 
-	auto DrawTaperedCylinder = [&PDI](const Chaos::TTaperedCylinder<float>& TaperedCylinder, const TRotation<float, 3>& Rotation, const Chaos::TVector<float, 3>& Position, const FLinearColor& Color)
+	auto DrawTaperedCylinder = [&PDI, this](const Chaos::TTaperedCylinder<float>& TaperedCylinder, const TRotation<float, 3>& Rotation, const Chaos::TVector<float, 3>& Position, const FLinearColor& Color)
 	{
 		const float HalfHeight = TaperedCylinder.GetHeight() * 0.5f;
 		const float Radius1 = TaperedCylinder.GetRadius1();
@@ -1684,16 +1754,16 @@ void ClothingSimulation::DebugDrawCollision(USkeletalMeshComponent* OwnerCompone
 			const FVector Vertex1 = Position1 + ArcPos * Radius1;
 			const FVector Vertex2 = Position2 + ArcPos * Radius2;
 
-			PDI->DrawLine(LastVertex1, Vertex1, Color, SDPG_World, 0.0f, 0.001f, false);
-			PDI->DrawLine(LastVertex2, Vertex2, Color, SDPG_World, 0.0f, 0.001f, false);
-			PDI->DrawLine(LastVertex1, LastVertex2, Color, SDPG_World, 0.0f, 0.001f, false);
+			PDI->DrawLine(LastVertex1 + LocalSimSpaceOffset, Vertex1 + LocalSimSpaceOffset, Color, SDPG_World, 0.0f, 0.001f, false);
+			PDI->DrawLine(LastVertex2 + LocalSimSpaceOffset, Vertex2 + LocalSimSpaceOffset, Color, SDPG_World, 0.0f, 0.001f, false);
+			PDI->DrawLine(LastVertex1 + LocalSimSpaceOffset, LastVertex2 + LocalSimSpaceOffset, Color, SDPG_World, 0.0f, 0.001f, false);
 
 			LastVertex1 = Vertex1;
 			LastVertex2 = Vertex2;
 		}
 	};
 
-	auto DrawConvex = [&PDI](const Chaos::FConvex& Convex, const TRotation<float, 3>& Rotation, const Chaos::TVector<float, 3>& Position, const FLinearColor& Color)
+	auto DrawConvex = [&PDI, this](const Chaos::FConvex& Convex, const TRotation<float, 3>& Rotation, const Chaos::TVector<float, 3>& Position, const FLinearColor& Color)
 	{
 		const TArray<Chaos::TPlaneConcrete<float, 3>>& Planes = Convex.GetFaces();
 		for (int32 PlaneIndex1 = 0; PlaneIndex1 < Planes.Num(); ++PlaneIndex1)
@@ -1720,7 +1790,7 @@ void ClothingSimulation::DebugDrawCollision(USkeletalMeshComponent* OwnerCompone
 							const Chaos::TVector<float, 3>& X1 = SurfaceParticles.X(ParticleIndex1);
 							const FVector Position1 = Position + Rotation.RotateVector(X1);
 							const FVector Position2 = Position + Rotation.RotateVector(X);
-							PDI->DrawLine(Position1, Position2, Color, SDPG_World, 0.0f, 0.001f, false);
+							PDI->DrawLine(Position1 + LocalSimSpaceOffset, Position2 + LocalSimSpaceOffset, Color, SDPG_World, 0.0f, 0.001f, false);
 							break;
 						}
 						ParticleIndex1 = ParticleIndex;
@@ -1786,7 +1856,7 @@ void ClothingSimulation::DebugDrawCollision(USkeletalMeshComponent* OwnerCompone
 				break;
 
 			default:
-				DrawCoordinateSystem(PDI, Position, FRotator(Rotation), 10.0f, SDPG_World, 0.1f);  // Draw everything else as a coordinate for now
+				DrawCoordinateSystem(PDI, Position + LocalSimSpaceOffset, FRotator(Rotation), 10.0f, SDPG_World, 0.1f);  // Draw everything else as a coordinate for now
 				break;
 			}
 		}
@@ -1819,12 +1889,12 @@ void ClothingSimulation::DebugDrawBackstops(USkeletalMeshComponent* OwnerCompone
 			const uint32 WeightMapIndex = ParticleIndex - IndexToRangeMap[i][0];
 			const float Radius = BackstopRadiuses[WeightMapIndex];
 			const float Distance = BackstopDistances[WeightMapIndex];
-			PDI->DrawLine(AnimationPositions[ParticleIndex], AnimationPositions[ParticleIndex] - AnimationNormals[ParticleIndex] * (Distance - Radius), FColor::White, SDPG_World, 0.0f, 0.001f);
+			PDI->DrawLine(AnimationPositions[ParticleIndex] + LocalSimSpaceOffset, AnimationPositions[ParticleIndex] - AnimationNormals[ParticleIndex] * (Distance - Radius) + LocalSimSpaceOffset, FColor::White, SDPG_World, 0.0f, 0.001f);
 			if (Radius > 0.0f)
 			{
 				const FVector& Normal = AnimationNormals[ParticleIndex];
 				const FVector& Position = AnimationPositions[ParticleIndex];
-				auto DrawBackstop = [Radius, Distance, &Normal, &Position, PDI](const FVector& Axis, const FColor& Color)
+				auto DrawBackstop = [Radius, Distance, &Normal, &Position, PDI, this](const FVector& Axis, const FColor& Color)
 				{
 					const float ArcLength = 5.0f; // Arch length in cm
 					const float ArcAngle = ArcLength * 360.0f / (Radius * 2.0f * PI);
@@ -1832,7 +1902,7 @@ void ClothingSimulation::DebugDrawBackstops(USkeletalMeshComponent* OwnerCompone
 					const float MaxCosAngle = 0.99f;
 					if (FMath::Abs(FVector::DotProduct(Normal, Axis)) < MaxCosAngle)
 					{
-						DrawArc(PDI, Position - Normal * Distance, Normal, FVector::CrossProduct(Axis, Normal).GetSafeNormal(), -ArcAngle / 2.0f, ArcAngle / 2.0f, Radius, 10, Color, SDPG_World);
+						DrawArc(PDI, Position - Normal * Distance + LocalSimSpaceOffset, Normal, FVector::CrossProduct(Axis, Normal).GetSafeNormal(), -ArcAngle / 2.0f, ArcAngle / 2.0f, Radius, 10, Color, SDPG_World);
 					}
 				};
 				DrawBackstop(FVector::ForwardVector, FColor::Blue);
@@ -1873,11 +1943,11 @@ void ClothingSimulation::DebugDrawMaxDistances(USkeletalMeshComponent* OwnerComp
 				const FMatrix& ViewMatrix = PDI->View->ViewMatrices.GetViewMatrix();
 				const FVector& XAxis = ViewMatrix.GetColumn(0); // Just using transpose here (orthogonal transform assumed)
 				const FVector& YAxis = ViewMatrix.GetColumn(1);
-				DrawDisc(PDI, AnimationPositions[ParticleIndex], XAxis, YAxis, FColor::White, 0.2f, 10, DebugClothMaterialVertex->GetRenderProxy(), SDPG_World);
+				DrawDisc(PDI, AnimationPositions[ParticleIndex] + LocalSimSpaceOffset, XAxis, YAxis, FColor::White, 0.2f, 10, DebugClothMaterialVertex->GetRenderProxy(), SDPG_World);
 			}
 			else
 			{
-				PDI->DrawLine(AnimationPositions[ParticleIndex], AnimationPositions[ParticleIndex] + AnimationNormals[ParticleIndex] * Distance, FColor::White, SDPG_World, 0.0f, 0.001f);
+				PDI->DrawLine(AnimationPositions[ParticleIndex] + LocalSimSpaceOffset, AnimationPositions[ParticleIndex] + AnimationNormals[ParticleIndex] * Distance + LocalSimSpaceOffset, FColor::White, SDPG_World, 0.0f, 0.001f);
 			}
 		}
 	}
@@ -1908,7 +1978,7 @@ void ClothingSimulation::DebugDrawAnimDrive(USkeletalMeshComponent* OwnerCompone
 		{
 			const uint32 WeightMapIndex = ParticleIndex - IndexToRangeMap[i][0];
 			const float Multiplier = AnimDriveMultipliers[WeightMapIndex];
-			PDI->DrawLine(AnimationPositions[ParticleIndex], Particles.X(ParticleIndex), Multiplier * AnimDriveSpringStiffness[i] * FColor::Cyan, SDPG_World, 0.0f, 0.001f);
+			PDI->DrawLine(AnimationPositions[ParticleIndex] + LocalSimSpaceOffset, Particles.X(ParticleIndex) + LocalSimSpaceOffset, Multiplier * AnimDriveSpringStiffness[i] * FColor::Cyan, SDPG_World, 0.0f, 0.001f);
 		}
 	}
 }
@@ -1940,8 +2010,8 @@ void ClothingSimulation::DebugDrawLongRangeConstraint(USkeletalMeshComponent* Ow
 			const Chaos::TVector<float, 3> Direction = (Particles.X(Path[Path.Num() - 2]) - P1).GetSafeNormal();
 			const Chaos::TVector<float, 3> P2 = P1 + Direction * Offset;
 
-			PDI->DrawLine(P0, P1, FColor::Purple, SDPG_World, 0.0f, 0.001f);
-			PDI->DrawLine(P1, P2, FColor::Black, SDPG_World, 0.0f, 0.001f);
+			PDI->DrawLine(P0 + LocalSimSpaceOffset, P1 + LocalSimSpaceOffset, FColor::Purple, SDPG_World, 0.0f, 0.001f);
+			PDI->DrawLine(P1 + LocalSimSpaceOffset, P2 + LocalSimSpaceOffset, FColor::Black, SDPG_World, 0.0f, 0.001f);
 		}
 	}
 }
@@ -1964,7 +2034,7 @@ void ClothingSimulation::DebugDrawWindDragForces(USkeletalMeshComponent* OwnerCo
 				Particles.X(Element[1]) +
 				Particles.X(Element[2]));
 			const TVector<float, 3>& Force = Forces[ElementIndex];
-			PDI->DrawLine(Position, Position + Force, FColor::Green, SDPG_World, 0.0f, 0.001f);
+			PDI->DrawLine(Position + LocalSimSpaceOffset, Position + Force + LocalSimSpaceOffset, FColor::Green, SDPG_World, 0.0f, 0.001f);
 		}
 	}
 }
diff --git a/Engine/Plugins/Experimental/ChaosCloth/Source/Chaos/Public/ChaosCloth/ChaosClothConfig.h b/Engine/Plugins/Experimental/ChaosCloth/Source/Chaos/Public/ChaosCloth/ChaosClothConfig.h
--- a/Engine/Plugins/Experimental/ChaosCloth/Source/Chaos/Public/ChaosCloth/ChaosClothConfig.h
+++ b/Engine/Plugins/Experimental/ChaosCloth/Source/Chaos/Public/ChaosCloth/ChaosClothConfig.h
@@ -194,4 +194,14 @@ public:
 	// Experimental, this feature might be removed without warning, not for production use
 	UPROPERTY(EditAnywhere, Category = "Experimental", meta = (DisplayName="Use XPBD Constraints (Experimental)"))
 	bool bUseXPBDConstraints = false;
+
+	// Enable local space simulation
+	UPROPERTY(EditAnywhere, Category = LocalSpaceSimulation, meta = (DisplayName = "Use Local space simulation"))
+	bool bUseLocalSpaceSimulation = false;
+
+	UPROPERTY(EditAnywhere, Category = LocalSpaceSimulation, meta = (DisplayName = "Component Linear Acc Scale"))
+	FVector ComponentLinearAccScale = { 1.f, 1.f, 1.f };
+
+	UPROPERTY()
+	FVector ComponentLinearAccClamp = { 4000.f, 4000.f, 4000.f };
 };
diff --git a/Engine/Plugins/Experimental/ChaosCloth/Source/Chaos/Public/ChaosCloth/ChaosClothingSimulation.h b/Engine/Plugins/Experimental/ChaosCloth/Source/Chaos/Public/ChaosCloth/ChaosClothingSimulation.h
--- a/Engine/Plugins/Experimental/ChaosCloth/Source/Chaos/Public/ChaosCloth/ChaosClothingSimulation.h
+++ b/Engine/Plugins/Experimental/ChaosCloth/Source/Chaos/Public/ChaosCloth/ChaosClothingSimulation.h
@@ -148,6 +148,18 @@ namespace Chaos
 
 		TArray<TSharedPtr<TPBDLongRangeConstraintsBase<float, 3>>> LongRangeConstraints;
 
+		// This is used to translate between world space and simulation space. Add this to simulation space coordinates to get world space coordinates
+		// The function of this is to help with floating point precision errors if the character is far away form the world origin
+		// and to decouple the simulation from character acceleration and velocities if it is desired
+		bool	bLocalSimSpaceEnabled;
+		FVector LocalSimSpaceOffset;  
+		FVector PrevLocalSimSpaceOffset;
+		FVector LocalSimSpaceVelocity;
+		FVector LocalSimSpaceCappedVelocity;
+		FVector PrevLocalSimSpaceVelocity;
+		FVector ComponentLinearAccScale;
+		FVector ComponentLinearAccClamp;
+
 #if WITH_EDITOR
 		// Visualization material
 		UMaterial* DebugClothMaterial;
