Chaos Cloth - Fix NaNs during collision particle creation.
- Add identity particle transform requirement to the SetDynamicGeometry documentation.

diff --git a/Engine/Plugins/Experimental/ChaosCloth/Source/Chaos/Private/ChaosCloth/ChaosClothingSimulation.cpp b/Engine/Plugins/Experimental/ChaosCloth/Source/Chaos/Private/ChaosCloth/ChaosClothingSimulation.cpp
--- a/Engine/Plugins/Experimental/ChaosCloth/Source/Chaos/Private/ChaosCloth/ChaosClothingSimulation.cpp
+++ b/Engine/Plugins/Experimental/ChaosCloth/Source/Chaos/Private/ChaosCloth/ChaosClothingSimulation.cpp
@@ -981,6 +981,10 @@ void ClothingSimulation::AddCollisions(const FClothCollisionData& ClothCollision
 		CapsuleEnds.Reserve(NumCapsules * 2);
 		for (uint32 i = Offset; i < CollisionParticles.Size(); ++i)
 		{
+			// Always initialize the collision particle transforms before setting any geometry as otherwise NaNs gets detected during the bounding box updates
+			CollisionParticles.X(i) = Chaos::TVector<float, 3>(0.f);
+			CollisionParticles.R(i) = Chaos::TRotation<float, 3>::FromIdentity();
+
 			const FClothCollisionPrim_SphereConnection& Connection = ClothCollisionData.SphereConnections[i - Offset];
 
 			const int32 SphereIndex0 = Connection.SphereIndices[0];
@@ -1067,27 +1071,31 @@ void ClothingSimulation::AddCollisions(const FClothCollisionData& ClothCollision
 	{
 		const uint32 Offset = CollisionParticles.Size();
 		CollisionParticles.AddParticles(NumSpheres);
-		// i = Spheres index, j = CollisionParticles index
-		for (uint32 i = 0, j = Offset; i < (uint32)ClothCollisionData.Spheres.Num(); ++i)
+		// i = CollisionParticles index, j = Spheres index
+		for (uint32 i = Offset, j = 0; j < (uint32)ClothCollisionData.Spheres.Num(); ++j)
 		{
 			// Skip spheres that are the end caps of capsules.
-			if (CapsuleEnds.Contains(i))
+			if (CapsuleEnds.Contains(j))
 				continue;
 
-			const FClothCollisionPrim_Sphere& Sphere = ClothCollisionData.Spheres[i];
+			const FClothCollisionPrim_Sphere& Sphere = ClothCollisionData.Spheres[j];
+
+			// Always initialize the collision particle transforms before setting any geometry as otherwise NaNs gets detected during the bounding box updates
+			CollisionParticles.X(i) = Chaos::TVector<float, 3>(0.f);
+			CollisionParticles.R(i) = Chaos::TRotation<float, 3>::FromIdentity();
 
-			BoneIndices[j] = GetMappedBoneIndex(UsedBoneIndices, Sphere.BoneIndex);
-			UE_LOG(LogChaosCloth, VeryVerbose, TEXT("Found collision sphere on bone index %d."), BoneIndices[j]);
+			BoneIndices[i] = GetMappedBoneIndex(UsedBoneIndices, Sphere.BoneIndex);
+			UE_LOG(LogChaosCloth, VeryVerbose, TEXT("Found collision sphere on bone index %d."), BoneIndices[i]);
 
-			BaseTransforms[j] = Chaos::TRigidTransform<float, 3>(FTransform::Identity);
+			BaseTransforms[i] = Chaos::TRigidTransform<float, 3>(FTransform::Identity);
 
 			CollisionParticles.SetDynamicGeometry(
-				j,
+				i,
 				MakeUnique<Chaos::TSphere<float, 3>>(
 					Sphere.LocalPosition,
 					Sphere.Radius));
 
-			++j;
+			++i;
 		}
 	}
 
@@ -1101,6 +1109,10 @@ void ClothingSimulation::AddCollisions(const FClothCollisionData& ClothCollision
 		{
 			const FClothCollisionPrim_Convex& Convex = ClothCollisionData.Convexes[i - Offset];
 
+			// Always initialize the collision particle transforms before setting any geometry as otherwise NaNs gets detected during the bounding box updates
+			CollisionParticles.X(i) = Chaos::TVector<float, 3>(0.f);
+			CollisionParticles.R(i) = Chaos::TRotation<float, 3>::FromIdentity();
+
 			BaseTransforms[i] = Chaos::TRigidTransform<float, 3>(FTransform::Identity);
 
 			BoneIndices[i] = GetMappedBoneIndex(UsedBoneIndices, Convex.BoneIndex);
@@ -1172,6 +1184,10 @@ void ClothingSimulation::AddCollisions(const FClothCollisionData& ClothCollision
 		{
 			const FClothCollisionPrim_Box& Box = ClothCollisionData.Boxes[i - Offset];
 			
+			// Always initialize the collision particle transforms before setting any geometry as otherwise NaNs gets detected during the bounding box updates
+			CollisionParticles.X(i) = Chaos::TVector<float, 3>(0.f);
+			CollisionParticles.R(i) = Chaos::TRotation<float, 3>::FromIdentity();
+
 			BaseTransforms[i] = Chaos::TRigidTransform<float, 3>(Box.LocalPosition, Box.LocalRotation);
 			
 			BoneIndices[i] = GetMappedBoneIndex(UsedBoneIndices, Box.BoneIndex);
diff --git a/Engine/Source/Runtime/Experimental/Chaos/Public/Chaos/GeometryParticles.h b/Engine/Source/Runtime/Experimental/Chaos/Public/Chaos/GeometryParticles.h
--- a/Engine/Source/Runtime/Experimental/Chaos/Public/Chaos/GeometryParticles.h
+++ b/Engine/Source/Runtime/Experimental/Chaos/Public/Chaos/GeometryParticles.h
@@ -245,7 +245,7 @@ namespace Chaos
 		CHAOS_API FParticleID ParticleID(const int32 Idx) const { return MParticleIDs[Idx]; }
 		CHAOS_API FParticleID& ParticleID(const int32 Idx) { return MParticleIDs[Idx]; }
 #endif
-
+		// Set a dynamic geometry. Note that X and R must be initialized before calling this function.
 		CHAOS_API void SetDynamicGeometry(const int32 Index, TUniquePtr<FImplicitObject>&& InUnique)
 		{
 			check(!SharedGeometry(Index));	// If shared geometry exists we should not be setting dynamic geometry on top
@@ -255,6 +255,7 @@ namespace Chaos
 			MapImplicitShapes(Index);
 		}
 
+		// Set a shared geometry. Note that X and R must be initialized before calling this function.
 		CHAOS_API void SetSharedGeometry(const int32 Index, TSharedPtr<FImplicitObject, ESPMode::ThreadSafe> InShared)
 		{
 			check(!DynamicGeometry(Index));	// If dynamic geometry exists we should not be setting shared geometry on top
