Chaos Cloth - Add debug draw in game.

diff --git a/Engine/Plugins/Experimental/ChaosCloth/Source/Chaos/Private/ChaosCloth/ChaosClothingSimulation.cpp b/Engine/Plugins/Experimental/ChaosCloth/Source/Chaos/Private/ChaosCloth/ChaosClothingSimulation.cpp
--- a/Engine/Plugins/Experimental/ChaosCloth/Source/Chaos/Private/ChaosCloth/ChaosClothingSimulation.cpp
+++ b/Engine/Plugins/Experimental/ChaosCloth/Source/Chaos/Private/ChaosCloth/ChaosClothingSimulation.cpp
@@ -44,6 +44,26 @@
 #include "PhysXIncludes.h"
 #endif
 
+#if CHAOS_DEBUG_DRAW
+#include "Chaos/DebugDrawQueue.h"
+#include "HAL/IConsoleManager.h"
+
+namespace ChaosClothingSimulationConsoleVariables
+{
+	TAutoConsoleVariable<bool> CVarDebugDrawPhysMeshWired   (TEXT("p.ChaosCloth.DebugDrawPhysMeshWired"       ), false, TEXT("Whether to debug draw the Chaos Cloth wireframe meshes"), ECVF_Cheat);
+	TAutoConsoleVariable<bool> CVarDebugPointNormals        (TEXT("p.ChaosCloth.DebugDrawPointNormals"        ), false, TEXT("Whether to debug draw the Chaos Cloth point normals"), ECVF_Cheat);
+	TAutoConsoleVariable<bool> CVarDebugInversedPointNormals(TEXT("p.ChaosCloth.DebugDrawInversedPointNormals"), false, TEXT("Whether to debug draw the Chaos Cloth inversed point normals"), ECVF_Cheat);
+	TAutoConsoleVariable<bool> CVarDebugFaceNormals         (TEXT("p.ChaosCloth.DebugDrawFaceNormals"         ), false, TEXT("Whether to debug draw the Chaos Cloth face normals"), ECVF_Cheat);
+	TAutoConsoleVariable<bool> CVarDebugInversedFaceNormals (TEXT("p.ChaosCloth.DebugDrawInversedFaceNormals" ), false, TEXT("Whether to debug draw the Chaos Cloth inversed face normals"), ECVF_Cheat);
+	TAutoConsoleVariable<bool> CVarDebugCollision           (TEXT("p.ChaosCloth.DebugDrawCollision"           ), false, TEXT("Whether to debug draw the Chaos Cloth collisions"), ECVF_Cheat);
+	TAutoConsoleVariable<bool> CVarDebugBackstops           (TEXT("p.ChaosCloth.DebugDrawBackstops"           ), false, TEXT("Whether to debug draw the Chaos Cloth backstops"), ECVF_Cheat);
+	TAutoConsoleVariable<bool> CVarDebugMaxDistances        (TEXT("p.ChaosCloth.DebugDrawMaxDistances"        ), false, TEXT("Whether to debug draw the Chaos Cloth max distances"), ECVF_Cheat);
+	TAutoConsoleVariable<bool> CVarDebugAnimDrive           (TEXT("p.ChaosCloth.DebugDrawAnimDrive"           ), false, TEXT("Whether to debug draw the Chaos Cloth anim drive"), ECVF_Cheat);
+	TAutoConsoleVariable<bool> CVarDebugLongRangeConstraint (TEXT("p.ChaosCloth.DebugDrawLongRangeConstraint" ), false, TEXT("Whether to debug draw the Chaos Cloth long range constraint (aka tether constraint)"), ECVF_Cheat);
+	TAutoConsoleVariable<bool> CVarDebugWindDragForces      (TEXT("p.ChaosCloth.DebugDrawWindDragForces"      ), false, TEXT("Whether to debug draw the Chaos Cloth wind drag forces"), ECVF_Cheat);
+}
+#endif  // #if CHAOS_DEBUG_DRAW
+
 using namespace Chaos;
 
 DECLARE_CYCLE_STAT(TEXT("Chaos Cloth Simulate"), STAT_ChaosClothSimulate, STATGROUP_ChaosCloth);
@@ -51,7 +71,7 @@ DECLARE_CYCLE_STAT(TEXT("Chaos Cloth Get Animation Data"), STAT_ChaosClothGetAni
 DECLARE_CYCLE_STAT(TEXT("Chaos Cloth Update Collision Transforms"), STAT_ChaosClothUpdateCollisionTransforms, STATGROUP_ChaosCloth);
 
 // Default parameters, will be overwritten when cloth assets are loaded
-namespace ClothingSimulationDefault
+namespace ChaosClothingSimulationDefault
 {
 	static const FVector Gravity(0.f, 0.f, -980.665f);
 	static const int32 NumIterations = 1;
@@ -66,7 +86,7 @@ ClothingSimulation::ClothingSimulation()
 	: ClothSharedSimConfig(nullptr)
 	, ExternalCollisionsOffset(0)
 	, bOverrideGravity(false)
-	, Gravity(ClothingSimulationDefault::Gravity)
+	, Gravity(ChaosClothingSimulationDefault::Gravity)
 	, WindVelocity(FVector::ZeroVector)
 	, bLocalSimSpaceEnabled(false)
 	, LocalSimSpaceOffset(FVector(0.0f, 0.0f, 0.0f))
@@ -95,11 +115,11 @@ void ClothingSimulation::Initialize()
 			MoveTemp(LocalParticles),
 			MoveTemp(TRigidParticles),
 			{}, // CollisionTriangles
-			ClothingSimulationDefault::NumIterations,
-			ClothingSimulationDefault::CollisionThickness,
-			ClothingSimulationDefault::SelfCollisionThickness,
-			ClothingSimulationDefault::FrictionCoefficient,
-			ClothingSimulationDefault::DampingCoefficient));
+			ChaosClothingSimulationDefault::NumIterations,
+			ChaosClothingSimulationDefault::CollisionThickness,
+			ChaosClothingSimulationDefault::SelfCollisionThickness,
+			ChaosClothingSimulationDefault::FrictionCoefficient,
+			ChaosClothingSimulationDefault::DampingCoefficient));
     Evolution->CollisionParticles().AddArray(&BoneIndices);
 	Evolution->CollisionParticles().AddArray(&BaseTransforms);
     Evolution->GetGravityForces().SetAcceleration(Gravity);
@@ -422,7 +442,7 @@ void ClothingSimulation::SetParticleMasses(const UChaosClothConfig* ChaosClothCo
 	}
 	case EClothMassMode::Density:
 	{
-		const float Density = ChaosClothConfig->Density / FMath::Square(ClothingSimulationDefault::WorldScale);
+		const float Density = ChaosClothConfig->Density / FMath::Square(ChaosClothingSimulationDefault::WorldScale);
 		for (const int32 Vertex : Vertices)
 		{
 			Particles.M(Vertex) *= Density;
@@ -432,7 +452,7 @@ void ClothingSimulation::SetParticleMasses(const UChaosClothConfig* ChaosClothCo
 	}
 	};
 	UE_LOG(LogChaosCloth, Verbose, TEXT("Density: %f, Total surface: %f, Total mass: %f, "), TotalArea > 0.f ? TotalMass / TotalArea : 1.f, TotalArea, TotalMass);
-	UE_LOG(LogChaosCloth, Verbose, TEXT("SI Density: %f, SI Total surface: %f, SI Total mass: %f, "), (TotalArea > 0.f ? TotalMass / TotalArea : 1.f) * FMath::Square(ClothingSimulationDefault::WorldScale), TotalArea / FMath::Square(ClothingSimulationDefault::WorldScale), TotalMass);
+	UE_LOG(LogChaosCloth, Verbose, TEXT("SI Density: %f, SI Total surface: %f, SI Total mass: %f, "), (TotalArea > 0.f ? TotalMass / TotalArea : 1.f) * FMath::Square(ChaosClothingSimulationDefault::WorldScale), TotalArea / FMath::Square(ChaosClothingSimulationDefault::WorldScale), TotalMass);
 
 	// Clamp and enslave
 	const FPointWeightMap& MaxDistances = PhysMesh.GetWeightMap(EChaosWeightMapTarget::MaxDistance);
@@ -1258,7 +1278,7 @@ void ClothingSimulation::Simulate(IClothingSimulationContext* InContext)
 		Context->WorldGravity));
 
 	// Set wind velocity, used by the velocity field lambda
-	WindVelocity = Context->WindVelocity * ClothingSimulationDefault::WorldScale;  // Wind speed is set in m/s and need to be converted to cm/s
+	WindVelocity = Context->WindVelocity * ChaosClothingSimulationDefault::WorldScale;  // Wind speed is set in m/s and need to be converted to cm/s
 
 	// Check teleport modes
 	const bool bTeleport = (Context->TeleportMode > EClothingTeleportMode::None);
@@ -1349,6 +1369,21 @@ void ClothingSimulation::Simulate(IClothingSimulationContext* InContext)
 	Evolution->AdvanceOneTimeStep(DeltaTime);
 	Time += DeltaTime;
 	UE_LOG(LogChaosCloth, VeryVerbose, TEXT("DeltaTime: %.6f, FilteredDeltaTime: %.6f, Time = %.6f,  MaxPhysicsDelta = %.6f"), Context->DeltaSeconds, DeltaTime, Time, FClothingSimulationCommon::MaxPhysicsDelta);
+
+	// Debug draw
+#if CHAOS_DEBUG_DRAW
+	if (ChaosClothingSimulationConsoleVariables::CVarDebugDrawPhysMeshWired   .GetValueOnAnyThread()) { DebugDrawPhysMeshWired       (); }
+	if (ChaosClothingSimulationConsoleVariables::CVarDebugPointNormals        .GetValueOnAnyThread()) { DebugDrawPointNormals        (); }
+	if (ChaosClothingSimulationConsoleVariables::CVarDebugInversedPointNormals.GetValueOnAnyThread()) { DebugDrawInversedPointNormals(); }
+	if (ChaosClothingSimulationConsoleVariables::CVarDebugFaceNormals         .GetValueOnAnyThread()) { DebugDrawFaceNormals         (); }
+	if (ChaosClothingSimulationConsoleVariables::CVarDebugInversedFaceNormals .GetValueOnAnyThread()) { DebugDrawInversedFaceNormals (); }
+	if (ChaosClothingSimulationConsoleVariables::CVarDebugCollision           .GetValueOnAnyThread()) { DebugDrawCollision           (); }
+	if (ChaosClothingSimulationConsoleVariables::CVarDebugBackstops           .GetValueOnAnyThread()) { DebugDrawBackstops           (); }
+	if (ChaosClothingSimulationConsoleVariables::CVarDebugMaxDistances        .GetValueOnAnyThread()) { DebugDrawMaxDistances        (); }
+	if (ChaosClothingSimulationConsoleVariables::CVarDebugAnimDrive           .GetValueOnAnyThread()) { DebugDrawAnimDrive           (); }
+	if (ChaosClothingSimulationConsoleVariables::CVarDebugLongRangeConstraint .GetValueOnAnyThread()) { DebugDrawLongRangeConstraint (); }
+	if (ChaosClothingSimulationConsoleVariables::CVarDebugWindDragForces      .GetValueOnAnyThread()) { DebugDrawWindDragForces      (); }
+#endif  // #if CHAOS_DEBUG_DRAW
 }
 
 void ClothingSimulation::GetSimulationData(
@@ -1550,17 +1585,21 @@ void ClothingSimulation::AddReferencedObjects(FReferenceCollector& Collector)
 	Collector.AddReferencedObject(DebugClothMaterial);
 }
 
-void ClothingSimulation::DebugDrawPhysMeshWired(USkeletalMeshComponent* OwnerComponent, FPrimitiveDrawInterface* PDI) const
+void ClothingSimulation::DebugDrawPhysMeshShaded(USkeletalMeshComponent* OwnerComponent, FPrimitiveDrawInterface* PDI) const
 {
+	if (!DebugClothMaterial) { return; }
+
+	FDynamicMeshBuilder MeshBuilder(PDI->View->GetFeatureLevel());
 	const TPBDParticles<float, 3>& Particles = Evolution->Particles();
 
+	int32 VertexIndex = 0;
 	for (int32 MeshIndex = 0; MeshIndex < Meshes.Num(); ++MeshIndex)
 	{
 		if (const TUniquePtr<Chaos::TTriangleMesh<float>>& Mesh = Meshes[MeshIndex])
 		{
 			const TArray<TVector<int32, 3>>& Elements = Mesh->GetElements();
 
-			for (int32 ElementIndex = 0; ElementIndex < Elements.Num(); ++ElementIndex)
+			for (int32 ElementIndex = 0; ElementIndex < Elements.Num(); ++ElementIndex, VertexIndex += 3)
 			{
 				const auto& Element = Elements[ElementIndex];
 
@@ -1568,50 +1607,226 @@ void ClothingSimulation::DebugDrawPhysMeshWired(USkeletalMeshComponent* OwnerCom
 				const FVector& Pos1 = Particles.X(Element.Y);
 				const FVector& Pos2 = Particles.X(Element.Z);
 
-				PDI->DrawLine(Pos0 + LocalSimSpaceOffset, Pos1 + LocalSimSpaceOffset, FLinearColor::White, SDPG_World, 0.0f, 0.001f);
-				PDI->DrawLine(Pos1 + LocalSimSpaceOffset, Pos2 + LocalSimSpaceOffset, FLinearColor::White, SDPG_World, 0.0f, 0.001f);
-				PDI->DrawLine(Pos2 + LocalSimSpaceOffset, Pos0 + LocalSimSpaceOffset, FLinearColor::White, SDPG_World, 0.0f, 0.001f);
+				const FVector& Normal = FVector::CrossProduct(Pos2 - Pos0, Pos1 - Pos0).GetSafeNormal();
+				const FVector Tangent = ((Pos1 + Pos2) * 0.5f - Pos0).GetSafeNormal();
+
+				MeshBuilder.AddVertex(FDynamicMeshVertex(Pos0, Tangent, Normal, FVector2D(0.f, 0.f), FColor::White));
+				MeshBuilder.AddVertex(FDynamicMeshVertex(Pos1, Tangent, Normal, FVector2D(0.f, 1.f), FColor::White));
+				MeshBuilder.AddVertex(FDynamicMeshVertex(Pos2, Tangent, Normal, FVector2D(1.f, 1.f), FColor::White));
+				MeshBuilder.AddTriangle(VertexIndex, VertexIndex + 1, VertexIndex + 2);
 			}
 		}
 	}
+
+	FMatrix LocalSimSpaceToWorld(FMatrix::Identity);
+	LocalSimSpaceToWorld.SetOrigin(LocalSimSpaceOffset);
+	MeshBuilder.Draw(PDI, LocalSimSpaceToWorld, DebugClothMaterial->GetRenderProxy(), SDPG_World, false, false);
 }
+#endif  // #if WITH_EDITOR
 
-void ClothingSimulation::DebugDrawPhysMeshShaded(USkeletalMeshComponent* OwnerComponent, FPrimitiveDrawInterface* PDI) const
+#if WITH_EDITOR || CHAOS_DEBUG_DRAW
+static void DrawPoint(FPrimitiveDrawInterface* PDI, const FVector& Pos, const FLinearColor& Color, UMaterial* DebugClothMaterialVertex)  // USe color or material
 {
-	if (!DebugClothMaterial) { return; }
+#if CHAOS_DEBUG_DRAW
+	if (!PDI)
+	{
+		FDebugDrawQueue::GetInstance().DrawDebugPoint(Pos, Color.ToFColor(true), false, KINDA_SMALL_NUMBER, SDPG_Foreground, 1.f);
+		return;
+	}
+#endif
+#if WITH_EDITOR
+	const FMatrix& ViewMatrix = PDI->View->ViewMatrices.GetViewMatrix();
+	const FVector XAxis = ViewMatrix.GetColumn(0); // Just using transpose here (orthogonal transform assumed)
+	const FVector YAxis = ViewMatrix.GetColumn(1);
+	DrawDisc(PDI, Pos, XAxis, YAxis, FColor::White, 0.2f, 10, DebugClothMaterialVertex->GetRenderProxy(), SDPG_World);
+#endif
+}
 
-	FDynamicMeshBuilder MeshBuilder(PDI->View->GetFeatureLevel());
+static void DrawLine(FPrimitiveDrawInterface* PDI, const FVector& Pos0, const FVector& Pos1, const FLinearColor& Color)
+{
+#if CHAOS_DEBUG_DRAW
+	if (!PDI)
+	{
+		FDebugDrawQueue::GetInstance().DrawDebugLine(Pos0, Pos1, Color.ToFColor(true), false, KINDA_SMALL_NUMBER, SDPG_Foreground, 0.f);
+		return;
+	}
+#endif
+#if WITH_EDITOR
+	PDI->DrawLine(Pos0, Pos1, Color, SDPG_World, 0.0f, 0.001f);
+#endif
+}
+
+static void DrawArc(FPrimitiveDrawInterface* PDI, const FVector& Base, const FVector& X, const FVector& Y, float MinAngle, float MaxAngle, float Radius, const FLinearColor& Color)
+{
+	static const int32 Sections = 10;
+	const float AngleStep = FMath::DegreesToRadians((MaxAngle - MinAngle) / (float)Sections);
+	float CurrentAngle = FMath::DegreesToRadians(MinAngle);
+	FVector LastVertex = Base + Radius * (FMath::Cos(CurrentAngle) * X + FMath::Sin(CurrentAngle) * Y);
+
+	for(int32 i = 0; i < Sections; i++)
+	{
+		CurrentAngle += AngleStep;
+		const FVector ThisVertex = Base + Radius * (FMath::Cos(CurrentAngle) * X + FMath::Sin(CurrentAngle) * Y);
+		DrawLine(PDI, LastVertex, ThisVertex, Color);
+		LastVertex = ThisVertex;
+	}
+}
+
+static void DrawSphere(FPrimitiveDrawInterface* PDI, const TSphere<float, 3>& Sphere, const FQuat& Rotation, const FVector& Position, const FLinearColor& Color)
+{
+	const float Radius = Sphere.GetRadius();
+	const TVector<float, 3> Center = Position + Rotation.RotateVector(Sphere.GetCenter());
+#if CHAOS_DEBUG_DRAW
+	if (!PDI)
+	{
+		FDebugDrawQueue::GetInstance().DrawDebugSphere(Center, Radius, 12, Color.ToFColor(true), false, KINDA_SMALL_NUMBER, SDPG_Foreground, 0.f);
+		return;
+	}
+#endif
+#if WITH_EDITOR
+	const FTransform Transform(Rotation, Center);
+	DrawWireSphere(PDI, Transform, Color, Radius, 12, SDPG_World, 0.0f, 0.001f, false);
+#endif
+}
+
+static void DrawBox(FPrimitiveDrawInterface* PDI, const TBox<float, 3>& Box, const FQuat& Rotation, const FVector& Position, const FLinearColor& Color)
+{
+#if CHAOS_DEBUG_DRAW
+	if (!PDI)
+	{
+		FDebugDrawQueue::GetInstance().DrawDebugBox(Position, Box.Extents(), Rotation, Color.ToFColor(true), false, KINDA_SMALL_NUMBER, SDPG_Foreground, 0.f);
+		return;
+	}
+#endif
+#if WITH_EDITOR
+	const FMatrix BoxToWorld = FTransform(Rotation, Position).ToMatrixNoScale();
+	DrawWireBox(PDI, BoxToWorld, FBox(Box.Min(), Box.Max()), Color, SDPG_World, 0.0f, 0.001f, false);
+#endif
+}
+
+static void DrawCapsule(FPrimitiveDrawInterface* PDI, const TCapsule<float>& Capsule, const FQuat& Rotation, const FVector& Position, const FLinearColor& Color)
+{
+	const float Radius = Capsule.GetRadius();
+	const float HalfHeight = Capsule.GetHeight() * 0.5f + Radius;
+#if CHAOS_DEBUG_DRAW
+	if (!PDI)
+	{
+		FDebugDrawQueue::GetInstance().DrawDebugCapsule(Position, HalfHeight, Radius, Rotation, Color.ToFColor(true), false, KINDA_SMALL_NUMBER, SDPG_Foreground, 0.f);
+		return;
+	}
+#endif
+#if WITH_EDITOR
+	const FVector X = Rotation.RotateVector(FVector::ForwardVector);
+	const FVector Y = Rotation.RotateVector(FVector::RightVector);
+	const FVector Z = Rotation.RotateVector(FVector::UpVector);
+	DrawWireCapsule(PDI, Position, X, Y, Z, Color, Radius, HalfHeight, 12, SDPG_World, 0.0f, 0.001f, false);
+#endif
+}
+
+static void DrawTaperedCylinder(FPrimitiveDrawInterface* PDI, const TTaperedCylinder<float>& TaperedCylinder, const FQuat& Rotation, const FVector& Position, const FLinearColor& Color)
+{
+	const float HalfHeight = TaperedCylinder.GetHeight() * 0.5f;
+	const float Radius1 = TaperedCylinder.GetRadius1();
+	const float Radius2 = TaperedCylinder.GetRadius2();
+	const FVector Position1 = Position + Rotation.RotateVector(TaperedCylinder.GetX1());
+	const FVector Position2 = Position + Rotation.RotateVector(TaperedCylinder.GetX2());
+	const FQuat Q = (Position2 - Position1).ToOrientationQuat();
+	const FVector I = Q.GetRightVector();
+	const FVector J = Q.GetUpVector();
+
+	static const int32 NumSides = 12;
+	static const float	AngleDelta = 2.0f * PI / NumSides;
+	FVector	LastVertex1 = Position1 + I * Radius1;
+	FVector	LastVertex2 = Position2 + I * Radius2;
+
+	for (int32 SideIndex = 1; SideIndex <= NumSides; ++SideIndex)
+	{
+		const float Angle = AngleDelta * float(SideIndex);
+		const FVector ArcPos = I * FMath::Cos(Angle) + J * FMath::Sin(Angle);
+		const FVector Vertex1 = Position1 + ArcPos * Radius1;
+		const FVector Vertex2 = Position2 + ArcPos * Radius2;
+
+		DrawLine(PDI, LastVertex1, Vertex1, Color);
+		DrawLine(PDI, LastVertex2, Vertex2, Color);
+		DrawLine(PDI, LastVertex1, LastVertex2, Color);
+
+		LastVertex1 = Vertex1;
+		LastVertex2 = Vertex2;
+	}
+}
+
+static void DrawConvex(FPrimitiveDrawInterface* PDI, const FConvex& Convex, const FQuat& Rotation, const FVector& Position, const FLinearColor& Color)
+{
+	const TArray<TPlaneConcrete<float, 3>>& Planes = Convex.GetFaces();
+	for (int32 PlaneIndex1 = 0; PlaneIndex1 < Planes.Num(); ++PlaneIndex1)
+	{
+		const TPlaneConcrete<float, 3>& Plane1 = Planes[PlaneIndex1];
+
+		for (int32 PlaneIndex2 = PlaneIndex1 + 1; PlaneIndex2 < Planes.Num(); ++PlaneIndex2)
+		{
+			const TPlaneConcrete<float, 3>& Plane2 = Planes[PlaneIndex2];
+
+			// Find the two surface points that belong to both Plane1 and Plane2
+			uint32 ParticleIndex1 = INDEX_NONE;
+
+			const TParticles<float, 3>& SurfaceParticles = Convex.GetSurfaceParticles();
+			for (uint32 ParticleIndex = 0; ParticleIndex < SurfaceParticles.Size(); ++ParticleIndex)
+			{
+				const TVector<float, 3>& X = SurfaceParticles.X(ParticleIndex);
+
+				if (FMath::Square(Plane1.SignedDistance(X)) < KINDA_SMALL_NUMBER && 
+					FMath::Square(Plane2.SignedDistance(X)) < KINDA_SMALL_NUMBER)
+				{
+					if (ParticleIndex1 != INDEX_NONE)
+					{
+						const TVector<float, 3>& X1 = SurfaceParticles.X(ParticleIndex1);
+						const FVector Position1 = Position + Rotation.RotateVector(X1);
+						const FVector Position2 = Position + Rotation.RotateVector(X);
+						DrawLine(PDI, Position1, Position2, Color);
+						break;
+					}
+					ParticleIndex1 = ParticleIndex;
+				}
+			}
+		}
+	}
+}
+
+static void DrawCoordinateSystem(FPrimitiveDrawInterface* PDI, const FQuat& Rotation, const FVector& Position)
+{
+	const FVector X = Rotation.RotateVector(FVector::ForwardVector);
+	const FVector Y = Rotation.RotateVector(FVector::RightVector);
+	const FVector Z = Rotation.RotateVector(FVector::UpVector);
+
+	DrawLine(PDI, Position, Position + X, FLinearColor::Red);
+	DrawLine(PDI, Position, Position + Y, FLinearColor::Green);
+	DrawLine(PDI, Position, Position + Z, FLinearColor::Blue);
+}
+
+void ClothingSimulation::DebugDrawPhysMeshWired(USkeletalMeshComponent* /*OwnerComponent*/, FPrimitiveDrawInterface* PDI) const
+{
 	const TPBDParticles<float, 3>& Particles = Evolution->Particles();
 
-	int32 VertexIndex = 0;
 	for (int32 MeshIndex = 0; MeshIndex < Meshes.Num(); ++MeshIndex)
 	{
-		if (const TUniquePtr<Chaos::TTriangleMesh<float>>& Mesh = Meshes[MeshIndex])
+		if (const TUniquePtr<TTriangleMesh<float>>& Mesh = Meshes[MeshIndex])
 		{
 			const TArray<TVector<int32, 3>>& Elements = Mesh->GetElements();
 
-			for (int32 ElementIndex = 0; ElementIndex < Elements.Num(); ++ElementIndex, VertexIndex += 3)
+			for (int32 ElementIndex = 0; ElementIndex < Elements.Num(); ++ElementIndex)
 			{
 				const auto& Element = Elements[ElementIndex];
 
-				const FVector& Pos0 = Particles.X(Element.X);
-				const FVector& Pos1 = Particles.X(Element.Y);
-				const FVector& Pos2 = Particles.X(Element.Z);
-
-				const FVector& Normal = FVector::CrossProduct(Pos2 - Pos0, Pos1 - Pos0).GetSafeNormal();
-				const FVector Tangent = ((Pos1 + Pos2) * 0.5f - Pos0).GetSafeNormal();
+				const FVector Pos0 = Particles.X(Element.X) + LocalSimSpaceOffset;
+				const FVector Pos1 = Particles.X(Element.Y) + LocalSimSpaceOffset;
+				const FVector Pos2 = Particles.X(Element.Z) + LocalSimSpaceOffset;
 
-				MeshBuilder.AddVertex(FDynamicMeshVertex(Pos0, Tangent, Normal, FVector2D(0.f, 0.f), FColor::White));
-				MeshBuilder.AddVertex(FDynamicMeshVertex(Pos1, Tangent, Normal, FVector2D(0.f, 1.f), FColor::White));
-				MeshBuilder.AddVertex(FDynamicMeshVertex(Pos2, Tangent, Normal, FVector2D(1.f, 1.f), FColor::White));
-				MeshBuilder.AddTriangle(VertexIndex, VertexIndex + 1, VertexIndex + 2);
+				DrawLine(PDI, Pos0, Pos1, FLinearColor::White);
+				DrawLine(PDI, Pos1, Pos2, FLinearColor::White);
+				DrawLine(PDI, Pos2, Pos0, FLinearColor::White);
 			}
 		}
 	}
-
-	FMatrix LocalSimSpaceToWorld(FMatrix::Identity);
-	LocalSimSpaceToWorld.SetOrigin(LocalSimSpaceOffset);
-	MeshBuilder.Draw(PDI, LocalSimSpaceToWorld, DebugClothMaterial->GetRenderProxy(), SDPG_World, false, false);
 }
 
 void ClothingSimulation::DebugDrawPointNormals(USkeletalMeshComponent* OwnerComponent, FPrimitiveDrawInterface* PDI) const
@@ -1624,15 +1839,15 @@ void ClothingSimulation::DebugDrawPointNormals(USkeletalMeshComponent* OwnerComp
 	{
 		if (Meshes[MeshIndex])
 		{
-			const Chaos::TVector<uint32, 2> Range = IndexToRangeMap[MeshIndex];
-			const TArray<Chaos::TVector<float, 3>>& MeshPointNormals = PointNormals[MeshIndex];
+			const TVector<uint32, 2> Range = IndexToRangeMap[MeshIndex];
+			const TArray<TVector<float, 3>>& MeshPointNormals = PointNormals[MeshIndex];
 
 			for (uint32 ParticleIndex = Range[0]; ParticleIndex < Range[1]; ++ParticleIndex)
 			{
-				const TVector<float, 3>& Pos = Particles.X(ParticleIndex);
-				const TVector<float, 3>& Normal = MeshPointNormals[ParticleIndex - Range[0]];
+				const FVector Pos0 = Particles.X(ParticleIndex) + LocalSimSpaceOffset;
+				const FVector Pos1 = Pos0 + MeshPointNormals[ParticleIndex - Range[0]] * 20.0f;
 
-				PDI->DrawLine(Pos + LocalSimSpaceOffset, Pos + Normal * 20.0f + LocalSimSpaceOffset, FLinearColor::White, SDPG_World, 0.0f, 0.001f);
+				DrawLine(PDI, Pos0, Pos1, FLinearColor::White);
 			}
 		}
 	}
@@ -1648,15 +1863,15 @@ void ClothingSimulation::DebugDrawInversedPointNormals(USkeletalMeshComponent* O
 	{
 		if (Meshes[MeshIndex])
 		{
-			const Chaos::TVector<uint32, 2> Range = IndexToRangeMap[MeshIndex];
-			const TArray<Chaos::TVector<float, 3>>& MeshPointNormals = PointNormals[MeshIndex];
+			const TVector<uint32, 2> Range = IndexToRangeMap[MeshIndex];
+			const TArray<TVector<float, 3>>& MeshPointNormals = PointNormals[MeshIndex];
 
 			for (uint32 ParticleIndex = Range[0]; ParticleIndex < Range[1]; ++ParticleIndex)
 			{
-				const TVector<float, 3>& Pos = Particles.X(ParticleIndex);
-				const TVector<float, 3>& Normal = MeshPointNormals[ParticleIndex - Range[0]];
+				const FVector Pos0 = Particles.X(ParticleIndex) + LocalSimSpaceOffset;
+				const FVector Pos1 = Pos0 - MeshPointNormals[ParticleIndex - Range[0]] * 20.0f;
 
-				PDI->DrawLine(Pos + LocalSimSpaceOffset, Pos - Normal * 20.0f + LocalSimSpaceOffset, FLinearColor::White, SDPG_World, 0.0f, 0.001f);
+				DrawLine(PDI, Pos0, Pos1, FLinearColor::White);
 			}
 		}
 	}
@@ -1670,22 +1885,22 @@ void ClothingSimulation::DebugDrawFaceNormals(USkeletalMeshComponent* OwnerCompo
 
 	for (int32 MeshIndex = 0; MeshIndex < Meshes.Num(); ++MeshIndex)
 	{
-		if (const TUniquePtr<Chaos::TTriangleMesh<float>>& Mesh = Meshes[MeshIndex])
+		if (const TUniquePtr<TTriangleMesh<float>>& Mesh = Meshes[MeshIndex])
 		{
-			const TArray<Chaos::TVector<float, 3>>& MeshFaceNormals = FaceNormals[MeshIndex];
+			const TArray<TVector<float, 3>>& MeshFaceNormals = FaceNormals[MeshIndex];
 
 			const TArray<TVector<int32, 3>>& Elements = Mesh->GetElements();
 			for (int32 ElementIndex = 0; ElementIndex < Elements.Num(); ++ElementIndex)
 			{
 				const TVector<int32, 3>& Element = Elements[ElementIndex];
 
-				const TVector<float, 3> Pos = (
+				const FVector Pos0 = LocalSimSpaceOffset + (
 					Particles.X(Element.X) +
 					Particles.X(Element.Y) +
 					Particles.X(Element.Z)) / 3.f;
-				const TVector<float, 3>& Normal = MeshFaceNormals[ElementIndex];
+				const FVector Pos1 = Pos0 + MeshFaceNormals[ElementIndex] * 20.0f;
 
-				PDI->DrawLine(Pos + LocalSimSpaceOffset, Pos + Normal * 20.0f + LocalSimSpaceOffset, FLinearColor::Yellow, SDPG_World, 0.0f, 0.001f);
+				DrawLine(PDI, Pos0, Pos1, FLinearColor::Yellow);
 			}
 		}
 	}
@@ -1699,22 +1914,22 @@ void ClothingSimulation::DebugDrawInversedFaceNormals(USkeletalMeshComponent* Ow
 
 	for (int32 MeshIndex = 0; MeshIndex < Meshes.Num(); ++MeshIndex)
 	{
-		if (const TUniquePtr<Chaos::TTriangleMesh<float>>& Mesh = Meshes[MeshIndex])
+		if (const TUniquePtr<TTriangleMesh<float>>& Mesh = Meshes[MeshIndex])
 		{
-			const TArray<Chaos::TVector<float, 3>>& MeshFaceNormals = FaceNormals[MeshIndex];
+			const TArray<TVector<float, 3>>& MeshFaceNormals = FaceNormals[MeshIndex];
 
 			const TArray<TVector<int32, 3>>& Elements = Mesh->GetElements();
 			for (int32 ElementIndex = 0; ElementIndex < Elements.Num(); ++ElementIndex)
 			{
 				const TVector<int32, 3>& Element = Elements[ElementIndex];
 
-				const TVector<float, 3> Pos = (
+				const TVector<float, 3> Pos0 = LocalSimSpaceOffset + (
 					Particles.X(Element.X) +
 					Particles.X(Element.Y) +
 					Particles.X(Element.Z)) / 3.f;
-				const TVector<float, 3>& Normal = MeshFaceNormals[ElementIndex];
+				const FVector Pos1 = Pos0 - MeshFaceNormals[ElementIndex] * 20.0f;
 
-				PDI->DrawLine(Pos + LocalSimSpaceOffset, Pos - Normal * 20.0f + LocalSimSpaceOffset, FLinearColor::Yellow, SDPG_World, 0.0f, 0.001f);
+				DrawLine(PDI, Pos0, Pos1, FLinearColor::Yellow);
 			}
 		}
 	}
@@ -1722,100 +1937,6 @@ void ClothingSimulation::DebugDrawInversedFaceNormals(USkeletalMeshComponent* Ow
 
 void ClothingSimulation::DebugDrawCollision(USkeletalMeshComponent* OwnerComponent, FPrimitiveDrawInterface* PDI) const
 {
-	auto DrawSphere = [&PDI, this](const Chaos::TSphere<float, 3>& Sphere, const TRotation<float, 3>& Rotation, const Chaos::TVector<float, 3>& Position, const FLinearColor& Color)
-	{
-		const float Radius = Sphere.GetRadius();
-		const Chaos::TVector<float, 3> Center = Sphere.GetCenter();
-		const FTransform Transform(Rotation, Position + Rotation.RotateVector(Center) + LocalSimSpaceOffset);
-		DrawWireSphere(PDI, Transform, Color, Radius, 12, SDPG_World, 0.0f, 0.001f, false);
-	};
-
-	auto DrawBox = [&PDI, this](const Chaos::TBox<float, 3>& Box, const TRotation<float, 3>& Rotation, const Chaos::TVector<float, 3>& Position, const FLinearColor& Color)
-	{
-		const FMatrix BoxToWorld = FTransform(Rotation, Position + LocalSimSpaceOffset).ToMatrixNoScale();
-		const FVector Radii = Box.Extents() * 0.5f;
-		DrawWireBox(PDI, BoxToWorld, FBox(Box.Min(), Box.Max()), Color, SDPG_World, 0.0f, 0.001f, false);
-	};
-
-	auto DrawCapsule = [&PDI, this](const Chaos::TCapsule<float>& Capsule, const TRotation<float, 3>& Rotation, const Chaos::TVector<float, 3>& Position, const FLinearColor& Color)
-	{
-		const float HalfHeight = Capsule.GetHeight() * 0.5f;
-		const float Radius = Capsule.GetRadius();
-		const FVector X = Rotation.RotateVector(FVector::ForwardVector);
-		const FVector Y = Rotation.RotateVector(FVector::RightVector);
-		const FVector Z = Rotation.RotateVector(FVector::UpVector);
-		DrawWireCapsule(PDI, Position + LocalSimSpaceOffset, X, Y, Z, Color, Radius, HalfHeight + Radius, 12, SDPG_World, 0.0f, 0.001f, false);
-	};
-
-	auto DrawTaperedCylinder = [&PDI, this](const Chaos::TTaperedCylinder<float>& TaperedCylinder, const TRotation<float, 3>& Rotation, const Chaos::TVector<float, 3>& Position, const FLinearColor& Color)
-	{
-		const float HalfHeight = TaperedCylinder.GetHeight() * 0.5f;
-		const float Radius1 = TaperedCylinder.GetRadius1();
-		const float Radius2 = TaperedCylinder.GetRadius2();
-		const FVector Position1 = Position + Rotation.RotateVector(TaperedCylinder.GetX1());
-		const FVector Position2 = Position + Rotation.RotateVector(TaperedCylinder.GetX2());
-		const FQuat Q = (Position2 - Position1).ToOrientationQuat();
-		const FVector I = Q.GetRightVector();
-		const FVector J = Q.GetUpVector();
-
-		static const int32 NumSides = 12;
-		static const float	AngleDelta = 2.0f * PI / NumSides;
-		FVector	LastVertex1 = Position1 + I * Radius1;
-		FVector	LastVertex2 = Position2 + I * Radius2;
-
-		for (int32 SideIndex = 1; SideIndex <= NumSides; ++SideIndex)
-		{
-			const float Angle = AngleDelta * float(SideIndex);
-			const FVector ArcPos = I * FMath::Cos(Angle) + J * FMath::Sin(Angle);
-			const FVector Vertex1 = Position1 + ArcPos * Radius1;
-			const FVector Vertex2 = Position2 + ArcPos * Radius2;
-
-			PDI->DrawLine(LastVertex1 + LocalSimSpaceOffset, Vertex1 + LocalSimSpaceOffset, Color, SDPG_World, 0.0f, 0.001f, false);
-			PDI->DrawLine(LastVertex2 + LocalSimSpaceOffset, Vertex2 + LocalSimSpaceOffset, Color, SDPG_World, 0.0f, 0.001f, false);
-			PDI->DrawLine(LastVertex1 + LocalSimSpaceOffset, LastVertex2 + LocalSimSpaceOffset, Color, SDPG_World, 0.0f, 0.001f, false);
-
-			LastVertex1 = Vertex1;
-			LastVertex2 = Vertex2;
-		}
-	};
-
-	auto DrawConvex = [&PDI, this](const Chaos::FConvex& Convex, const TRotation<float, 3>& Rotation, const Chaos::TVector<float, 3>& Position, const FLinearColor& Color)
-	{
-		const TArray<Chaos::TPlaneConcrete<float, 3>>& Planes = Convex.GetFaces();
-		for (int32 PlaneIndex1 = 0; PlaneIndex1 < Planes.Num(); ++PlaneIndex1)
-		{
-			const Chaos::TPlaneConcrete<float, 3>& Plane1 = Planes[PlaneIndex1];
-
-			for (int32 PlaneIndex2 = PlaneIndex1 + 1; PlaneIndex2 < Planes.Num(); ++PlaneIndex2)
-			{
-				const Chaos::TPlaneConcrete<float, 3>& Plane2 = Planes[PlaneIndex2];
-
-				// Find the two surface points that belong to both Plane1 and Plane2
-				uint32 ParticleIndex1 = INDEX_NONE;
-
-				const Chaos::TParticles<float, 3>& SurfaceParticles = Convex.GetSurfaceParticles();
-				for (uint32 ParticleIndex = 0; ParticleIndex < SurfaceParticles.Size(); ++ParticleIndex)
-				{
-					const Chaos::TVector<float, 3>& X = SurfaceParticles.X(ParticleIndex);
-
-					if (FMath::Square(Plane1.SignedDistance(X)) < KINDA_SMALL_NUMBER && 
-						FMath::Square(Plane2.SignedDistance(X)) < KINDA_SMALL_NUMBER)
-					{
-						if (ParticleIndex1 != INDEX_NONE)
-						{
-							const Chaos::TVector<float, 3>& X1 = SurfaceParticles.X(ParticleIndex1);
-							const FVector Position1 = Position + Rotation.RotateVector(X1);
-							const FVector Position2 = Position + Rotation.RotateVector(X);
-							PDI->DrawLine(Position1 + LocalSimSpaceOffset, Position2 + LocalSimSpaceOffset, Color, SDPG_World, 0.0f, 0.001f, false);
-							break;
-						}
-						ParticleIndex1 = ParticleIndex;
-					}
-				}
-			}
-		}
-	};
-
 	static const FLinearColor MappedColor(FColor::Cyan);
 	static const FLinearColor UnmappedColor(FColor::Red);
 
@@ -1823,41 +1944,41 @@ void ClothingSimulation::DebugDrawCollision(USkeletalMeshComponent* OwnerCompone
 
 	for (uint32 Index = 0; Index < CollisionParticles.Size(); ++Index)
 	{
-		if (const Chaos::FImplicitObject* const Object = CollisionParticles.DynamicGeometry(Index).Get())
+		if (const FImplicitObject* const Object = CollisionParticles.DynamicGeometry(Index).Get())
 		{
 			const uint32 BoneIndex = BoneIndices[Index];
 			const FLinearColor Color = (BoneIndex != INDEX_NONE) ? MappedColor : UnmappedColor;
 
-			const Chaos::TVector<float, 3>& Position = CollisionParticles.X(Index);
+			const TVector<float, 3>& Position = CollisionParticles.X(Index) + LocalSimSpaceOffset;
 			const TRotation<float, 3>& Rotation = CollisionParticles.R(Index);
 
 			switch (Object->GetType())
 			{
-			case Chaos::ImplicitObjectType::Sphere:
-				DrawSphere(Object->GetObjectChecked<Chaos::TSphere<float, 3>>(), Rotation, Position, Color);
+			case ImplicitObjectType::Sphere:
+				DrawSphere(PDI, Object->GetObjectChecked<TSphere<float, 3>>(), Rotation, Position, Color);
 				break;
 
-			case Chaos::ImplicitObjectType::Box:
-				DrawBox(Object->GetObjectChecked<Chaos::TBox<float, 3>>(), Rotation, Position, Color);
+			case ImplicitObjectType::Box:
+				DrawBox(PDI, Object->GetObjectChecked<TBox<float, 3>>(), Rotation, Position, Color);
 				break;
 
-			case Chaos::ImplicitObjectType::Capsule:
-				DrawCapsule(Object->GetObjectChecked<Chaos::TCapsule<float>>(), Rotation, Position, Color);
+			case ImplicitObjectType::Capsule:
+				DrawCapsule(PDI, Object->GetObjectChecked<TCapsule<float>>(), Rotation, Position, Color);
 				break;
 
-			case Chaos::ImplicitObjectType::Union:  // Union only used as collision tappered capsules
-				for (const TUniquePtr<Chaos::FImplicitObject>& SubObjectPtr : Object->GetObjectChecked<Chaos::FImplicitObjectUnion>().GetObjects())
+			case ImplicitObjectType::Union:  // Union only used as collision tappered capsules
+				for (const TUniquePtr<FImplicitObject>& SubObjectPtr : Object->GetObjectChecked<FImplicitObjectUnion>().GetObjects())
 				{
-					if (const Chaos::FImplicitObject* const SubObject = SubObjectPtr.Get())
+					if (const FImplicitObject* const SubObject = SubObjectPtr.Get())
 					{
 						switch (SubObject->GetType())
 						{
-						case Chaos::ImplicitObjectType::Sphere:
-							DrawSphere(SubObject->GetObjectChecked<Chaos::TSphere<float, 3>>(), Rotation, Position, Color);
+						case ImplicitObjectType::Sphere:
+							DrawSphere(PDI, SubObject->GetObjectChecked<TSphere<float, 3>>(), Rotation, Position, Color);
 							break;
 
-						case Chaos::ImplicitObjectType::TaperedCylinder:
-							DrawTaperedCylinder(SubObject->GetObjectChecked<Chaos::TTaperedCylinder<float>>(), Rotation, Position, Color);
+						case ImplicitObjectType::TaperedCylinder:
+							DrawTaperedCylinder(PDI, SubObject->GetObjectChecked<TTaperedCylinder<float>>(), Rotation, Position, Color);
 							break;
 
 						default:
@@ -1867,19 +1988,19 @@ void ClothingSimulation::DebugDrawCollision(USkeletalMeshComponent* OwnerCompone
 				}
 				break;
 	
-			case Chaos::ImplicitObjectType::Convex:
-				DrawConvex(Object->GetObjectChecked<Chaos::FConvex>(), Rotation, Position, Color);
+			case ImplicitObjectType::Convex:
+				DrawConvex(PDI, Object->GetObjectChecked<FConvex>(), Rotation, Position, Color);
 				break;
 
 			default:
-				DrawCoordinateSystem(PDI, Position + LocalSimSpaceOffset, FRotator(Rotation), 10.0f, SDPG_World, 0.1f);  // Draw everything else as a coordinate for now
+				DrawCoordinateSystem(PDI, Rotation, Position);  // Draw everything else as a coordinate for now
 				break;
 			}
 		}
 	}
 }
 
-void ClothingSimulation::DebugDrawBackstops(USkeletalMeshComponent* OwnerComponent, FPrimitiveDrawInterface* PDI) const
+void ClothingSimulation::DebugDrawBackstops(USkeletalMeshComponent* /*OwnerComponent*/, FPrimitiveDrawInterface* PDI) const
 {
 	for (int32 i = 0; i < IndexToRangeMap.Num(); ++i)
 	{
@@ -1905,12 +2026,12 @@ void ClothingSimulation::DebugDrawBackstops(USkeletalMeshComponent* OwnerCompone
 			const uint32 WeightMapIndex = ParticleIndex - IndexToRangeMap[i][0];
 			const float Radius = BackstopRadiuses[WeightMapIndex];
 			const float Distance = BackstopDistances[WeightMapIndex];
-			PDI->DrawLine(AnimationPositions[ParticleIndex] + LocalSimSpaceOffset, AnimationPositions[ParticleIndex] - AnimationNormals[ParticleIndex] * (Distance - Radius) + LocalSimSpaceOffset, FColor::White, SDPG_World, 0.0f, 0.001f);
+			DrawLine(PDI, AnimationPositions[ParticleIndex] + LocalSimSpaceOffset, AnimationPositions[ParticleIndex] - AnimationNormals[ParticleIndex] * (Distance - Radius) + LocalSimSpaceOffset, FLinearColor::White);
 			if (Radius > 0.0f)
 			{
 				const FVector& Normal = AnimationNormals[ParticleIndex];
 				const FVector& Position = AnimationPositions[ParticleIndex];
-				auto DrawBackstop = [Radius, Distance, &Normal, &Position, PDI, this](const FVector& Axis, const FColor& Color)
+				auto DrawBackstop = [Radius, Distance, &Normal, &Position, PDI, this](const FVector& Axis, const FLinearColor& Color)
 				{
 					const float ArcLength = 5.0f; // Arch length in cm
 					const float ArcAngle = ArcLength * 360.0f / (Radius * 2.0f * PI);
@@ -1918,18 +2039,18 @@ void ClothingSimulation::DebugDrawBackstops(USkeletalMeshComponent* OwnerCompone
 					const float MaxCosAngle = 0.99f;
 					if (FMath::Abs(FVector::DotProduct(Normal, Axis)) < MaxCosAngle)
 					{
-						DrawArc(PDI, Position - Normal * Distance + LocalSimSpaceOffset, Normal, FVector::CrossProduct(Axis, Normal).GetSafeNormal(), -ArcAngle / 2.0f, ArcAngle / 2.0f, Radius, 10, Color, SDPG_World);
+						DrawArc(PDI, Position - Normal * Distance + LocalSimSpaceOffset, Normal, FVector::CrossProduct(Axis, Normal).GetSafeNormal(), -ArcAngle / 2.0f, ArcAngle / 2.0f, Radius, Color);
 					}
 				};
-				DrawBackstop(FVector::ForwardVector, FColor::Blue);
-				DrawBackstop(FVector::UpVector, FColor::Blue);
-				DrawBackstop(FVector::RightVector, FColor::Blue);
+				DrawBackstop(FVector::ForwardVector, FLinearColor::Blue);
+				DrawBackstop(FVector::UpVector, FLinearColor::Blue);
+				DrawBackstop(FVector::RightVector, FLinearColor::Blue);
 			}
 		}
 	}
 }
 
-void ClothingSimulation::DebugDrawMaxDistances(USkeletalMeshComponent* OwnerComponent, FPrimitiveDrawInterface* PDI) const
+void ClothingSimulation::DebugDrawMaxDistances(USkeletalMeshComponent* /*OwnerComponent*/, FPrimitiveDrawInterface* PDI) const
 {
 	const TPBDParticles<float, 3>& Particles = Evolution->Particles();
 	for (int32 i = 0; i < IndexToRangeMap.Num(); ++i)
@@ -1956,20 +2077,17 @@ void ClothingSimulation::DebugDrawMaxDistances(USkeletalMeshComponent* OwnerComp
 			const float Distance = MaxDistances[WeightMapIndex];
 			if (Particles.InvM(ParticleIndex) == 0.0f)
 			{
-				const FMatrix& ViewMatrix = PDI->View->ViewMatrices.GetViewMatrix();
-				const FVector& XAxis = ViewMatrix.GetColumn(0); // Just using transpose here (orthogonal transform assumed)
-				const FVector& YAxis = ViewMatrix.GetColumn(1);
-				DrawDisc(PDI, AnimationPositions[ParticleIndex] + LocalSimSpaceOffset, XAxis, YAxis, FColor::White, 0.2f, 10, DebugClothMaterialVertex->GetRenderProxy(), SDPG_World);
+				DrawPoint(PDI, AnimationPositions[ParticleIndex] + LocalSimSpaceOffset, FLinearColor::Red, DebugClothMaterialVertex);
 			}
 			else
 			{
-				PDI->DrawLine(AnimationPositions[ParticleIndex] + LocalSimSpaceOffset, AnimationPositions[ParticleIndex] + AnimationNormals[ParticleIndex] * Distance + LocalSimSpaceOffset, FColor::White, SDPG_World, 0.0f, 0.001f);
+				DrawLine(PDI, AnimationPositions[ParticleIndex] + LocalSimSpaceOffset, AnimationPositions[ParticleIndex] + AnimationNormals[ParticleIndex] * Distance + LocalSimSpaceOffset, FLinearColor::White);
 			}
 		}
 	}
 }
 
-void ClothingSimulation::DebugDrawAnimDrive(USkeletalMeshComponent* OwnerComponent, FPrimitiveDrawInterface* PDI) const
+void ClothingSimulation::DebugDrawAnimDrive(USkeletalMeshComponent* /*OwnerComponent*/, FPrimitiveDrawInterface* PDI) const
 {
 	const TPBDParticles<float, 3>& Particles = Evolution->Particles();
 	for (int32 i = 0; i < IndexToRangeMap.Num(); ++i)
@@ -1994,12 +2112,12 @@ void ClothingSimulation::DebugDrawAnimDrive(USkeletalMeshComponent* OwnerCompone
 		{
 			const uint32 WeightMapIndex = ParticleIndex - IndexToRangeMap[i][0];
 			const float Multiplier = AnimDriveMultipliers[WeightMapIndex];
-			PDI->DrawLine(AnimationPositions[ParticleIndex] + LocalSimSpaceOffset, Particles.X(ParticleIndex) + LocalSimSpaceOffset, Multiplier * AnimDriveSpringStiffness[i] * FColor::Cyan, SDPG_World, 0.0f, 0.001f);
+			DrawLine(PDI, AnimationPositions[ParticleIndex] + LocalSimSpaceOffset, Particles.X(ParticleIndex) + LocalSimSpaceOffset, FLinearColor(FColor::Cyan) * Multiplier * AnimDriveSpringStiffness[i]);
 		}
 	}
 }
 
-void ClothingSimulation::DebugDrawLongRangeConstraint(USkeletalMeshComponent* OwnerComponent, FPrimitiveDrawInterface* PDI) const
+void ClothingSimulation::DebugDrawLongRangeConstraint(USkeletalMeshComponent* /*OwnerComponent*/, FPrimitiveDrawInterface* PDI) const
 {
 	const TPBDParticles<float, 3>& Particles = Evolution->Particles();
 
@@ -2020,19 +2138,19 @@ void ClothingSimulation::DebugDrawLongRangeConstraint(USkeletalMeshComponent* Ow
 			const float CurDist = TPBDLongRangeConstraintsBase<float, 3>::ComputeGeodesicDistance(Particles, Path);
 			const float Offset = CurDist - RefDist;
 
-			const Chaos::TVector<float, 3> P0 = Particles.X(Path[0]);  // Kinematic particle
-			const Chaos::TVector<float, 3> P1 = Particles.X(Path[Path.Num() - 1]);  // Target particle
+			const TVector<float, 3> P0 = Particles.X(Path[0]) + LocalSimSpaceOffset;  // Kinematic particle
+			const TVector<float, 3> P1 = Particles.X(Path[Path.Num() - 1]) + LocalSimSpaceOffset;  // Target particle
 
-			const Chaos::TVector<float, 3> Direction = (Particles.X(Path[Path.Num() - 2]) - P1).GetSafeNormal();
-			const Chaos::TVector<float, 3> P2 = P1 + Direction * Offset;
+			const TVector<float, 3> Direction = (Particles.X(Path[Path.Num() - 2]) - P1).GetSafeNormal();
+			const TVector<float, 3> P2 = P1 + Direction * Offset + LocalSimSpaceOffset;
 
-			PDI->DrawLine(P0 + LocalSimSpaceOffset, P1 + LocalSimSpaceOffset, FColor::Purple, SDPG_World, 0.0f, 0.001f);
-			PDI->DrawLine(P1 + LocalSimSpaceOffset, P2 + LocalSimSpaceOffset, FColor::Black, SDPG_World, 0.0f, 0.001f);
+			DrawLine(PDI, P0, P1, FLinearColor(FColor::Purple));
+			DrawLine(PDI, P1, P2, FLinearColor::Black);
 		}
 	}
 }
 
-void ClothingSimulation::DebugDrawWindDragForces(USkeletalMeshComponent* OwnerComponent, FPrimitiveDrawInterface* PDI) const
+void ClothingSimulation::DebugDrawWindDragForces(USkeletalMeshComponent* /*OwnerComponent*/, FPrimitiveDrawInterface* PDI) const
 {
 	const TPBDParticles<float, 3>& Particles = Evolution->Particles();
 	const TArray<TVelocityField<float, 3>>& VelocityFields = Evolution->GetVelocityFields();
@@ -2045,13 +2163,13 @@ void ClothingSimulation::DebugDrawWindDragForces(USkeletalMeshComponent* OwnerCo
 		for (int32 ElementIndex = 0; ElementIndex < Elements.Num(); ++ElementIndex)
 		{
 			const TVector<int32, 3>& Element = Elements[ElementIndex];
-			const TVector<float, 3> Position = (1.f / 3.f ) * (
+			const TVector<float, 3> Position = LocalSimSpaceOffset + (1.f / 3.f ) * (
 				Particles.X(Element[0]) +
 				Particles.X(Element[1]) +
 				Particles.X(Element[2]));
 			const TVector<float, 3>& Force = Forces[ElementIndex];
-			PDI->DrawLine(Position + LocalSimSpaceOffset, Position + Force + LocalSimSpaceOffset, FColor::Green, SDPG_World, 0.0f, 0.001f);
+			DrawLine(PDI, Position, Position + Force, FColor::Green);
 		}
 	}
 }
-#endif  // #if WITH_EDITOR
+#endif  // #if WITH_EDITOR || CHAOS_DEBUG_DRAW
diff --git a/Engine/Plugins/Experimental/ChaosCloth/Source/Chaos/Public/ChaosCloth/ChaosClothingSimulation.h b/Engine/Plugins/Experimental/ChaosCloth/Source/Chaos/Public/ChaosCloth/ChaosClothingSimulation.h
--- a/Engine/Plugins/Experimental/ChaosCloth/Source/Chaos/Public/ChaosCloth/ChaosClothingSimulation.h
+++ b/Engine/Plugins/Experimental/ChaosCloth/Source/Chaos/Public/ChaosCloth/ChaosClothingSimulation.h
@@ -7,6 +7,7 @@
 #include "Chaos/PBDEvolution.h"
 #include "Chaos/Transform.h"
 #include "Chaos/TriangleMesh.h"
+#include "Chaos/ChaosDebugDrawDeclares.h"
 #include "ChaosCloth/ChaosClothConfig.h"
 #include "Components/SkeletalMeshComponent.h"
 
@@ -41,20 +42,25 @@ namespace Chaos
 		virtual void AddReferencedObjects(FReferenceCollector& Collector) override;
 		// End of FGCObject interface
 
-		CHAOSCLOTH_API void DebugDrawPhysMeshWired(USkeletalMeshComponent* OwnerComponent, FPrimitiveDrawInterface* PDI) const;
+		// Editor only debug draw function
 		CHAOSCLOTH_API void DebugDrawPhysMeshShaded(USkeletalMeshComponent* OwnerComponent, FPrimitiveDrawInterface* PDI) const;
-		CHAOSCLOTH_API void DebugDrawPointNormals(USkeletalMeshComponent* OwnerComponent, FPrimitiveDrawInterface* PDI) const;
-		CHAOSCLOTH_API void DebugDrawInversedPointNormals(USkeletalMeshComponent* OwnerComponent, FPrimitiveDrawInterface* PDI) const;
-		CHAOSCLOTH_API void DebugDrawFaceNormals(USkeletalMeshComponent* OwnerComponent, FPrimitiveDrawInterface* PDI) const;
-		CHAOSCLOTH_API void DebugDrawInversedFaceNormals(USkeletalMeshComponent* OwnerComponent, FPrimitiveDrawInterface* PDI) const;
-		CHAOSCLOTH_API void DebugDrawCollision(USkeletalMeshComponent* OwnerComponent, FPrimitiveDrawInterface* PDI) const;
-		CHAOSCLOTH_API void DebugDrawBackstops(USkeletalMeshComponent* OwnerComponent, FPrimitiveDrawInterface* PDI) const;
-		CHAOSCLOTH_API void DebugDrawMaxDistances(USkeletalMeshComponent* OwnerComponent, FPrimitiveDrawInterface* PDI) const;
-		CHAOSCLOTH_API void DebugDrawAnimDrive(USkeletalMeshComponent* OwnerComponent, FPrimitiveDrawInterface* PDI) const;
-		CHAOSCLOTH_API void DebugDrawLongRangeConstraint(USkeletalMeshComponent* OwnerComponent, FPrimitiveDrawInterface* PDI) const;
-		CHAOSCLOTH_API void DebugDrawWindDragForces(USkeletalMeshComponent* OwnerComponent, FPrimitiveDrawInterface* PDI) const;
 #endif  // #if WITH_EDITOR
 
+#if WITH_EDITOR || CHAOS_DEBUG_DRAW
+		// Editor & runtime debug draw functions
+		CHAOSCLOTH_API void DebugDrawPhysMeshWired(USkeletalMeshComponent* OwnerComponent = nullptr, FPrimitiveDrawInterface* PDI = nullptr) const;
+		CHAOSCLOTH_API void DebugDrawPointNormals(USkeletalMeshComponent* OwnerComponent = nullptr, FPrimitiveDrawInterface* PDI = nullptr) const;
+		CHAOSCLOTH_API void DebugDrawInversedPointNormals(USkeletalMeshComponent* OwnerComponent = nullptr, FPrimitiveDrawInterface* PDI = nullptr) const;
+		CHAOSCLOTH_API void DebugDrawFaceNormals(USkeletalMeshComponent* OwnerComponent = nullptr, FPrimitiveDrawInterface* PDI = nullptr) const;
+		CHAOSCLOTH_API void DebugDrawInversedFaceNormals(USkeletalMeshComponent* OwnerComponent = nullptr, FPrimitiveDrawInterface* PDI = nullptr) const;
+		CHAOSCLOTH_API void DebugDrawCollision(USkeletalMeshComponent* OwnerComponent = nullptr, FPrimitiveDrawInterface* PDI = nullptr) const;
+		CHAOSCLOTH_API void DebugDrawBackstops(USkeletalMeshComponent* OwnerComponent = nullptr, FPrimitiveDrawInterface* PDI = nullptr) const;
+		CHAOSCLOTH_API void DebugDrawMaxDistances(USkeletalMeshComponent* OwnerComponent = nullptr, FPrimitiveDrawInterface* PDI = nullptr) const;
+		CHAOSCLOTH_API void DebugDrawAnimDrive(USkeletalMeshComponent* OwnerComponent = nullptr, FPrimitiveDrawInterface* PDI = nullptr) const;
+		CHAOSCLOTH_API void DebugDrawLongRangeConstraint(USkeletalMeshComponent* OwnerComponent = nullptr, FPrimitiveDrawInterface* PDI = nullptr) const;
+		CHAOSCLOTH_API void DebugDrawWindDragForces(USkeletalMeshComponent* OwnerComponent = nullptr, FPrimitiveDrawInterface* PDI = nullptr) const;
+#endif  // #if WITH_EDITOR || CHAOS_DEBUG_DRAW
+
 	protected:
 		// IClothingSimulation interface
 		virtual void Initialize() override;
