CIS - Add missing #ifdef WITH_EDITOR.

diff --git a/Engine/Plugins/Experimental/ChaosCloth/Source/Chaos/Private/ChaosCloth/ChaosClothingSimulation.cpp b/Engine/Plugins/Experimental/ChaosCloth/Source/Chaos/Private/ChaosCloth/ChaosClothingSimulation.cpp
--- a/Engine/Plugins/Experimental/ChaosCloth/Source/Chaos/Private/ChaosCloth/ChaosClothingSimulation.cpp
+++ b/Engine/Plugins/Experimental/ChaosCloth/Source/Chaos/Private/ChaosCloth/ChaosClothingSimulation.cpp
@@ -2077,7 +2077,9 @@ void ClothingSimulation::DebugDrawMaxDistances(USkeletalMeshComponent* /*OwnerCo
 			const float Distance = MaxDistances[WeightMapIndex];
 			if (Particles.InvM(ParticleIndex) == 0.0f)
 			{
+#if WITH_EDITOR
 				DrawPoint(PDI, AnimationPositions[ParticleIndex] + LocalSimSpaceOffset, FLinearColor::Red, DebugClothMaterialVertex);
+#endif
 			}
 			else
 			{
