Chaos Cloth - Fix Bounds that were calculated in world space instead of local component space.
- Fix debug draw box size.
- Add bounds debug draw.

diff --git a/Engine/Plugins/Experimental/ChaosCloth/Source/Chaos/Private/ChaosCloth/ChaosClothingSimulation.cpp b/Engine/Plugins/Experimental/ChaosCloth/Source/Chaos/Private/ChaosCloth/ChaosClothingSimulation.cpp
--- a/Engine/Plugins/Experimental/ChaosCloth/Source/Chaos/Private/ChaosCloth/ChaosClothingSimulation.cpp
+++ b/Engine/Plugins/Experimental/ChaosCloth/Source/Chaos/Private/ChaosCloth/ChaosClothingSimulation.cpp
@@ -50,6 +50,7 @@
 
 namespace ChaosClothingSimulationConsoleVariables
 {
+	TAutoConsoleVariable<bool> CVarDebugDrawBounds          (TEXT("p.ChaosCloth.DebugDrawBounds"              ), false, TEXT("Whether to debug draw the Chaos Cloth bounds"), ECVF_Cheat);
 	TAutoConsoleVariable<bool> CVarDebugDrawPhysMeshWired   (TEXT("p.ChaosCloth.DebugDrawPhysMeshWired"       ), false, TEXT("Whether to debug draw the Chaos Cloth wireframe meshes"), ECVF_Cheat);
 	TAutoConsoleVariable<bool> CVarDebugPointNormals        (TEXT("p.ChaosCloth.DebugDrawPointNormals"        ), false, TEXT("Whether to debug draw the Chaos Cloth point normals"), ECVF_Cheat);
 	TAutoConsoleVariable<bool> CVarDebugInversedPointNormals(TEXT("p.ChaosCloth.DebugDrawInversedPointNormals"), false, TEXT("Whether to debug draw the Chaos Cloth inversed point normals"), ECVF_Cheat);
@@ -1372,6 +1373,7 @@ void ClothingSimulation::Simulate(IClothingSimulationContext* InContext)
 
 	// Debug draw
 #if CHAOS_DEBUG_DRAW
+	if (ChaosClothingSimulationConsoleVariables::CVarDebugDrawBounds          .GetValueOnAnyThread()) { DebugDrawBounds              (); }
 	if (ChaosClothingSimulationConsoleVariables::CVarDebugDrawPhysMeshWired   .GetValueOnAnyThread()) { DebugDrawPhysMeshWired       (); }
 	if (ChaosClothingSimulationConsoleVariables::CVarDebugPointNormals        .GetValueOnAnyThread()) { DebugDrawPointNormals        (); }
 	if (ChaosClothingSimulationConsoleVariables::CVarDebugInversedPointNormals.GetValueOnAnyThread()) { DebugDrawInversedPointNormals(); }
@@ -1438,6 +1440,46 @@ void ClothingSimulation::GetSimulationData(
     }
 }
 
+FBoxSphereBounds ClothingSimulation::GetBounds(const USkeletalMeshComponent* InOwnerComponent) const
+{
+	const TPBDParticles<float, 3>& Particles = Evolution->Particles();
+	if (!Particles.Size())
+	{
+		return FBoxSphereBounds(EForceInit::ForceInit);
+	}
+
+	// Find bounding box
+	TAABB<float, 3> BoundingBox = TAABB<float, 3>::EmptyAABB();
+	for (uint32 ParticleIndex = 0; ParticleIndex < Particles.Size(); ++ParticleIndex)
+	{
+		BoundingBox.GrowToInclude(Particles.X(ParticleIndex) + LocalSimSpaceOffset);
+	}
+
+	// Find bounding sphere (squared) radius
+	const TVector<float, 3> Center = BoundingBox.Center();
+	float SquaredRadius = 0.f;
+	for (uint32 ParticleIndex = 0; ParticleIndex < Particles.Size(); ++ParticleIndex)
+	{
+		SquaredRadius = FMath::Max(SquaredRadius, (Particles.X(ParticleIndex) + LocalSimSpaceOffset - Center).SizeSquared());
+	}
+
+	// Create bounds (in world space)
+	const FBoxSphereBounds Bounds(BoundingBox.Center(), BoundingBox.Extents() * 0.5f, FMath::Sqrt(SquaredRadius));
+
+	// Transform to component space
+	if (InOwnerComponent)
+	{
+		// Retrieve the master component (unlike the one passed to the context, this could be a slave component)
+		const USkinnedMeshComponent* const OwnerComponent = InOwnerComponent->MasterPoseComponent.IsValid() ?
+			InOwnerComponent->MasterPoseComponent.Get() : InOwnerComponent;
+
+		// Return local bounds
+		return Bounds.TransformBy(OwnerComponent->GetComponentTransform().Inverse());
+	}
+	// Return world space bounds (for debug drawing)
+	return Bounds;
+}
+
 void ClothingSimulation::AddExternalCollisions(const FClothCollisionData& InData)
 {
 	// Keep track of the external collisions data
@@ -1694,7 +1736,7 @@ static void DrawBox(FPrimitiveDrawInterface* PDI, const TBox<float, 3>& Box, con
 #if CHAOS_DEBUG_DRAW
 	if (!PDI)
 	{
-		FDebugDrawQueue::GetInstance().DrawDebugBox(Position, Box.Extents(), Rotation, Color.ToFColor(true), false, KINDA_SMALL_NUMBER, SDPG_Foreground, 0.f);
+		FDebugDrawQueue::GetInstance().DrawDebugBox(Position, Box.Extents() * 0.5f, Rotation, Color.ToFColor(true), false, KINDA_SMALL_NUMBER, SDPG_Foreground, 0.f);
 		return;
 	}
 #endif
@@ -1803,6 +1845,18 @@ static void DrawCoordinateSystem(FPrimitiveDrawInterface* PDI, const FQuat& Rota
 	DrawLine(PDI, Position, Position + Z, FLinearColor::Blue);
 }
 
+#if CHAOS_DEBUG_DRAW
+void ClothingSimulation::DebugDrawBounds() const
+{
+	// Calculate World space bounds
+	const FBoxSphereBounds Bounds = GetBounds(nullptr);
+
+	// Draw bounds
+	DrawBox(nullptr, TBox<float, 3>(-Bounds.BoxExtent, Bounds.BoxExtent), FQuat::Identity, Bounds.Origin, FLinearColor(FColor::Purple));
+	DrawSphere(nullptr, TSphere<float, 3>(FVector::ZeroVector, Bounds.SphereRadius), FQuat::Identity, Bounds.Origin, FLinearColor(FColor::Orange));
+}
+#endif  // #if CHAOS_DEBUG_DRAW
+
 void ClothingSimulation::DebugDrawPhysMeshWired(USkeletalMeshComponent* /*OwnerComponent*/, FPrimitiveDrawInterface* PDI) const
 {
 	const TPBDParticles<float, 3>& Particles = Evolution->Particles();
diff --git a/Engine/Plugins/Experimental/ChaosCloth/Source/Chaos/Public/ChaosCloth/ChaosClothingSimulation.h b/Engine/Plugins/Experimental/ChaosCloth/Source/Chaos/Public/ChaosCloth/ChaosClothingSimulation.h
--- a/Engine/Plugins/Experimental/ChaosCloth/Source/Chaos/Public/ChaosCloth/ChaosClothingSimulation.h
+++ b/Engine/Plugins/Experimental/ChaosCloth/Source/Chaos/Public/ChaosCloth/ChaosClothingSimulation.h
@@ -74,8 +74,7 @@ namespace Chaos
 		virtual void DestroyContext(IClothingSimulationContext* InContext) override { delete InContext; }
 		virtual void GetSimulationData(TMap<int32, FClothSimulData>& OutData, USkeletalMeshComponent* InOwnerComponent, USkinnedMeshComponent* InOverrideComponent) const override;
 
-		virtual FBoxSphereBounds GetBounds(const USkeletalMeshComponent* InOwnerComponent) const override
-		{ return FBoxSphereBounds(Evolution->Particles().X().GetData(), Evolution->Particles().Size()); }
+		virtual FBoxSphereBounds GetBounds(const USkeletalMeshComponent* InOwnerComponent) const override;
 
 		virtual void AddExternalCollisions(const FClothCollisionData& InData) override;
 		virtual void ClearExternalCollisions() override;
@@ -83,6 +82,11 @@ namespace Chaos
 		// End of IClothingSimulation interface
 
 	private:
+#if CHAOS_DEBUG_DRAW
+		// Runtime only debug draw functions
+		void DebugDrawBounds() const;
+#endif  // #if CHAOS_DEBUG_DRAW
+
 		void UpdateSimulationFromSharedSimConfig();
 		void BuildMesh(const FClothPhysicalMeshData& PhysMesh, int32 InSimDataIndex);
 
