Chaos Cloth - Add root bone space linear and angular velocity scale.
- Local space simulation fixes.
- Fix per cloth physics asset collisions.
- External collision transforms fixes.
- Add local space debug draw.
- Add gravity debug draw.
- Fix missing gravity scale on gravity overrides.

Michael.Lentine

diff --git a/Engine/Plugins/Experimental/ChaosCloth/Source/Chaos/Private/ChaosCloth/ChaosClothConfig.cpp b/Engine/Plugins/Experimental/ChaosCloth/Source/Chaos/Private/ChaosCloth/ChaosClothConfig.cpp
--- a/Engine/Plugins/Experimental/ChaosCloth/Source/Chaos/Private/ChaosCloth/ChaosClothConfig.cpp
+++ b/Engine/Plugins/Experimental/ChaosCloth/Source/Chaos/Private/ChaosCloth/ChaosClothConfig.cpp
@@ -21,9 +21,6 @@
 //  WindLiftCoefficient
 //  LinearDrag
 //  AngularDrag
-//  LinearInertiaScale
-//  AngularInertiaScale
-//  CentrifugalInertiaScale
 //  StiffnessFrequency
 //  TetherLimit
 //  AnimDriveSpringStiffness
@@ -68,6 +65,10 @@ void UChaosClothConfig::MigrateFrom(const FClothConfig_Legacy& ClothConfig)
 
 	const float Damping = (ClothConfig.Damping.X + ClothConfig.Damping.Y + ClothConfig.Damping.Z) / 3.f;
 	DampingCoefficient = FMath::Clamp(Damping * Damping * 0.95f, 0.f, 1.f);  // Nv Cloth seems to have a different damping formulation.
+
+	LinearVelocityScale = ClothConfig.LinearInertiaScale;
+	const FVector AngularInertiaScale = ClothConfig.AngularInertiaScale * ClothConfig.CentrifugalInertiaScale;
+	AngularVelocityScale = (AngularInertiaScale.X + AngularInertiaScale.Y + AngularInertiaScale.Z) / 3.f;
 }
 
 void UChaosClothConfig::Serialize(FArchive& Ar)
diff --git a/Engine/Plugins/Experimental/ChaosCloth/Source/Chaos/Private/ChaosCloth/ChaosClothingSimulation.cpp b/Engine/Plugins/Experimental/ChaosCloth/Source/Chaos/Private/ChaosCloth/ChaosClothingSimulation.cpp
--- a/Engine/Plugins/Experimental/ChaosCloth/Source/Chaos/Private/ChaosCloth/ChaosClothingSimulation.cpp
+++ b/Engine/Plugins/Experimental/ChaosCloth/Source/Chaos/Private/ChaosCloth/ChaosClothingSimulation.cpp
@@ -50,7 +50,9 @@
 
 namespace ChaosClothingSimulationConsoleVariables
 {
+	TAutoConsoleVariable<bool> CVarDebugDrawLocalSpace      (TEXT("p.ChaosCloth.DebugDrawLocalSpace"          ), false, TEXT("Whether to debug draw the Chaos Cloth local space"), ECVF_Cheat);
 	TAutoConsoleVariable<bool> CVarDebugDrawBounds          (TEXT("p.ChaosCloth.DebugDrawBounds"              ), false, TEXT("Whether to debug draw the Chaos Cloth bounds"), ECVF_Cheat);
+	TAutoConsoleVariable<bool> CVarDebugDrawGravity         (TEXT("p.ChaosCloth.DebugDrawGravity"             ), false, TEXT("Whether to debug draw the Chaos Cloth gravity acceleration vector"), ECVF_Cheat);
 	TAutoConsoleVariable<bool> CVarDebugDrawPhysMeshWired   (TEXT("p.ChaosCloth.DebugDrawPhysMeshWired"       ), false, TEXT("Whether to debug draw the Chaos Cloth wireframe meshes"), ECVF_Cheat);
 	TAutoConsoleVariable<bool> CVarDebugPointNormals        (TEXT("p.ChaosCloth.DebugDrawPointNormals"        ), false, TEXT("Whether to debug draw the Chaos Cloth point normals"), ECVF_Cheat);
 	TAutoConsoleVariable<bool> CVarDebugInversedPointNormals(TEXT("p.ChaosCloth.DebugDrawInversedPointNormals"), false, TEXT("Whether to debug draw the Chaos Cloth inversed point normals"), ECVF_Cheat);
@@ -89,14 +91,8 @@ ClothingSimulation::ClothingSimulation()
 	, bOverrideGravity(false)
 	, Gravity(ChaosClothingSimulationDefault::Gravity)
 	, WindVelocity(FVector::ZeroVector)
-	, bLocalSimSpaceEnabled(false)
-	, LocalSimSpaceOffset(FVector(0.0f, 0.0f, 0.0f))
-	, PrevLocalSimSpaceOffset(FVector(0.0f, 0.0f, 0.0f))
-	, LocalSimSpaceVelocity(FVector(0.0f, 0.0f, 0.0f))
-	, LocalSimSpaceCappedVelocity(FVector(0.0f, 0.0f, 0.0f))
-	, PrevLocalSimSpaceVelocity(FVector(0.0f, 0.0f, 0.0f))
-	, ComponentLinearAccScale(FVector(0.0f, 0.0f, 0.0f))
-	, ComponentLinearAccClamp(FVector(0.0f, 0.0f, 0.0f))
+	, bUseLocalSpaceSimulation(false)
+	, LocalSpaceLocation(FVector::ZeroVector)
 {
 #if WITH_EDITOR
 	DebugClothMaterial = LoadObject<UMaterial>(nullptr, TEXT("/Engine/EditorMaterials/Cloth/CameraLitDoubleSided.CameraLitDoubleSided"), nullptr, LOAD_None, nullptr);  // LOAD_EditorOnly
@@ -171,10 +167,14 @@ void ClothingSimulation::Shutdown()
 	AnimationNormals.Reset();
 	IndexToRangeMap.Reset();
 	RootBoneWorldTransforms.Reset();
+	LinearDeltaRatios.Reset();
+	AngularDeltaRatios.Reset();
 	Meshes.Reset();
 	FaceNormals.Reset();
 	PointNormals.Reset();
 	Evolution.Reset();
+	CollisionsRangeMap.Reset();
+	ExternalCollisionsRangeMaps.Reset();
 	ExternalCollisionsOffset = 0;
 	ClothSharedSimConfig = nullptr;
 	LongRangeConstraints.Reset();
@@ -206,18 +206,27 @@ void ClothingSimulation::CreateActor(USkeletalMeshComponent* InOwnerComponent, U
 
 	if (Assets.Num() <= InSimDataIndex)
 	{
-		Assets.SetNumZeroed(InSimDataIndex + 1);
-		AnimDriveSpringStiffness.SetNumZeroed(InSimDataIndex + 1);
-	}
-	Assets[InSimDataIndex] = Asset;
+		const int32 NumAssets = InSimDataIndex + 1;
 
-	if (Meshes.Num() <= InSimDataIndex)
-	{
-		Meshes.SetNum(InSimDataIndex + 1);
-		FaceNormals.SetNum(InSimDataIndex + 1);
-		PointNormals.SetNum(InSimDataIndex + 1);
-		LongRangeConstraints.SetNum(InSimDataIndex + 1);
+		// TODO: Refactor all these arrays into a single cloth runtime asset structure
+		Assets.SetNumZeroed(NumAssets);
+		AnimDriveSpringStiffness.SetNumZeroed(NumAssets);
+
+		Meshes.SetNum(NumAssets);
+		FaceNormals.SetNum(NumAssets);
+		PointNormals.SetNum(NumAssets);
+
+		IndexToRangeMap.SetNum(NumAssets);
+
+		LongRangeConstraints.SetNum(NumAssets);
+
+		RootBoneWorldTransforms.SetNum(NumAssets);
+		LinearDeltaRatios.SetNum(NumAssets);
+		AngularDeltaRatios.SetNum(NumAssets);
+
+		CollisionsRangeMap.SetNumZeroed(NumAssets);
 	}
+	Assets[InSimDataIndex] = Asset;
 
 	check(Asset->GetNumLods() > 0);
 	UE_CLOG(Asset->GetNumLods() != 1,
@@ -228,16 +237,9 @@ void ClothingSimulation::CreateActor(USkeletalMeshComponent* InOwnerComponent, U
 
 	// Add particles
 	TPBDParticles<float, 3>& Particles = Evolution->Particles();
-	const uint32 Offset = Particles.Size();
-	Evolution->AddParticles(PhysMesh.Vertices.Num(), (uint32)InSimDataIndex);
+	const uint32 Offset = Evolution->AddParticles(PhysMesh.Vertices.Num(), (uint32)InSimDataIndex);
 
-	if (IndexToRangeMap.Num() <= InSimDataIndex)
-	{
-		IndexToRangeMap.SetNum(InSimDataIndex + 1);
-		RootBoneWorldTransforms.SetNum(InSimDataIndex + 1);
-	}
 	IndexToRangeMap[InSimDataIndex] = Chaos::TVector<uint32, 2>(Offset, Particles.Size());
-	RootBoneWorldTransforms[InSimDataIndex] = Context.BoneTransforms[Asset->ReferenceBoneIndex] * Context.ComponentToWorld;
 
 	// ClothSharedSimConfig should either be a nullptr, or point to an object common to the whole skeletal mesh
 	if (ClothSharedSimConfig == nullptr)
@@ -252,19 +254,33 @@ void ClothingSimulation::CreateActor(USkeletalMeshComponent* InOwnerComponent, U
 	AnimationPositions.SetNum(Particles.Size());
 	AnimationNormals.SetNum(Particles.Size());
 
-	//
-	FTransform ComponentToLocalSimSpace = Context.ComponentToWorld;
-	if (bLocalSimSpaceEnabled && Offset == 0)
+	// Initialize the local simulation space transform
+	FTransform ComponentToLocalSpace = Context.ComponentToWorld;
+	if (Offset == 0) // Only initialize this once for all cloth instances
 	{
-		LocalSimSpaceOffset = ComponentToLocalSimSpace.GetTranslation(); // Only initialize this once for all cloth instances
-		PrevLocalSimSpaceOffset = LocalSimSpaceOffset;
+		// Set local offset
+		if (bUseLocalSpaceSimulation)
+		{
+			LocalSpaceLocation = ComponentToLocalSpace.GetTranslation();
+		}
+		else
+		{
+			LocalSpaceLocation = FVector::ZeroVector;
+		}
 	}
-	ComponentToLocalSimSpace.AddToTranslation(-LocalSimSpaceOffset);
+	ComponentToLocalSpace.AddToTranslation(-LocalSpaceLocation);
+
+	// Init local cloth sim space & teleport tranform
+	const FTransform RootBoneTransform = Context.BoneTransforms[Asset->ReferenceBoneIndex];
+	RootBoneWorldTransforms[InSimDataIndex] = RootBoneTransform * Context.ComponentToWorld;  // Velocity scale deltas are calculated in world space
+	LinearDeltaRatios[InSimDataIndex] = FVector::OneVector - ChaosClothSimConfig->LinearVelocityScale.BoundToBox(FVector::ZeroVector, FVector::OneVector);
+	AngularDeltaRatios[InSimDataIndex] = 1.f - FMath::Clamp(ChaosClothSimConfig->AngularVelocityScale, 0.f, 1.f);
 
+	// Skin start pose
 	ClothingMeshUtils::SkinPhysicsMesh<true, false>(
 		Asset->UsedBoneIndices,
 		PhysMesh, // curr pos and norm
-		ComponentToLocalSimSpace,
+		ComponentToLocalSpace,
 		Context.RefToLocals.GetData(),
 		Context.RefToLocals.Num(),
 		reinterpret_cast<TArray<FVector>&>(AnimationPositions),
@@ -327,22 +343,24 @@ void ClothingSimulation::CreateActor(USkeletalMeshComponent* InOwnerComponent, U
 		ExternalCollisions.SphereConnections.Num() == 0 &&
 		ExternalCollisions.Convexes.Num() == 0 &&
 		ExternalCollisions.Boxes.Num() == 0, TEXT("There cannot be any external collisions added before all the cloth assets collisions are processed."));
-	ExtractCollisions(Asset);
+	ExtractCollisions(Asset, InSimDataIndex);
 
 	// Update collision transforms, including initial state for particles' X & R
-	UpdateCollisionTransforms(Context, /*bReinit =*/ false);  // Reinit is false, since we only need to update the new collision particles
+	UpdateCollisionTransforms(Context, InSimDataIndex);
 }
 
-void ClothingSimulation::ExtractCollisions(const UClothingAssetCommon* Asset)
+void ClothingSimulation::ExtractCollisions(const UClothingAssetCommon* Asset, int32 InSimDataIndex)
 {
+	CollisionsRangeMap[InSimDataIndex][0] = Evolution->CollisionParticles().Size();
+
 	// Pull collisions from the specified physics asset inside the clothing asset
-	ExtractPhysicsAssetCollisions(Asset);
+	ExtractPhysicsAssetCollisions(Asset, InSimDataIndex);
 
 	// Extract the legacy Apex collision from the clothing asset
-	ExtractLegacyAssetCollisions(Asset);
+	ExtractLegacyAssetCollisions(Asset, InSimDataIndex);
 
-	// Update the external collision offset
-	ExternalCollisionsOffset = Evolution->CollisionParticles().Size();
+	// Update the external collision offset and collision range for this asset
+	CollisionsRangeMap[InSimDataIndex][1] = ExternalCollisionsOffset = Evolution->CollisionParticles().Size();
 }
 
 void ClothingSimulation::PostActorCreationInitialize()
@@ -354,9 +372,8 @@ void ClothingSimulation::UpdateSimulationFromSharedSimConfig()
 {
 	if (ClothSharedSimConfig) // ClothSharedSimConfig will be a null pointer if all cloth instances are disabled in which case we will use default Evolution parameters
 	{
-		bLocalSimSpaceEnabled = ClothSharedSimConfig->bUseLocalSpaceSimulation;
-		ComponentLinearAccScale = ClothSharedSimConfig->ComponentLinearAccScale;
-		ComponentLinearAccClamp = ClothSharedSimConfig->ComponentLinearAccClamp;
+		// Update local space simulation switch
+		bUseLocalSpaceSimulation = ClothSharedSimConfig->bUseLocalSpaceSimulation;
 
 		// Now set all the common parameters on the simulation
 		Evolution->SetIterations(ClothSharedSimConfig->IterationCount);
@@ -730,71 +747,70 @@ void ClothingSimulation::AddSelfCollisions(int32 InSimDataIndex)
 	}
 }
 
-void ClothingSimulation::UpdateCollisionTransforms(const ClothingSimulationContext& Context, bool bReinit)
+void ClothingSimulation::ForAllCollisions(TFunction<void(TGeometryClothParticles<float, 3>&, uint32)> CollisionFunction, int32 SimDataIndex)
 {
-	SCOPE_CYCLE_COUNTER(STAT_ChaosClothUpdateCollisionTransforms);
+	TGeometryClothParticles<float, 3>& CollisionParticles = Evolution->CollisionParticles();
 
-	// Save collision transforms into the OldCollision transforms
-	// before overwriting it in this function
-	Swap(OldCollisionTransforms, CollisionTransforms);
+	for (uint32 Index = CollisionsRangeMap[SimDataIndex][0]; Index < CollisionsRangeMap[SimDataIndex][1]; ++Index)
+	{
+		CollisionFunction(CollisionParticles, Index);
+	}
+	for (const TArray<TVector<uint32, 2>>& ExternalCollisionsRangeMap : ExternalCollisionsRangeMaps)
+	{
+		for (uint32 Index = ExternalCollisionsRangeMap[SimDataIndex][0]; Index < ExternalCollisionsRangeMap[SimDataIndex][1]; ++Index)
+		{
+			CollisionFunction(CollisionParticles, Index);
+		}
+	}
+}
+
+void ClothingSimulation::UpdateCollisionTransforms(const ClothingSimulationContext& Context, int32 InSimDataIndex)
+{
+	SCOPE_CYCLE_COUNTER(STAT_ChaosClothUpdateCollisionTransforms);
 
 	TGeometryClothParticles<float, 3>& CollisionParticles = Evolution->CollisionParticles();
 
-	// Resize the transform arrays
+	// Resize the transform arrays if collision have changed
 	const int32 PrevNumCollisions = OldCollisionTransforms.Num();
 	const int32 NumCollisions = BaseTransforms.Num();
 	check(NumCollisions == int32(CollisionParticles.Size()));  // BaseTransforms should always automatically grow with the number of collision particles (collection array)
 
-	if (NumCollisions != PrevNumCollisions)
+	const bool bHasNumCollisionsChanged = (NumCollisions != PrevNumCollisions);
+	if (bHasNumCollisionsChanged)
 	{
 		CollisionTransforms.SetNum(NumCollisions);
 		OldCollisionTransforms.SetNum(NumCollisions);
 	}
 
+	// Update the collision transforms
 	FTransform ComponentToLocalSimulationSpace = Context.ComponentToWorld;
-	ComponentToLocalSimulationSpace.AddToTranslation(-LocalSimSpaceOffset);
+	ComponentToLocalSimulationSpace.AddToTranslation(-LocalSpaceLocation);
 
-	// Update the collision transforms
-	for (int32 Index = 0; Index < NumCollisions; ++Index)
+	ForAllCollisions([this, &Context, &ComponentToLocalSimulationSpace, bHasNumCollisionsChanged](TGeometryClothParticles<float, 3>& CollisionParticles, uint32 Index)
 	{
+		// Update the collision transforms
 		const int32 BoneIndex = BoneIndices[Index];
+		Chaos::TRigidTransform<float, 3>& CollisionTransform = CollisionTransforms[Index];
 		if (Context.BoneTransforms.IsValidIndex(BoneIndex))
 		{
 			const FTransform& BoneTransform = Context.BoneTransforms[BoneIndex];
-			CollisionTransforms[Index] = BaseTransforms[Index] * BoneTransform * ComponentToLocalSimulationSpace;
+			CollisionTransform = BaseTransforms[Index] * BoneTransform * ComponentToLocalSimulationSpace;
 		}
 		else
 		{
-			CollisionTransforms[Index] = BaseTransforms[Index] * ComponentToLocalSimulationSpace;  // External collisions often don't map to a bone
+			CollisionTransform = BaseTransforms[Index] * ComponentToLocalSimulationSpace;  // External collisions often don't map to a bone
 		}
-	}
-
-	// External collisions are reinitialized at every frame, but don't affect the size of the arrays if the same amount has been added/removed
-	const int32 Offset = FMath::Min((int32)ExternalCollisionsOffset, PrevNumCollisions);
-
-	// Reinit old collisions particles and transforms
-	if (bReinit)
-	{
-		for (int32 Index = 0; Index < Offset; ++Index)
+		// Reset initial states if required
+		if (bHasNumCollisionsChanged)
 		{
-			const Chaos::TRigidTransform<float, 3>& CollisionTransform = CollisionTransforms[Index];
 			CollisionParticles.X(Index) = CollisionTransform.GetTranslation();
 			CollisionParticles.R(Index) = CollisionTransform.GetRotation();
 			OldCollisionTransforms[Index] = CollisionTransform;
 		}
-	}
-
-	// Set the new collision particles and transforms initial state
-	for (int32 Index = Offset; Index < NumCollisions; ++Index)
-	{
-		const Chaos::TRigidTransform<float, 3>& CollisionTransform = CollisionTransforms[Index];
-		CollisionParticles.X(Index) = CollisionTransform.GetTranslation();
-		CollisionParticles.R(Index) = CollisionTransform.GetRotation();
-		OldCollisionTransforms[Index] = CollisionTransform;
-	}
+	}, InSimDataIndex);
 }
 
-void ClothingSimulation::ExtractPhysicsAssetCollisions(const UClothingAssetCommon* Asset)
+void ClothingSimulation::ExtractPhysicsAssetCollisions(const UClothingAssetCommon* Asset, int32 InSimDataIndex)
 {
 	FClothCollisionData ExtractedCollisions;
 
@@ -969,12 +985,12 @@ void ClothingSimulation::ExtractPhysicsAssetCollisions(const UClothingAssetCommo
 
 		// Add collisions particles
 		UE_LOG(LogChaosCloth, VeryVerbose, TEXT("Adding physics asset collisions..."));
-		AddCollisions(ExtractedCollisions, UsedBoneIndices);
+		AddCollisions(ExtractedCollisions, UsedBoneIndices, InSimDataIndex);
 
 	}  // End if Asset->PhysicsAsset
 }
 
-void ClothingSimulation::ExtractLegacyAssetCollisions(const UClothingAssetCommon* Asset)
+void ClothingSimulation::ExtractLegacyAssetCollisions(const UClothingAssetCommon* Asset, int32 InSimDataIndex)
 {
 	if (const UClothLODDataCommon* const AssetLodData = Asset->ClothLodData[0])
 	{
@@ -982,12 +998,12 @@ void ClothingSimulation::ExtractLegacyAssetCollisions(const UClothingAssetCommon
 		if (LodCollData.Spheres.Num() || LodCollData.SphereConnections.Num() || LodCollData.Convexes.Num())
 		{
 			UE_LOG(LogChaosCloth, VeryVerbose, TEXT("Adding legacy cloth asset collisions..."));
-			AddCollisions(LodCollData, Asset->UsedBoneIndices);
+			AddCollisions(LodCollData, Asset->UsedBoneIndices, InSimDataIndex);
 		}
 	}
 }
 
-void ClothingSimulation::AddCollisions(const FClothCollisionData& ClothCollisionData, const TArray<int32>& UsedBoneIndices)
+void ClothingSimulation::AddCollisions(const FClothCollisionData& ClothCollisionData, const TArray<int32>& UsedBoneIndices, int32 InSimDataIndex)
 {
 	TGeometryClothParticles<float, 3>& CollisionParticles = Evolution->CollisionParticles();
 
@@ -996,8 +1012,7 @@ void ClothingSimulation::AddCollisions(const FClothCollisionData& ClothCollision
 	const int32 NumCapsules = ClothCollisionData.SphereConnections.Num();
 	if (NumCapsules)
 	{
-		const uint32 Offset = CollisionParticles.Size();
-		CollisionParticles.AddParticles(NumCapsules);
+		const uint32 Offset = Evolution->AddCollisionParticles(NumCapsules, InSimDataIndex);
 
 		CapsuleEnds.Reserve(NumCapsules * 2);
 		for (uint32 i = Offset; i < CollisionParticles.Size(); ++i)
@@ -1090,8 +1105,7 @@ void ClothingSimulation::AddCollisions(const FClothCollisionData& ClothCollision
 	const int32 NumSpheres = ClothCollisionData.Spheres.Num() - CapsuleEnds.Num();
 	if (NumSpheres != 0)
 	{
-		const uint32 Offset = CollisionParticles.Size();
-		CollisionParticles.AddParticles(NumSpheres);
+		const uint32 Offset = Evolution->AddCollisionParticles(NumSpheres, InSimDataIndex);
 		// i = CollisionParticles index, j = Spheres index
 		for (uint32 i = Offset, j = 0; j < (uint32)ClothCollisionData.Spheres.Num(); ++j)
 		{
@@ -1124,8 +1138,7 @@ void ClothingSimulation::AddCollisions(const FClothCollisionData& ClothCollision
 	const uint32 NumConvexes = ClothCollisionData.Convexes.Num();
 	if (NumConvexes != 0)
 	{
-		const uint32 Offset = CollisionParticles.Size();
-		CollisionParticles.AddParticles(NumConvexes);
+		const uint32 Offset = Evolution->AddCollisionParticles(NumConvexes, InSimDataIndex);
 		for (uint32 i = Offset; i < CollisionParticles.Size(); ++i)
 		{
 			const FClothCollisionPrim_Convex& Convex = ClothCollisionData.Convexes[i - Offset];
@@ -1199,8 +1212,7 @@ void ClothingSimulation::AddCollisions(const FClothCollisionData& ClothCollision
 	const uint32 NumBoxes = ClothCollisionData.Boxes.Num();
 	if (NumBoxes != 0)
 	{
-		const uint32 Offset = CollisionParticles.Size();
-		CollisionParticles.AddParticles(NumBoxes);
+		const uint32 Offset = Evolution->AddCollisionParticles(NumBoxes, InSimDataIndex);
 		for (uint32 i = Offset; i < CollisionParticles.Size(); ++i)
 		{
 			const FClothCollisionPrim_Box& Box = ClothCollisionData.Boxes[i - Offset];
@@ -1230,50 +1242,13 @@ void ClothingSimulation::Simulate(IClothingSimulationContext* InContext)
 		return;
 	}
 
-	// New LocalSimSpaceOffset (use the component transform for now, but something else like the middle of the bounding volume can be used here)
-	if(bLocalSimSpaceEnabled)
-	{
-		PrevLocalSimSpaceOffset = LocalSimSpaceOffset;
-		LocalSimSpaceOffset = Context->ComponentToWorld.GetTranslation();
-
-		PrevLocalSimSpaceVelocity = LocalSimSpaceVelocity;
-		LocalSimSpaceVelocity = (LocalSimSpaceOffset - PrevLocalSimSpaceOffset) / Context->DeltaSeconds;  // Estimate simspace frame velocity
-		const FVector ScaledVelocityDelta = (LocalSimSpaceVelocity - PrevLocalSimSpaceVelocity) * ComponentLinearAccScale;
-
-		//const FVector ScaledVelocityDeltaMax = ComponentLinearAccClamp * Context->DeltaSeconds;
-		//FVector ScaledClampedVelocityDelta = ScaledVelocityDelta.ComponentMin(ScaledVelocityDeltaMax);
-		//ScaledClampedVelocityDelta = ScaledClampedVelocityDelta.ComponentMax(-ScaledVelocityDeltaMax);
-
-		const FVector SimSpaceDelta = (ScaledVelocityDelta + LocalSimSpaceCappedVelocity) * Context->DeltaSeconds;
-		LocalSimSpaceCappedVelocity = SimSpaceDelta / Context->DeltaSeconds;
-
-		// Particle positions updated to be relative to the new space
-		TPBDParticles<float, 3>& Particles = Evolution->Particles();
-		for (uint32 i = 0; i < Particles.Size(); ++i)
-		{
-			Particles.X(i) -= SimSpaceDelta;
-			Particles.P(i) -= SimSpaceDelta;
-		}
-		TGeometryClothParticles<float, 3>& CollisionParticles = Evolution->CollisionParticles();
-		for (uint32 i = 0; i < CollisionParticles.Size(); i++)
-		{
-			CollisionParticles.X(i) -= SimSpaceDelta;
-		}
-
-		for (int32 i = 0; i < OldAnimationPositions.Num(); i++)
-		{
-			OldAnimationPositions[i] -= SimSpaceDelta;
-		}
-
-		for (int32 i = 0; i < OldCollisionTransforms.Num(); i++)
-		{
-			OldCollisionTransforms[i].AddToTranslation(-SimSpaceDelta);
-		}
-	}
+	// Filter delta time to smoothen time variations and prevent unwanted vibrations
+	static const float Decay = 0.1f;
+	DeltaTime = DeltaTime + (Context->DeltaSeconds - DeltaTime) * Decay;
 
 	// Set gravity, using the legacy priority: 1) config override, 2) game override, 3) world gravity
 	Evolution->GetGravityForces().SetAcceleration(Chaos::TVector<float, 3>(
-		bOverrideGravity ? Gravity :
+		bOverrideGravity ? Gravity * ClothSharedSimConfig->GravityScale :
 		(ClothSharedSimConfig && ClothSharedSimConfig->bUseGravityOverride) ? ClothSharedSimConfig->Gravity :
 		ClothSharedSimConfig ? Context->WorldGravity * ClothSharedSimConfig->GravityScale :
 		Context->WorldGravity));
@@ -1284,87 +1259,176 @@ void ClothingSimulation::Simulate(IClothingSimulationContext* InContext)
 	// Check teleport modes
 	const bool bTeleport = (Context->TeleportMode > EClothingTeleportMode::None);
 	const bool bTeleportAndReset = (Context->TeleportMode == EClothingTeleportMode::TeleportAndReset);
-	UE_CLOG(bTeleport, LogChaosCloth, Verbose, TEXT("Teleport, reset: %d"), bTeleportAndReset);
 
-	// Get New Animation Positions and Normals + deal with teleportation
+	// Get New Animation Positions and Normals + deal with local space & teleportation
 	{
 		SCOPE_CYCLE_COUNTER(STAT_ChaosClothGetAnimationData);
 
-		Swap(OldAnimationPositions, AnimationPositions);
 		check(OldAnimationPositions.Num() == AnimationPositions.Num());
+		Swap(OldAnimationPositions, AnimationPositions);
+
+		check(OldCollisionTransforms.Num() == CollisionTransforms.Num());
+		Swap(OldCollisionTransforms, CollisionTransforms);
 
+		// Update the local space transform
+		const FVector PrevLocalSpaceLocation = LocalSpaceLocation;
+		if (bUseLocalSpaceSimulation)
+		{
+			LocalSpaceLocation = Context->ComponentToWorld.GetLocation();
+		}
+		const FVector DeltaLocalSpaceLocation = LocalSpaceLocation - PrevLocalSpaceLocation;
+
+		// Iterate all cloths
 		TPBDParticles<float, 3>& Particles = Evolution->Particles();
+		TGeometryClothParticles<float, 3>& CollisionParticles = Evolution->CollisionParticles();
 
-		for (int32 Index = 0; Index < IndexToRangeMap.Num(); ++Index)
+		for (int32 Index = 0; Index < Assets.Num(); ++Index)
 		{
 			const UClothingAssetCommon* const Asset = Assets[Index];
-			if (!Asset)
-			{
-				continue;
-			}
+			if (!Asset) { continue; }
 
-			const UClothLODDataCommon* AssetLodData = Asset->ClothLodData[0];
-			const FClothPhysicalMeshData& PhysMesh = AssetLodData->ClothPhysicalMeshData;
-			const uint32 PointCount = IndexToRangeMap[Index][1] - IndexToRangeMap[Index][0];
+			const uint32 Offset = IndexToRangeMap[Index][0];
+			const uint32 Range = IndexToRangeMap[Index][1];
 
+			// Update collision transforms using new local space transform
+			UpdateCollisionTransforms(*Context, Index);
+
+			// Update animation transforms via skinning
 			// Optimization note:
 			// This function usually receives the RootBoneTransform in order to transform the result from Component space to RootBone space.
 			// (So the mesh vectors and positions (Mesh is in component space) is multiplied by Inv(RootBoneTransform) at the end of the function)
 			// We actually require world space coordinates so will instead pass Inv(ComponentToWorld)
 			// This saves a lot of Matrix multiplication work later
-			const int32 Offset = IndexToRangeMap[Index][0];
-
-			FTransform ComponentToLocalSimSpace = Context->ComponentToWorld;
-			ComponentToLocalSimSpace.AddToTranslation(-LocalSimSpaceOffset);
+			FTransform ComponentToLocalSpace = Context->ComponentToWorld;
+			ComponentToLocalSpace.AddToTranslation(-LocalSpaceLocation);
 
 			ClothingMeshUtils::SkinPhysicsMesh<true, false>(
 				Asset->UsedBoneIndices,
-				PhysMesh,
-				ComponentToLocalSimSpace,
+				Asset->ClothLodData[0]->ClothPhysicalMeshData,
+				ComponentToLocalSpace,
 				Context->RefToLocals.GetData(),
 				Context->RefToLocals.Num(),
 				reinterpret_cast<TArray<FVector>&>(AnimationPositions),
 				reinterpret_cast<TArray<FVector>&>(AnimationNormals), Offset);
 
-			// Teleport
+			// Update root bone reference transforms
 			const FTransform RootBoneTransform = Context->BoneTransforms[Asset->ReferenceBoneIndex];
-			const FTransform RootBoneWorldTransform = RootBoneTransform * Context->ComponentToWorld;
+			const FTransform PrevRootBoneWorldTransform = RootBoneWorldTransforms[Index];
+			RootBoneWorldTransforms[Index] = RootBoneTransform * Context->ComponentToWorld;
+
+			FTransform PrevRootBoneLocalTransform = PrevRootBoneWorldTransform;
+			PrevRootBoneLocalTransform.AddToTranslation(-PrevLocalSpaceLocation);
 
-			if (bTeleport)
+			// Teleport & reset
+			if (bTeleportAndReset)
 			{
-				UE_LOG(LogChaosCloth, Verbose, TEXT("Teleport before: %s, after: %s"), *RootBoneWorldTransforms[Index].ToString(), *RootBoneWorldTransform.ToString());
-				const uint32 Range = IndexToRangeMap[Index][1];
-				if (bTeleportAndReset)
+				UE_LOG(LogChaosCloth, Verbose, TEXT("Teleport & Reset"));
+				for (uint32 i = Offset; i < Range; ++i)
 				{
-					// Teleport & reset
-					for (uint32 i = Offset; i < Range; ++i)
-					{
-						Particles.P(i) = Particles.X(i) = AnimationPositions[i];
-						Particles.V(i) = Chaos::TVector<float, 3>(0.f);
-					}
+					// Update initial state for particles
+					Particles.P(i) = Particles.X(i) = AnimationPositions[i];
+					Particles.V(i) = Chaos::TVector<float, 3>(0.f);
+
+					// Update anim initial state (target updated by skinning)
+					OldAnimationPositions[i] = AnimationPositions[i];
 				}
-				else
+				ForAllCollisions([this](TGeometryClothParticles<float, 3>& CollisionParticles, uint32 i)
 				{
-					// Teleport only
-					const FTransform Delta = RootBoneWorldTransforms[Index].GetRelativeTransformReverse(RootBoneWorldTransform);
-					for (uint32 i = Offset; i < Range; ++i)
-					{
-						Particles.X(i) = Delta.TransformPositionNoScale(Particles.X(i));
-						Particles.V(i) = Delta.TransformVectorNoScale(Particles.V(i));
-					}
+					// Update initial state for collisions
+					OldCollisionTransforms[i] = CollisionTransforms[i];
+					CollisionParticles.X(i) = CollisionTransforms[i].GetTranslation();
+					CollisionParticles.R(i) = CollisionTransforms[i].GetRotation();
+				}, Index);
+			}
+			// Teleport only
+			else if (bTeleport)
+			{
+				UE_LOG(LogChaosCloth, Verbose, TEXT("Teleport before: %s, after: %s"), *PrevRootBoneWorldTransform.ToString(), *RootBoneWorldTransforms[Index].ToString());
+				const FTransform DeltaTransform = RootBoneWorldTransforms[Index].GetRelativeTransform(PrevRootBoneWorldTransform);
+				const FMatrix Matrix = (PrevRootBoneLocalTransform.Inverse() * DeltaTransform * PrevRootBoneLocalTransform).ToMatrixNoScale();
+
+				for (uint32 i = Offset; i < Range; ++i)
+				{
+					// Update initial state for particles
+					Particles.P(i) = Particles.X(i) = Matrix.TransformPosition(Particles.X(i)) - DeltaLocalSpaceLocation;
+					Particles.V(i) = Matrix.TransformVector(Particles.V(i));
+
+					// Update anim initial state (target updated by skinning)
+					OldAnimationPositions[i] = Matrix.TransformPosition(OldAnimationPositions[i]) - DeltaLocalSpaceLocation;
 				}
+
+				ForAllCollisions([this, &Matrix, &DeltaLocalSpaceLocation](TGeometryClothParticles<float, 3>& CollisionParticles, uint32 i)
+				{
+					// Update initial state for collisions
+					OldCollisionTransforms[i] = Matrix * OldCollisionTransforms[i];
+					OldCollisionTransforms[i].AddToTranslation(-DeltaLocalSpaceLocation);
+					CollisionParticles.X(i) = OldCollisionTransforms[i].GetTranslation();
+					CollisionParticles.R(i) = OldCollisionTransforms[i].GetRotation();
+				}, Index);
 			}
+			// Apply reference space velocity scales
+			else if (AngularDeltaRatios[Index] > KINDA_SMALL_NUMBER ||
+				LinearDeltaRatios[Index].X > KINDA_SMALL_NUMBER ||
+				LinearDeltaRatios[Index].Y > KINDA_SMALL_NUMBER ||
+				LinearDeltaRatios[Index].Z > KINDA_SMALL_NUMBER)
+			{
+				// Calculate deltas
+				const FTransform DeltaTransform = RootBoneWorldTransforms[Index].GetRelativeTransform(PrevRootBoneWorldTransform);
 
-			RootBoneWorldTransforms[Index] = RootBoneWorldTransform;
-		}
-	}
+				const FVector DeltaPosition = LinearDeltaRatios[Index] * DeltaTransform.GetTranslation();
 
-	// Update collision transforms
-	UpdateCollisionTransforms(*Context, bTeleport);
+				FQuat DeltaRotation = DeltaTransform.GetRotation();
+				FVector Axis;
+				float DeltaAngle;
+				DeltaRotation.ToAxisAndAngle(Axis, DeltaAngle);
+				if (DeltaAngle > PI) { DeltaAngle -= 2.f * PI; }
+				DeltaAngle *= AngularDeltaRatios[Index];
+				DeltaRotation = FQuat(Axis, DeltaAngle);
+				DeltaRotation.Normalize();  // ToMatrixNoScale does not like quaternions built straight from axis angles without being normilized (although they should have been already).
 
-	// Filter delta time to smoothen time variations and prevent unwanted vibrations
-	static const float Decay = 0.1f;
-	DeltaTime = DeltaTime + (Context->DeltaSeconds - DeltaTime) * Decay;
+				// Transform points back into the previous frame of reference before applying the adjusted deltas 
+				const FMatrix Matrix = (PrevRootBoneLocalTransform.Inverse() * FTransform(DeltaRotation, DeltaPosition) * PrevRootBoneLocalTransform).ToMatrixNoScale();
+
+				for (uint32 i = Offset; i < Range; ++i)
+				{
+					// Update initial state for particles
+					Particles.P(i) = Particles.X(i) = Matrix.TransformPosition(Particles.X(i)) - DeltaLocalSpaceLocation;
+					Particles.V(i) = Matrix.TransformVector(Particles.V(i));
+
+					// Update anim initial state (target updated by skinning)
+					OldAnimationPositions[i] = Matrix.TransformPosition(OldAnimationPositions[i]) - DeltaLocalSpaceLocation;
+				}
+
+				ForAllCollisions([this, &Matrix, &DeltaLocalSpaceLocation](TGeometryClothParticles<float, 3>& CollisionParticles, uint32 i)
+				{
+					// Update initial state for collisions
+					OldCollisionTransforms[i] = Matrix * OldCollisionTransforms[i];
+					OldCollisionTransforms[i].AddToTranslation(-DeltaLocalSpaceLocation);
+
+					CollisionParticles.X(i) = OldCollisionTransforms[i].GetTranslation();
+					CollisionParticles.R(i) = OldCollisionTransforms[i].GetRotation();
+				}, Index);
+			}
+			else if (bUseLocalSpaceSimulation)
+			{
+				for (uint32 i = Offset; i < Range; ++i)
+				{
+					// Update initial state for particles
+					Particles.P(i) = Particles.X(i) -= DeltaLocalSpaceLocation;
+
+					// Update anim initial state (target updated by skinning)
+					OldAnimationPositions[i] -= DeltaLocalSpaceLocation;
+				}
+
+				ForAllCollisions([this, &DeltaLocalSpaceLocation](TGeometryClothParticles<float, 3>& CollisionParticles, uint32 i)
+				{
+					// Update initial state for collisions
+					OldCollisionTransforms[i].AddToTranslation(-DeltaLocalSpaceLocation);
+					CollisionParticles.X(i) = OldCollisionTransforms[i].GetTranslation();
+				}, Index);
+			}
+		}
+	}
 
 	// Advance Sim
 	Evolution->AdvanceOneTimeStep(DeltaTime);
@@ -1373,7 +1437,9 @@ void ClothingSimulation::Simulate(IClothingSimulationContext* InContext)
 
 	// Debug draw
 #if CHAOS_DEBUG_DRAW
+	if (ChaosClothingSimulationConsoleVariables::CVarDebugDrawLocalSpace      .GetValueOnAnyThread()) { DebugDrawLocalSpace          (); }
 	if (ChaosClothingSimulationConsoleVariables::CVarDebugDrawBounds          .GetValueOnAnyThread()) { DebugDrawBounds              (); }
+	if (ChaosClothingSimulationConsoleVariables::CVarDebugDrawGravity         .GetValueOnAnyThread()) { DebugDrawGravity             (); }
 	if (ChaosClothingSimulationConsoleVariables::CVarDebugDrawPhysMeshWired   .GetValueOnAnyThread()) { DebugDrawPhysMeshWired       (); }
 	if (ChaosClothingSimulationConsoleVariables::CVarDebugPointNormals        .GetValueOnAnyThread()) { DebugDrawPointNormals        (); }
 	if (ChaosClothingSimulationConsoleVariables::CVarDebugInversedPointNormals.GetValueOnAnyThread()) { DebugDrawInversedPointNormals(); }
@@ -1434,7 +1500,7 @@ void ClothingSimulation::GetSimulationData(
 		for (uint32 j = VertexDomain[0]; j < VertexDomain[1]; ++j)
         {
 			const uint32 LocalIndex = j - VertexDomain[0];
-            Data.Positions[LocalIndex] = Evolution->Particles().X(j) + LocalSimSpaceOffset;
+            Data.Positions[LocalIndex] = Evolution->Particles().X(j) + LocalSpaceLocation;
             Data.Normals[LocalIndex] = -PointNormals[i][LocalIndex]; // Note the Normals are inverted due to how barycentric coordinates are calculated (see GetPointBaryAndDist in ClothingMeshUtils.cpp)
 		}
     }
@@ -1442,41 +1508,48 @@ void ClothingSimulation::GetSimulationData(
 
 FBoxSphereBounds ClothingSimulation::GetBounds(const USkeletalMeshComponent* InOwnerComponent) const
 {
+	FBoxSphereBounds Bounds(EForceInit::ForceInit);
+
+	// Calculate simulation bounds (in world space)
+	uint32 NumCloths = 0;
 	const TPBDParticles<float, 3>& Particles = Evolution->Particles();
-	if (!Particles.Size())
-	{
-		return FBoxSphereBounds(EForceInit::ForceInit);
-	}
 
-	// Find bounding box
-	TAABB<float, 3> BoundingBox = TAABB<float, 3>::EmptyAABB();
-	for (uint32 ParticleIndex = 0; ParticleIndex < Particles.Size(); ++ParticleIndex)
+	for (int32 Index = 0; Index < Assets.Num(); ++Index)
 	{
-		BoundingBox.GrowToInclude(Particles.X(ParticleIndex) + LocalSimSpaceOffset);
-	}
+		if (const UClothingAssetCommon* const Asset = Assets[Index])
+		{
+			const TVector<uint32, 2> Range = IndexToRangeMap[Index];
 
-	// Find bounding sphere (squared) radius
-	const TVector<float, 3> Center = BoundingBox.Center();
-	float SquaredRadius = 0.f;
-	for (uint32 ParticleIndex = 0; ParticleIndex < Particles.Size(); ++ParticleIndex)
-	{
-		SquaredRadius = FMath::Max(SquaredRadius, (Particles.X(ParticleIndex) + LocalSimSpaceOffset - Center).SizeSquared());
-	}
+			// Find bounds
+			TAABB<float, 3> BoundingBox = TAABB<float, 3>::EmptyAABB();
+			for (uint32 ParticleIndex = Range[0]; ParticleIndex < Range[1]; ++ParticleIndex)
+			{
+				BoundingBox.GrowToInclude(Particles.X(ParticleIndex));
+			}
 
-	// Create bounds (in world space)
-	const FBoxSphereBounds Bounds(BoundingBox.Center(), BoundingBox.Extents() * 0.5f, FMath::Sqrt(SquaredRadius));
+			// Find (squared) radius
+			const TVector<float, 3> Center = BoundingBox.Center();
+			float SquaredRadius = 0.f;
+			for (uint32 ParticleIndex = Range[0]; ParticleIndex < Range[1]; ++ParticleIndex)
+			{
+				SquaredRadius = FMath::Max(SquaredRadius, (Particles.X(ParticleIndex) - Center).SizeSquared());
+			}
+
+			// Update bounds with this cloth
+			const FBoxSphereBounds ClothBounds(BoundingBox.Center(), BoundingBox.Extents() * 0.5f, FMath::Sqrt(SquaredRadius));
+			Bounds = (NumCloths++ == 0) ? ClothBounds : Bounds + ClothBounds;
+		}
+	}
 
-	// Transform to component space
-	if (InOwnerComponent)
+	if (!bUseLocalSpaceSimulation && NumCloths && InOwnerComponent)
 	{
 		// Retrieve the master component (unlike the one passed to the context, this could be a slave component)
-		const USkinnedMeshComponent* const OwnerComponent = InOwnerComponent->MasterPoseComponent.IsValid() ?
-			InOwnerComponent->MasterPoseComponent.Get() : InOwnerComponent;
+		const bool bIsUsingMaster = InOwnerComponent->MasterPoseComponent.IsValid();
+		const USkinnedMeshComponent* const OwnerComponent = bIsUsingMaster ? InOwnerComponent->MasterPoseComponent.Get() : InOwnerComponent;
 
 		// Return local bounds
 		return Bounds.TransformBy(OwnerComponent->GetComponentTransform().Inverse());
 	}
-	// Return world space bounds (for debug drawing)
 	return Bounds;
 }
 
@@ -1485,22 +1558,47 @@ void ClothingSimulation::AddExternalCollisions(const FClothCollisionData& InData
 	// Keep track of the external collisions data
 	ExternalCollisions.Append(InData);
 
-	// Setup the new collisions particles
+	// Add new map entry
+	const int32 MapIndex = ExternalCollisionsRangeMaps.AddDefaulted();
+	ExternalCollisionsRangeMaps[MapIndex].AddUninitialized(Assets.Num());
+
+	// Setup the new collisions particles for all cloths
 	UE_LOG(LogChaosCloth, VeryVerbose, TEXT("Adding external collisions..."));
-	const TArray<int32> UsedBoneIndices;  // There is no bone mapping available for external collisions
-	AddCollisions(InData, UsedBoneIndices);
+	static const TArray<int32> EmptyUsedBoneIndices;  // There is no bone mapping available for external collisions
+	for (int32 Index = 0; Index < Assets.Num(); ++Index)
+	{
+		ExternalCollisionsRangeMaps[MapIndex][Index][0] = Evolution->CollisionParticles().Size();
+		if (Assets[Index])
+		{
+			AddCollisions(InData, EmptyUsedBoneIndices, Index);
+		}
+		ExternalCollisionsRangeMaps[MapIndex][Index][1] = Evolution->CollisionParticles().Size();
+
+		// Keep collision transforms from previous frame if they exist
+		TGeometryClothParticles<float, 3>& CollisionParticles = Evolution->CollisionParticles();
+		for (uint32 i = ExternalCollisionsRangeMaps[MapIndex][Index][0];
+			i < FMath::Min(ExternalCollisionsRangeMaps[MapIndex][Index][1], (uint32)CollisionTransforms.Num());
+			++i)
+		{
+			CollisionParticles.X(i) = CollisionTransforms[i].GetLocation();
+			CollisionParticles.R(i) = CollisionTransforms[i].GetRotation();
+		}
+	}
 }
 
 void ClothingSimulation::ClearExternalCollisions()
 {
 	// Remove all external collision particles, starting from the external collision offset
-	// But do not resize CollisionTransforms as it is only resized in UpdateCollisionTransforms()
+	// But do not resize CollisionTransforms as it is only resized in UpdateCollisionTransforms() to keep old transforms in between frames
 	TGeometryClothParticles<float, 3>& CollisionParticles = Evolution->CollisionParticles();
-	CollisionParticles.Resize(ExternalCollisionsOffset);  // This will also resize BoneIndices and BaseTransforms
+	CollisionParticles.Resize(ExternalCollisionsOffset);  // This will also resize GroupIds, BoneIndices and BaseTransforms
 
 	// Reset external collisions
 	ExternalCollisions.Reset();
 
+	// Reset external collision maps
+	ExternalCollisionsRangeMaps.Reset();
+
 	UE_LOG(LogChaosCloth, VeryVerbose, TEXT("Cleared all external collisions."));
 }
 
@@ -1553,6 +1651,9 @@ void ClothingSimulation::RefreshClothConfig()
 
 				AddConstraints(ChaosClothConfig, PhysMesh, SimDataIndex);
 
+				LinearDeltaRatios[SimDataIndex] = FVector::OneVector - ChaosClothConfig->LinearVelocityScale.BoundToBox(FVector::ZeroVector, FVector::OneVector);
+				AngularDeltaRatios[SimDataIndex] = 1.f - FMath::Clamp(ChaosClothConfig->AngularVelocityScale, 0.f, 1.f);
+
 				// Set damping
 				if (ClothSharedSimConfig && ClothSharedSimConfig->bUseDampingOverride)
 				{
@@ -1595,9 +1696,12 @@ void ClothingSimulation::RefreshPhysicsAsset()
 	ExternalCollisionsOffset = 0;
 
 	// Re-extract all collisions from every cloth asset
-	for (const UClothingAssetCommon* Asset : Assets)
+	for (int32 Index = 0; Index < Assets.Num(); ++Index)
 	{
-		ExtractCollisions(Asset);
+		if (const UClothingAssetCommon* const Asset = Assets[Index])
+		{
+			ExtractCollisions(Asset, Index);
+		}
 	}
 	UE_LOG(LogChaosCloth, VeryVerbose, TEXT("RefreshPhysicsAsset, all collisions have been re-added for all clothing assets"));
 }
@@ -1635,9 +1739,12 @@ void ClothingSimulation::DebugDrawPhysMeshShaded(USkeletalMeshComponent* OwnerCo
 	const TPBDParticles<float, 3>& Particles = Evolution->Particles();
 
 	int32 VertexIndex = 0;
-	for (int32 MeshIndex = 0; MeshIndex < Meshes.Num(); ++MeshIndex)
+
+	for (int32 Index = 0; Index < Assets.Num(); ++Index)
 	{
-		if (const TUniquePtr<Chaos::TTriangleMesh<float>>& Mesh = Meshes[MeshIndex])
+		if (!Assets[Index]) { continue; }
+
+		if (const TUniquePtr<Chaos::TTriangleMesh<float>>& Mesh = Meshes[Index])
 		{
 			const TArray<TVector<int32, 3>>& Elements = Mesh->GetElements();
 
@@ -1661,7 +1768,7 @@ void ClothingSimulation::DebugDrawPhysMeshShaded(USkeletalMeshComponent* OwnerCo
 	}
 
 	FMatrix LocalSimSpaceToWorld(FMatrix::Identity);
-	LocalSimSpaceToWorld.SetOrigin(LocalSimSpaceOffset);
+	LocalSimSpaceToWorld.SetOrigin(LocalSpaceLocation);
 	MeshBuilder.Draw(PDI, LocalSimSpaceToWorld, DebugClothMaterial->GetRenderProxy(), SDPG_World, false, false);
 }
 #endif  // #if WITH_EDITOR
@@ -1836,9 +1943,9 @@ static void DrawConvex(FPrimitiveDrawInterface* PDI, const FConvex& Convex, cons
 
 static void DrawCoordinateSystem(FPrimitiveDrawInterface* PDI, const FQuat& Rotation, const FVector& Position)
 {
-	const FVector X = Rotation.RotateVector(FVector::ForwardVector);
-	const FVector Y = Rotation.RotateVector(FVector::RightVector);
-	const FVector Z = Rotation.RotateVector(FVector::UpVector);
+	const FVector X = Rotation.RotateVector(FVector::ForwardVector) * 10.f;
+	const FVector Y = Rotation.RotateVector(FVector::RightVector) * 10.f;
+	const FVector Z = Rotation.RotateVector(FVector::UpVector) * 10.f;
 
 	DrawLine(PDI, Position, Position + X, FLinearColor::Red);
 	DrawLine(PDI, Position, Position + Y, FLinearColor::Green);
@@ -1852,8 +1959,19 @@ void ClothingSimulation::DebugDrawBounds() const
 	const FBoxSphereBounds Bounds = GetBounds(nullptr);
 
 	// Draw bounds
-	DrawBox(nullptr, TBox<float, 3>(-Bounds.BoxExtent, Bounds.BoxExtent), FQuat::Identity, Bounds.Origin, FLinearColor(FColor::Purple));
-	DrawSphere(nullptr, TSphere<float, 3>(FVector::ZeroVector, Bounds.SphereRadius), FQuat::Identity, Bounds.Origin, FLinearColor(FColor::Orange));
+	DrawBox(nullptr, TBox<float, 3>(-Bounds.BoxExtent, Bounds.BoxExtent), FQuat::Identity, LocalSpaceLocation + Bounds.Origin, FLinearColor(FColor::Purple));
+	DrawSphere(nullptr, TSphere<float, 3>(FVector::ZeroVector, Bounds.SphereRadius), FQuat::Identity, LocalSpaceLocation + Bounds.Origin, FLinearColor(FColor::Orange));
+}
+
+void ClothingSimulation::DebugDrawGravity() const
+{
+	// Calculate World space bounds
+	const FBoxSphereBounds Bounds = GetBounds(nullptr);
+
+	// Draw gravity
+	const FVector Pos0 = LocalSpaceLocation + Bounds.Origin;
+	const FVector Pos1 = Pos0 + Evolution->GetGravityForces().GetAcceleration();
+	DrawLine(nullptr, Pos0, Pos1, FLinearColor::Red);
 }
 #endif  // #if CHAOS_DEBUG_DRAW
 
@@ -1861,19 +1979,19 @@ void ClothingSimulation::DebugDrawPhysMeshWired(USkeletalMeshComponent* /*OwnerC
 {
 	const TPBDParticles<float, 3>& Particles = Evolution->Particles();
 
-	for (int32 MeshIndex = 0; MeshIndex < Meshes.Num(); ++MeshIndex)
+	for (int32 Index = 0; Index < Assets.Num(); ++Index)
 	{
-		if (const TUniquePtr<TTriangleMesh<float>>& Mesh = Meshes[MeshIndex])
+		if (Assets[Index] && Meshes[Index])
 		{
-			const TArray<TVector<int32, 3>>& Elements = Mesh->GetElements();
+			const TArray<TVector<int32, 3>>& Elements = Meshes[Index]->GetElements();
 
 			for (int32 ElementIndex = 0; ElementIndex < Elements.Num(); ++ElementIndex)
 			{
 				const auto& Element = Elements[ElementIndex];
 
-				const FVector Pos0 = Particles.X(Element.X) + LocalSimSpaceOffset;
-				const FVector Pos1 = Particles.X(Element.Y) + LocalSimSpaceOffset;
-				const FVector Pos2 = Particles.X(Element.Z) + LocalSimSpaceOffset;
+				const FVector Pos0 = LocalSpaceLocation + Particles.X(Element.X);
+				const FVector Pos1 = LocalSpaceLocation + Particles.X(Element.Y);
+				const FVector Pos2 = LocalSpaceLocation + Particles.X(Element.Z);
 
 				DrawLine(PDI, Pos0, Pos1, FLinearColor::White);
 				DrawLine(PDI, Pos1, Pos2, FLinearColor::White);
@@ -1885,20 +2003,18 @@ void ClothingSimulation::DebugDrawPhysMeshWired(USkeletalMeshComponent* /*OwnerC
 
 void ClothingSimulation::DebugDrawPointNormals(USkeletalMeshComponent* OwnerComponent, FPrimitiveDrawInterface* PDI) const
 {
-	check(Meshes.Num() == IndexToRangeMap.Num());
-
 	const TPBDParticles<float, 3>& Particles = Evolution->Particles();
 
-	for (int32 MeshIndex = 0; MeshIndex < Meshes.Num(); ++MeshIndex)
+	for (int32 Index = 0; Index < Assets.Num(); ++Index)
 	{
-		if (Meshes[MeshIndex])
+		if (Assets[Index])
 		{
-			const TVector<uint32, 2> Range = IndexToRangeMap[MeshIndex];
-			const TArray<TVector<float, 3>>& MeshPointNormals = PointNormals[MeshIndex];
+			const TVector<uint32, 2> Range = IndexToRangeMap[Index];
+			const TArray<TVector<float, 3>>& MeshPointNormals = PointNormals[Index];
 
 			for (uint32 ParticleIndex = Range[0]; ParticleIndex < Range[1]; ++ParticleIndex)
 			{
-				const FVector Pos0 = Particles.X(ParticleIndex) + LocalSimSpaceOffset;
+				const FVector Pos0 = LocalSpaceLocation + Particles.X(ParticleIndex);
 				const FVector Pos1 = Pos0 + MeshPointNormals[ParticleIndex - Range[0]] * 20.0f;
 
 				DrawLine(PDI, Pos0, Pos1, FLinearColor::White);
@@ -1909,20 +2025,18 @@ void ClothingSimulation::DebugDrawPointNormals(USkeletalMeshComponent* OwnerComp
 
 void ClothingSimulation::DebugDrawInversedPointNormals(USkeletalMeshComponent* OwnerComponent, FPrimitiveDrawInterface* PDI) const
 {
-	check(Meshes.Num() == IndexToRangeMap.Num());
-
 	const TPBDParticles<float, 3>& Particles = Evolution->Particles();
 
-	for (int32 MeshIndex = 0; MeshIndex < Meshes.Num(); ++MeshIndex)
+	for (int32 Index = 0; Index < Assets.Num(); ++Index)
 	{
-		if (Meshes[MeshIndex])
+		if (Assets[Index])
 		{
-			const TVector<uint32, 2> Range = IndexToRangeMap[MeshIndex];
-			const TArray<TVector<float, 3>>& MeshPointNormals = PointNormals[MeshIndex];
+			const TVector<uint32, 2> Range = IndexToRangeMap[Index];
+			const TArray<TVector<float, 3>>& MeshPointNormals = PointNormals[Index];
 
 			for (uint32 ParticleIndex = Range[0]; ParticleIndex < Range[1]; ++ParticleIndex)
 			{
-				const FVector Pos0 = Particles.X(ParticleIndex) + LocalSimSpaceOffset;
+				const FVector Pos0 = LocalSpaceLocation + Particles.X(ParticleIndex);
 				const FVector Pos1 = Pos0 - MeshPointNormals[ParticleIndex - Range[0]] * 20.0f;
 
 				DrawLine(PDI, Pos0, Pos1, FLinearColor::White);
@@ -1933,22 +2047,20 @@ void ClothingSimulation::DebugDrawInversedPointNormals(USkeletalMeshComponent* O
 
 void ClothingSimulation::DebugDrawFaceNormals(USkeletalMeshComponent* OwnerComponent, FPrimitiveDrawInterface* PDI) const
 {
-	check(Meshes.Num() == IndexToRangeMap.Num());
-
 	const TPBDParticles<float, 3>& Particles = Evolution->Particles();
 
-	for (int32 MeshIndex = 0; MeshIndex < Meshes.Num(); ++MeshIndex)
+	for (int32 Index = 0; Index < Assets.Num(); ++Index)
 	{
-		if (const TUniquePtr<TTriangleMesh<float>>& Mesh = Meshes[MeshIndex])
+		if (Assets[Index])
 		{
-			const TArray<TVector<float, 3>>& MeshFaceNormals = FaceNormals[MeshIndex];
+			const TArray<TVector<float, 3>>& MeshFaceNormals = FaceNormals[Index];
 
-			const TArray<TVector<int32, 3>>& Elements = Mesh->GetElements();
+			const TArray<TVector<int32, 3>>& Elements = Meshes[Index]->GetElements();
 			for (int32 ElementIndex = 0; ElementIndex < Elements.Num(); ++ElementIndex)
 			{
 				const TVector<int32, 3>& Element = Elements[ElementIndex];
 
-				const FVector Pos0 = LocalSimSpaceOffset + (
+				const FVector Pos0 = LocalSpaceLocation + (
 					Particles.X(Element.X) +
 					Particles.X(Element.Y) +
 					Particles.X(Element.Z)) / 3.f;
@@ -1962,22 +2074,20 @@ void ClothingSimulation::DebugDrawFaceNormals(USkeletalMeshComponent* OwnerCompo
 
 void ClothingSimulation::DebugDrawInversedFaceNormals(USkeletalMeshComponent* OwnerComponent, FPrimitiveDrawInterface* PDI) const
 {
-	check(Meshes.Num() == IndexToRangeMap.Num());
-
 	const TPBDParticles<float, 3>& Particles = Evolution->Particles();
 
-	for (int32 MeshIndex = 0; MeshIndex < Meshes.Num(); ++MeshIndex)
+	for (int32 Index = 0; Index < Assets.Num(); ++Index)
 	{
-		if (const TUniquePtr<TTriangleMesh<float>>& Mesh = Meshes[MeshIndex])
+		if (Assets[Index])
 		{
-			const TArray<TVector<float, 3>>& MeshFaceNormals = FaceNormals[MeshIndex];
+			const TArray<TVector<float, 3>>& MeshFaceNormals = FaceNormals[Index];
 
-			const TArray<TVector<int32, 3>>& Elements = Mesh->GetElements();
+			const TArray<TVector<int32, 3>>& Elements = Meshes[Index]->GetElements();
 			for (int32 ElementIndex = 0; ElementIndex < Elements.Num(); ++ElementIndex)
 			{
 				const TVector<int32, 3>& Element = Elements[ElementIndex];
 
-				const TVector<float, 3> Pos0 = LocalSimSpaceOffset + (
+				const FVector Pos0 = LocalSpaceLocation + (
 					Particles.X(Element.X) +
 					Particles.X(Element.Y) +
 					Particles.X(Element.Z)) / 3.f;
@@ -1989,66 +2099,83 @@ void ClothingSimulation::DebugDrawInversedFaceNormals(USkeletalMeshComponent* Ow
 	}
 }
 
-void ClothingSimulation::DebugDrawCollision(USkeletalMeshComponent* OwnerComponent, FPrimitiveDrawInterface* PDI) const
+void ClothingSimulation::DebugDrawCollision(USkeletalMeshComponent* /*OwnerComponent*/, FPrimitiveDrawInterface* PDI) const
 {
-	static const FLinearColor MappedColor(FColor::Cyan);
-	static const FLinearColor UnmappedColor(FColor::Red);
-
-	const TGeometryClothParticles<float, 3>& CollisionParticles = Evolution->CollisionParticles();
-
-	for (uint32 Index = 0; Index < CollisionParticles.Size(); ++Index)
+	auto DrawCollision = [this, PDI](const TVector<uint32, 2>& Ranges)
 	{
-		if (const FImplicitObject* const Object = CollisionParticles.DynamicGeometry(Index).Get())
+		static const FLinearColor MappedColor(FColor::Cyan);
+		static const FLinearColor UnmappedColor(FColor::Red);
+
+		const TGeometryClothParticles<float, 3>& CollisionParticles = Evolution->CollisionParticles();
+		for (uint32 Index = Ranges[0]; Index < Ranges[1]; ++Index)
 		{
-			const uint32 BoneIndex = BoneIndices[Index];
-			const FLinearColor Color = (BoneIndex != INDEX_NONE) ? MappedColor : UnmappedColor;
+			if (const FImplicitObject* const Object = CollisionParticles.DynamicGeometry(Index).Get())
+			{
+				const uint32 BoneIndex = BoneIndices[Index];
+				const FLinearColor Color = (BoneIndex != INDEX_NONE) ? MappedColor : UnmappedColor;
 
-			const TVector<float, 3>& Position = CollisionParticles.X(Index) + LocalSimSpaceOffset;
-			const TRotation<float, 3>& Rotation = CollisionParticles.R(Index);
+				const TVector<float, 3> Position = LocalSpaceLocation + CollisionParticles.X(Index);
+				const TRotation<float, 3>& Rotation = CollisionParticles.R(Index);
 
-			switch (Object->GetType())
-			{
-			case ImplicitObjectType::Sphere:
-				DrawSphere(PDI, Object->GetObjectChecked<TSphere<float, 3>>(), Rotation, Position, Color);
-				break;
+				switch (Object->GetType())
+				{
+				case ImplicitObjectType::Sphere:
+					DrawSphere(PDI, Object->GetObjectChecked<TSphere<float, 3>>(), Rotation, Position, Color);
+					break;
 
-			case ImplicitObjectType::Box:
-				DrawBox(PDI, Object->GetObjectChecked<TBox<float, 3>>(), Rotation, Position, Color);
-				break;
+				case ImplicitObjectType::Box:
+					DrawBox(PDI, Object->GetObjectChecked<TBox<float, 3>>(), Rotation, Position, Color);
+					break;
 
-			case ImplicitObjectType::Capsule:
-				DrawCapsule(PDI, Object->GetObjectChecked<TCapsule<float>>(), Rotation, Position, Color);
-				break;
+				case ImplicitObjectType::Capsule:
+					DrawCapsule(PDI, Object->GetObjectChecked<TCapsule<float>>(), Rotation, Position, Color);
+					break;
 
-			case ImplicitObjectType::Union:  // Union only used as collision tappered capsules
-				for (const TUniquePtr<FImplicitObject>& SubObjectPtr : Object->GetObjectChecked<FImplicitObjectUnion>().GetObjects())
-				{
-					if (const FImplicitObject* const SubObject = SubObjectPtr.Get())
+				case ImplicitObjectType::Union:  // Union only used as collision tappered capsules
+					for (const TUniquePtr<FImplicitObject>& SubObjectPtr : Object->GetObjectChecked<FImplicitObjectUnion>().GetObjects())
 					{
-						switch (SubObject->GetType())
+						if (const FImplicitObject* const SubObject = SubObjectPtr.Get())
 						{
-						case ImplicitObjectType::Sphere:
-							DrawSphere(PDI, SubObject->GetObjectChecked<TSphere<float, 3>>(), Rotation, Position, Color);
-							break;
+							switch (SubObject->GetType())
+							{
+							case ImplicitObjectType::Sphere:
+								DrawSphere(PDI, SubObject->GetObjectChecked<TSphere<float, 3>>(), Rotation, Position, Color);
+								break;
 
-						case ImplicitObjectType::TaperedCylinder:
-							DrawTaperedCylinder(PDI, SubObject->GetObjectChecked<TTaperedCylinder<float>>(), Rotation, Position, Color);
-							break;
+							case ImplicitObjectType::TaperedCylinder:
+								DrawTaperedCylinder(PDI, SubObject->GetObjectChecked<TTaperedCylinder<float>>(), Rotation, Position, Color);
+								break;
 
-						default:
-							break;
+							default:
+								break;
+							}
 						}
 					}
+					break;
+
+				case ImplicitObjectType::Convex:
+					DrawConvex(PDI, Object->GetObjectChecked<FConvex>(), Rotation, Position, Color);
+					break;
+
+				default:
+					DrawCoordinateSystem(PDI, Rotation, Position);  // Draw everything else as a coordinate for now
+					break;
 				}
-				break;
-	
-			case ImplicitObjectType::Convex:
-				DrawConvex(PDI, Object->GetObjectChecked<FConvex>(), Rotation, Position, Color);
-				break;
-
-			default:
-				DrawCoordinateSystem(PDI, Rotation, Position);  // Draw everything else as a coordinate for now
-				break;
+			}
+		}
+	};
+
+	for (int32 Index = 0; Index < Assets.Num(); ++Index)
+	{
+		if (Assets[Index])
+		{
+			// Draw collisions
+			DrawCollision(CollisionsRangeMap[Index]);
+
+			// Draw external collisions
+			for (const TArray<TVector<uint32, 2>>& ExternalCollisionsRangeMap : ExternalCollisionsRangeMaps)
+			{
+				DrawCollision(ExternalCollisionsRangeMap[Index]);
 			}
 		}
 	}
@@ -2056,13 +2183,10 @@ void ClothingSimulation::DebugDrawCollision(USkeletalMeshComponent* OwnerCompone
 
 void ClothingSimulation::DebugDrawBackstops(USkeletalMeshComponent* /*OwnerComponent*/, FPrimitiveDrawInterface* PDI) const
 {
-	for (int32 i = 0; i < IndexToRangeMap.Num(); ++i)
+	for (int32 Index = 0; Index < Assets.Num(); ++Index)
 	{
-		const UClothingAssetCommon* const Asset = Assets[i];
-		if (Asset == nullptr)
-		{
-			continue;
-		}
+		const UClothingAssetCommon* const Asset = Assets[Index];
+		if (!Asset) { continue; }
 
 		// Get Backstop Distances
 		const UClothLODDataCommon* const AssetLodData = Asset->ClothLodData[0];
@@ -2075,17 +2199,17 @@ void ClothingSimulation::DebugDrawBackstops(USkeletalMeshComponent* /*OwnerCompo
 			continue;
 		}
 
-		for (uint32 ParticleIndex = IndexToRangeMap[i][0]; ParticleIndex < IndexToRangeMap[i][1]; ++ParticleIndex)
+		for (uint32 ParticleIndex = IndexToRangeMap[Index][0]; ParticleIndex < IndexToRangeMap[Index][1]; ++ParticleIndex)
 		{
-			const uint32 WeightMapIndex = ParticleIndex - IndexToRangeMap[i][0];
+			const uint32 WeightMapIndex = ParticleIndex - IndexToRangeMap[Index][0];
 			const float Radius = BackstopRadiuses[WeightMapIndex];
 			const float Distance = BackstopDistances[WeightMapIndex];
-			DrawLine(PDI, AnimationPositions[ParticleIndex] + LocalSimSpaceOffset, AnimationPositions[ParticleIndex] - AnimationNormals[ParticleIndex] * (Distance - Radius) + LocalSimSpaceOffset, FLinearColor::White);
+			const FVector Position = LocalSpaceLocation + AnimationPositions[ParticleIndex];
+			const FVector& Normal = AnimationNormals[ParticleIndex];
+			DrawLine(PDI, Position, Position - Normal * (Distance - Radius), FLinearColor::White);
 			if (Radius > 0.0f)
 			{
-				const FVector& Normal = AnimationNormals[ParticleIndex];
-				const FVector& Position = AnimationPositions[ParticleIndex];
-				auto DrawBackstop = [Radius, Distance, &Normal, &Position, PDI, this](const FVector& Axis, const FLinearColor& Color)
+				auto DrawBackstop = [Radius, Distance, &Normal, &Position, PDI](const FVector& Axis, const FLinearColor& Color)
 				{
 					const float ArcLength = 5.0f; // Arch length in cm
 					const float ArcAngle = ArcLength * 360.0f / (Radius * 2.0f * PI);
@@ -2093,7 +2217,7 @@ void ClothingSimulation::DebugDrawBackstops(USkeletalMeshComponent* /*OwnerCompo
 					const float MaxCosAngle = 0.99f;
 					if (FMath::Abs(FVector::DotProduct(Normal, Axis)) < MaxCosAngle)
 					{
-						DrawArc(PDI, Position - Normal * Distance + LocalSimSpaceOffset, Normal, FVector::CrossProduct(Axis, Normal).GetSafeNormal(), -ArcAngle / 2.0f, ArcAngle / 2.0f, Radius, Color);
+						DrawArc(PDI, Position - Normal * Distance, Normal, FVector::CrossProduct(Axis, Normal).GetSafeNormal(), -ArcAngle / 2.0f, ArcAngle / 2.0f, Radius, Color);
 					}
 				};
 				DrawBackstop(FVector::ForwardVector, FLinearColor::Blue);
@@ -2107,13 +2231,10 @@ void ClothingSimulation::DebugDrawBackstops(USkeletalMeshComponent* /*OwnerCompo
 void ClothingSimulation::DebugDrawMaxDistances(USkeletalMeshComponent* /*OwnerComponent*/, FPrimitiveDrawInterface* PDI) const
 {
 	const TPBDParticles<float, 3>& Particles = Evolution->Particles();
-	for (int32 i = 0; i < IndexToRangeMap.Num(); ++i)
+	for (int32 Index = 0; Index < Assets.Num(); ++Index)
 	{
-		const UClothingAssetCommon* const Asset = Assets[i];
-		if (Asset == nullptr)
-		{
-			continue;
-		}
+		const UClothingAssetCommon* const Asset = Assets[Index];
+		if (!Asset) { continue; }
 
 		// Get Maximum Distances
 		const UClothLODDataCommon* const AssetLodData = Asset->ClothLodData[0];
@@ -2125,19 +2246,20 @@ void ClothingSimulation::DebugDrawMaxDistances(USkeletalMeshComponent* /*OwnerCo
 			continue;
 		}
 		
-		for (uint32 ParticleIndex = IndexToRangeMap[i][0]; ParticleIndex < IndexToRangeMap[i][1]; ++ParticleIndex)
+		for (uint32 ParticleIndex = IndexToRangeMap[Index][0]; ParticleIndex < IndexToRangeMap[Index][1]; ++ParticleIndex)
 		{
-			const uint32 WeightMapIndex = ParticleIndex - IndexToRangeMap[i][0];
+			const uint32 WeightMapIndex = ParticleIndex - IndexToRangeMap[Index][0];
 			const float Distance = MaxDistances[WeightMapIndex];
+			const FVector Position = LocalSpaceLocation + AnimationPositions[ParticleIndex];
 			if (Particles.InvM(ParticleIndex) == 0.0f)
 			{
 #if WITH_EDITOR
-				DrawPoint(PDI, AnimationPositions[ParticleIndex] + LocalSimSpaceOffset, FLinearColor::Red, DebugClothMaterialVertex);
+				DrawPoint(PDI, Position, FLinearColor::Red, DebugClothMaterialVertex);
 #endif
 			}
 			else
 			{
-				DrawLine(PDI, AnimationPositions[ParticleIndex] + LocalSimSpaceOffset, AnimationPositions[ParticleIndex] + AnimationNormals[ParticleIndex] * Distance + LocalSimSpaceOffset, FLinearColor::White);
+				DrawLine(PDI, Position, Position + AnimationNormals[ParticleIndex] * Distance, FLinearColor::White);
 			}
 		}
 	}
@@ -2146,13 +2268,10 @@ void ClothingSimulation::DebugDrawMaxDistances(USkeletalMeshComponent* /*OwnerCo
 void ClothingSimulation::DebugDrawAnimDrive(USkeletalMeshComponent* /*OwnerComponent*/, FPrimitiveDrawInterface* PDI) const
 {
 	const TPBDParticles<float, 3>& Particles = Evolution->Particles();
-	for (int32 i = 0; i < IndexToRangeMap.Num(); ++i)
+	for (int32 Index = 0; Index < Assets.Num(); ++Index)
 	{
-		const UClothingAssetCommon* const Asset = Assets[i];
-		if (Asset == nullptr)
-		{
-			continue;
-		}
+		const UClothingAssetCommon* const Asset = Assets[Index];
+		if (!Asset) { continue; }
 
 		// Get Animdrive Multiplier
 		const UClothLODDataCommon* const AssetLodData = Asset->ClothLodData[0];
@@ -2164,11 +2283,11 @@ void ClothingSimulation::DebugDrawAnimDrive(USkeletalMeshComponent* /*OwnerCompo
 			continue;
 		}
 
-		for (uint32 ParticleIndex = IndexToRangeMap[i][0]; ParticleIndex < IndexToRangeMap[i][1]; ++ParticleIndex)
+		for (uint32 ParticleIndex = IndexToRangeMap[Index][0]; ParticleIndex < IndexToRangeMap[Index][1]; ++ParticleIndex)
 		{
-			const uint32 WeightMapIndex = ParticleIndex - IndexToRangeMap[i][0];
+			const uint32 WeightMapIndex = ParticleIndex - IndexToRangeMap[Index][0];
 			const float Multiplier = AnimDriveMultipliers[WeightMapIndex];
-			DrawLine(PDI, AnimationPositions[ParticleIndex] + LocalSimSpaceOffset, Particles.X(ParticleIndex) + LocalSimSpaceOffset, FLinearColor(FColor::Cyan) * Multiplier * AnimDriveSpringStiffness[i]);
+			DrawLine(PDI, AnimationPositions[ParticleIndex] + LocalSpaceLocation, Particles.X(ParticleIndex) + LocalSpaceLocation, FLinearColor(FColor::Cyan) * Multiplier * AnimDriveSpringStiffness[Index]);
 		}
 	}
 }
@@ -2177,28 +2296,26 @@ void ClothingSimulation::DebugDrawLongRangeConstraint(USkeletalMeshComponent* /*
 {
 	const TPBDParticles<float, 3>& Particles = Evolution->Particles();
 
-	for (int32 i = 0; i < IndexToRangeMap.Num(); ++i)
+	for (int32 Index = 0; Index < Assets.Num(); ++Index)
 	{
-		if (!LongRangeConstraints[i])
-		{
-			continue;
-		}
+		const UClothingAssetCommon* const Asset = Assets[Index];
+		if (!Asset || !LongRangeConstraints[Index]) { continue; }
 
-		const TArray<TArray<uint32>>& Constraints = LongRangeConstraints[i]->GetConstraints();
-		const TArray<float>& Dists = LongRangeConstraints[i]->GetDists();
+		const TArray<TArray<uint32>>& Constraints = LongRangeConstraints[Index]->GetConstraints();
+		const TArray<float>& Dists = LongRangeConstraints[Index]->GetDists();
 
-		for (int32 j = 0; j < Constraints.Num(); ++j)
+		for (int32 ConstraintIndex = 0; ConstraintIndex < Constraints.Num(); ++ConstraintIndex)
 		{
-			const TArray<uint32>& Path = Constraints[j];
-			const float RefDist = Dists[j];
+			const TArray<uint32>& Path = Constraints[ConstraintIndex];
+			const float RefDist = Dists[ConstraintIndex];
 			const float CurDist = TPBDLongRangeConstraintsBase<float, 3>::ComputeGeodesicDistance(Particles, Path);
 			const float Offset = CurDist - RefDist;
 
-			const TVector<float, 3> P0 = Particles.X(Path[0]) + LocalSimSpaceOffset;  // Kinematic particle
-			const TVector<float, 3> P1 = Particles.X(Path[Path.Num() - 1]) + LocalSimSpaceOffset;  // Target particle
+			const TVector<float, 3> P0 = Particles.X(Path[0]) + LocalSpaceLocation;  // Kinematic particle
+			const TVector<float, 3> P1 = Particles.X(Path[Path.Num() - 1]) + LocalSpaceLocation;  // Target particle
 
-			const TVector<float, 3> Direction = (Particles.X(Path[Path.Num() - 2]) - P1).GetSafeNormal();
-			const TVector<float, 3> P2 = P1 + Direction * Offset + LocalSimSpaceOffset;
+			const TVector<float, 3> Direction = (LocalSpaceLocation + Particles.X(Path[Path.Num() - 2]) - P1).GetSafeNormal();
+			const TVector<float, 3> P2 = P1 + Direction * Offset;
 
 			DrawLine(PDI, P0, P1, FLinearColor(FColor::Purple));
 			DrawLine(PDI, P1, P2, FLinearColor::Black);
@@ -2219,13 +2336,28 @@ void ClothingSimulation::DebugDrawWindDragForces(USkeletalMeshComponent* /*Owner
 		for (int32 ElementIndex = 0; ElementIndex < Elements.Num(); ++ElementIndex)
 		{
 			const TVector<int32, 3>& Element = Elements[ElementIndex];
-			const TVector<float, 3> Position = LocalSimSpaceOffset + (1.f / 3.f ) * (
+			const TVector<float, 3> Position = LocalSpaceLocation + (
 				Particles.X(Element[0]) +
 				Particles.X(Element[1]) +
-				Particles.X(Element[2]));
+				Particles.X(Element[2])) / 3.f;
 			const TVector<float, 3>& Force = Forces[ElementIndex];
 			DrawLine(PDI, Position, Position + Force, FColor::Green);
 		}
 	}
 }
+
+void ClothingSimulation::DebugDrawLocalSpace(USkeletalMeshComponent* /*OwnerComponent*/, FPrimitiveDrawInterface* PDI) const
+{
+	// Draw local space
+	DrawCoordinateSystem(PDI, FQuat::Identity, LocalSpaceLocation);
+
+	// Draw reference spaces
+	for (int32 Index = 0; Index < Assets.Num(); ++Index)
+	{
+		if (Assets[Index])
+		{
+			DrawCoordinateSystem(PDI, RootBoneWorldTransforms[Index].GetRotation(), RootBoneWorldTransforms[Index].GetLocation());
+		}
+	}
+}
 #endif  // #if WITH_EDITOR || CHAOS_DEBUG_DRAW
diff --git a/Engine/Plugins/Experimental/ChaosCloth/Source/Chaos/Public/ChaosCloth/ChaosClothConfig.h b/Engine/Plugins/Experimental/ChaosCloth/Source/Chaos/Public/ChaosCloth/ChaosClothConfig.h
--- a/Engine/Plugins/Experimental/ChaosCloth/Source/Chaos/Public/ChaosCloth/ChaosClothConfig.h
+++ b/Engine/Plugins/Experimental/ChaosCloth/Source/Chaos/Public/ChaosCloth/ChaosClothConfig.h
@@ -127,6 +127,14 @@ public:
 	// Enable continuous collision detection
 	UPROPERTY()
 	bool bUseContinuousCollisionDetection = false;
+
+	// The amount of linear velocities sent to the local cloth space from the reference bone (the closest bone to the root on which the cloth section has been skinned, or the root itself if none has)
+	UPROPERTY(EditAnywhere, Category = "Velocity Scales", meta = (UIMin = "0", UIMax = "1", ClampMin = "0", ClampMax = "1"))
+	FVector LinearVelocityScale = { 1.f, 1.f, 1.f };
+
+	// The amount of angular velocities sent to the local cloth space from the reference bone (the closest bone to the root on which the cloth section has been skinned, or the root itself if none has)
+	UPROPERTY(EditAnywhere, Category = "Velocity Scales", meta = (UIMin = "0", UIMax = "1", ClampMin = "0", ClampMax = "1"))
+	float AngularVelocityScale = 1.f;
 };
 
 /**
@@ -182,26 +190,20 @@ public:
 	UPROPERTY(EditAnywhere, Category = Simulation, meta = (InlineEditConditionToggle))
 	bool bUseGravityOverride = false;
 
-	// Scale factor applied to the world gravity when not using the gravity override
-	UPROPERTY(EditAnywhere, Category = Simulation, meta = (EditCondition = "!bUseGravityOverride"))
+	// Scale factor applied to the world gravity and also to the clothing simulation interactor gravity. Does not affect the gravity if set using the override below.
+	UPROPERTY(EditAnywhere, Category = Simulation)
 	float GravityScale = 1.f;
 
 	// The gravitational acceleration vector [cm/s^2]
 	UPROPERTY(EditAnywhere, Category = Simulation, meta = (EditCondition = "bUseGravityOverride"))
 	FVector Gravity = { 0.f, 0.f, -980.665f };
 
+	// Enable local space simulation to help with floating point precision errors if the character is far away form the world origin
+	UPROPERTY(EditAnywhere, Category = Simulation)
+	bool bUseLocalSpaceSimulation = false;
+
 	// Enable the XPBD constraints that resolve stiffness independently from the number of iterations
 	// Experimental, this feature might be removed without warning, not for production use
 	UPROPERTY(EditAnywhere, Category = "Experimental", meta = (DisplayName="Use XPBD Constraints (Experimental)"))
 	bool bUseXPBDConstraints = false;
-
-	// Enable local space simulation
-	UPROPERTY(EditAnywhere, Category = LocalSpaceSimulation, meta = (DisplayName = "Use Local space simulation"))
-	bool bUseLocalSpaceSimulation = false;
-
-	UPROPERTY(EditAnywhere, Category = LocalSpaceSimulation, meta = (DisplayName = "Component Linear Acc Scale"))
-	FVector ComponentLinearAccScale = { 1.f, 1.f, 1.f };
-
-	UPROPERTY()
-	FVector ComponentLinearAccClamp = { 4000.f, 4000.f, 4000.f };
 };
diff --git a/Engine/Plugins/Experimental/ChaosCloth/Source/Chaos/Public/ChaosCloth/ChaosClothingSimulation.h b/Engine/Plugins/Experimental/ChaosCloth/Source/Chaos/Public/ChaosCloth/ChaosClothingSimulation.h
--- a/Engine/Plugins/Experimental/ChaosCloth/Source/Chaos/Public/ChaosCloth/ChaosClothingSimulation.h
+++ b/Engine/Plugins/Experimental/ChaosCloth/Source/Chaos/Public/ChaosCloth/ChaosClothingSimulation.h
@@ -59,6 +59,7 @@ namespace Chaos
 		CHAOSCLOTH_API void DebugDrawAnimDrive(USkeletalMeshComponent* OwnerComponent = nullptr, FPrimitiveDrawInterface* PDI = nullptr) const;
 		CHAOSCLOTH_API void DebugDrawLongRangeConstraint(USkeletalMeshComponent* OwnerComponent = nullptr, FPrimitiveDrawInterface* PDI = nullptr) const;
 		CHAOSCLOTH_API void DebugDrawWindDragForces(USkeletalMeshComponent* OwnerComponent = nullptr, FPrimitiveDrawInterface* PDI = nullptr) const;
+		CHAOSCLOTH_API void DebugDrawLocalSpace(USkeletalMeshComponent* OwnerComponent = nullptr, FPrimitiveDrawInterface* PDI = nullptr) const;
 #endif  // #if WITH_EDITOR || CHAOS_DEBUG_DRAW
 
 	protected:
@@ -74,6 +75,7 @@ namespace Chaos
 		virtual void DestroyContext(IClothingSimulationContext* InContext) override { delete InContext; }
 		virtual void GetSimulationData(TMap<int32, FClothSimulData>& OutData, USkeletalMeshComponent* InOwnerComponent, USkinnedMeshComponent* InOverrideComponent) const override;
 
+		// Return bounds in local space (or in world space if InOwnerComponent is null).
 		virtual FBoxSphereBounds GetBounds(const USkeletalMeshComponent* InOwnerComponent) const override;
 
 		virtual void AddExternalCollisions(const FClothCollisionData& InData) override;
@@ -85,6 +87,7 @@ namespace Chaos
 #if CHAOS_DEBUG_DRAW
 		// Runtime only debug draw functions
 		void DebugDrawBounds() const;
+		void DebugDrawGravity() const;
 #endif  // #if CHAOS_DEBUG_DRAW
 
 		void UpdateSimulationFromSharedSimConfig();
@@ -100,15 +103,17 @@ namespace Chaos
 		void AddSelfCollisions(int32 InSimDataIndex);
 
 		// Extract all collisions (from the physics asset and from the legacy apex collisions)
-		void ExtractCollisions(const UClothingAssetCommon* Asset);
+		void ExtractCollisions(const UClothingAssetCommon* Asset, int32 InSimDataIndex);
 		// Extract the collisions from the physics asset referenced by the specified clothing asset
-		void ExtractPhysicsAssetCollisions(const UClothingAssetCommon* Asset);
+		void ExtractPhysicsAssetCollisions(const UClothingAssetCommon* Asset, int32 InSimDataIndex);
 		// Extract the legacy apex collisions from the specified clothing asset 
-		void ExtractLegacyAssetCollisions(const UClothingAssetCommon* Asset);
+		void ExtractLegacyAssetCollisions(const UClothingAssetCommon* Asset, int32 InSimDataIndex);
 		// Add collisions from a ClothCollisionData structure
-		void AddCollisions(const FClothCollisionData& ClothCollisionData, const TArray<int32>& UsedBoneIndices);
-		// Update the collision transforms using the specified context, also update old collision if bReinit is true
-		void UpdateCollisionTransforms(const ClothingSimulationContext& Context, bool bReinit);
+		void AddCollisions(const FClothCollisionData& ClothCollisionData, const TArray<int32>& UsedBoneIndices, int32 InSimDataIndex);
+		// Process the collision function for all collisions of the specified cloth
+		void ForAllCollisions(TFunction<void(TGeometryClothParticles<float, 3>&, uint32)> CollisionFunction, int32 SimDataIndex);
+		// Update the collision transforms using the specified context, also reset initial collision particle states if the number of collisions has changed
+		void UpdateCollisionTransforms(const ClothingSimulationContext& Context, int32 InSimDataIndex);
 		// Return the correct bone index based on the asset used bone index array
 		FORCEINLINE int32 GetMappedBoneIndex(const TArray<int32>& UsedBoneIndices, int32 BoneIndex)
 		{
@@ -139,7 +144,6 @@ namespace Chaos
 
 		// Sim Data
 		TArray<Chaos::TVector<uint32, 2>> IndexToRangeMap;
-		TArray<FTransform> RootBoneWorldTransforms;  // Used for teleportation
 
 		TArray<TUniquePtr<Chaos::TTriangleMesh<float>>> Meshes;
 		mutable TArray<TArray<Chaos::TVector<float, 3>>> FaceNormals;
@@ -147,6 +151,8 @@ namespace Chaos
 
 		TUniquePtr<Chaos::TPBDEvolution<float, 3>> Evolution;
 
+		TArray<Chaos::TVector<uint32, 2>> CollisionsRangeMap;  // Collisions first particle index and number of collisions for each simulated cloth
+		TArray<TArray<Chaos::TVector<uint32, 2>>> ExternalCollisionsRangeMaps;  // External collisions first particle index and number of collisions for each simulated cloth
 		uint32 ExternalCollisionsOffset;  // External collisions first particle index
 
 		float Time;
@@ -158,17 +164,11 @@ namespace Chaos
 
 		TArray<TSharedPtr<TPBDLongRangeConstraintsBase<float, 3>>> LongRangeConstraints;
 
-		// This is used to translate between world space and simulation space. Add this to simulation space coordinates to get world space coordinates
-		// The function of this is to help with floating point precision errors if the character is far away form the world origin
-		// and to decouple the simulation from character acceleration and velocities if it is desired
-		bool	bLocalSimSpaceEnabled;
-		FVector LocalSimSpaceOffset;  
-		FVector PrevLocalSimSpaceOffset;
-		FVector LocalSimSpaceVelocity;
-		FVector LocalSimSpaceCappedVelocity;
-		FVector PrevLocalSimSpaceVelocity;
-		FVector ComponentLinearAccScale;
-		FVector ComponentLinearAccClamp;
+		bool bUseLocalSpaceSimulation;  // The function of this is to help with floating point precision errors if the character is far away form the world origin
+		FVector LocalSpaceLocation;  // This is used to translate between world space and simulation space. Add this to simulation space coordinates to get world space coordinates
+		TArray<FTransform> RootBoneWorldTransforms;  // Reference bone transform used for teleportation and to decouple the simulation from reference bone's accelerations
+		TArray<FVector> LinearDeltaRatios;  // Linear ratios applied to the reference bone transforms
+		TArray<float> AngularDeltaRatios;  // Angular ratios factor applied to the reference bone transforms
 
 #if WITH_EDITOR
 		// Visualization material
diff --git a/Engine/Plugins/Experimental/ChaosClothEditor/Source/Chaos/Private/ChaosClothEditor/ChaosSimulationEditorExtender.cpp b/Engine/Plugins/Experimental/ChaosClothEditor/Source/Chaos/Private/ChaosClothEditor/ChaosSimulationEditorExtender.cpp
--- a/Engine/Plugins/Experimental/ChaosClothEditor/Source/Chaos/Private/ChaosClothEditor/ChaosSimulationEditorExtender.cpp
+++ b/Engine/Plugins/Experimental/ChaosClothEditor/Source/Chaos/Private/ChaosClothEditor/ChaosSimulationEditorExtender.cpp
@@ -56,7 +56,8 @@ const FVisualizationOption FVisualizationOption::OptionData[] =
 	FVisualizationOption(&Chaos::ClothingSimulation::DebugDrawMaxDistances        , LOCTEXT("ChaosVisName_MaxDistance"         , "Max Distances"              ), LOCTEXT("ChaosVisName_MaxDistance_ToolTip"         , "Draws the current max distances for the sim particles as a line along its normal"), true),
 	FVisualizationOption(&Chaos::ClothingSimulation::DebugDrawAnimDrive           , LOCTEXT("ChaosVisName_AnimDrive"           , "Anim Drive"                 ), LOCTEXT("ChaosVisName_AnimDrive_Tooltip"           , "Draws the current skinned reference mesh for the simulation which anim drive will attempt to reach if enabled")),
 	FVisualizationOption(&Chaos::ClothingSimulation::DebugDrawLongRangeConstraint , LOCTEXT("ChaosVisName_LongRangeConstraint" , "Long Range Constraint"      ), LOCTEXT("ChaosVisName_LongRangeConstraint_Tooltip" , "Draws the long range attachment constraint distances")),
-	FVisualizationOption(&Chaos::ClothingSimulation::DebugDrawWindDragForces      , LOCTEXT("ChaosVisName_WindDragForces"      , "Wind drag forces"           ), LOCTEXT("ChaosVisName_Wind_Tooltip"                , "Draws the Wind drag forces"))
+	FVisualizationOption(&Chaos::ClothingSimulation::DebugDrawWindDragForces      , LOCTEXT("ChaosVisName_WindDragForces"      , "Wind drag forces"           ), LOCTEXT("ChaosVisName_Wind_Tooltip"                , "Draws the Wind drag forces")),
+	FVisualizationOption(&Chaos::ClothingSimulation::DebugDrawLocalSpace          , LOCTEXT("ChaosVisName_LocalSpace"          , "Local space reference bone" ), LOCTEXT("ChaosVisName_LocalSpace_Tooltip"          , "Draws the local space reference bone")),
 };
 const uint32 FVisualizationOption::Count = sizeof(OptionData) / sizeof(FVisualizationOption);
 
diff --git a/Engine/Source/Runtime/Experimental/Chaos/Private/Chaos/PBDEvolution.cpp b/Engine/Source/Runtime/Experimental/Chaos/Private/Chaos/PBDEvolution.cpp
--- a/Engine/Source/Runtime/Experimental/Chaos/Private/Chaos/PBDEvolution.cpp
+++ b/Engine/Source/Runtime/Experimental/Chaos/Private/Chaos/PBDEvolution.cpp
@@ -42,6 +42,7 @@ TPBDEvolution<T, d>::TPBDEvolution(TPBDParticles<T, d>&& InParticles, TKinematic
 	, MTime(0)
 {
 	MCollisionParticles.AddArray(&MCollided);
+	MCollisionParticles.AddArray(&MCollisionParticleGroupIds);
 	MParticles.AddArray(&MParticleGroupIds);
 	MPerGroupDamping.Add(Damping);
 
@@ -58,13 +59,13 @@ TPBDEvolution<T, d>::TPBDEvolution(TPBDParticles<T, d>&& InParticles, TKinematic
 }
 
 template<class T, int d>
-void TPBDEvolution<T, d>::AddParticles(uint32 Num, uint32 GroupId)
+uint32 TPBDEvolution<T, d>::AddParticles(uint32 Num, uint32 GroupId)
 {
 	// Add new particles
 	const uint32 Offset = MParticles.Size();
 	MParticles.AddParticles(Num);
 
-	// Initialize the new particle's group id
+	// Initialize the new particles' group ids
 	for (uint32 i = Offset; i < MParticles.Size(); ++i)
 	{
 		MParticleGroupIds[i] = GroupId;
@@ -80,6 +81,22 @@ void TPBDEvolution<T, d>::AddParticles(uint32 Num, uint32 GroupId)
 			MPerGroupDamping[i] = MDamping;
 		}
 	}
+	return Offset;
+}
+
+template<class T, int d>
+uint32 TPBDEvolution<T, d>::AddCollisionParticles(uint32 Num, uint32 GroupId)
+{
+	// Add new particles
+	const uint32 Offset = MCollisionParticles.Size();
+	MCollisionParticles.AddParticles(Num);
+
+	// Initialize the new particles' group ids
+	for (uint32 i = Offset; i < MCollisionParticles.Size(); ++i)
+	{
+		MCollisionParticleGroupIds[i] = GroupId;
+	}
+	return Offset;
 }
 
 template<class T, int d>
@@ -90,7 +107,7 @@ void TPBDEvolution<T, d>::AdvanceOneTimeStep(const T Dt)
 	TPerParticleEulerStepVelocity<T, d> EulerStepVelocityRule;
 	TPerGroupDampVelocity<T, d> DampVelocityRule(MParticleGroupIds, MPerGroupDamping);
 	TPerParticlePBDEulerStep<T, d> EulerStepRule;
-	TPerParticlePBDCollisionConstraint<T, d, EGeometryParticlesSimType::Other> CollisionRule(MCollisionParticles, MCollided, MCollisionThickness, MCoefficientOfFriction);
+	TPerParticlePBDCollisionConstraint<T, d, EGeometryParticlesSimType::Other> CollisionRule(MCollisionParticles, MCollided, MParticleGroupIds, MCollisionParticleGroupIds, MCollisionThickness, MCoefficientOfFriction);
 
 	{
 		SCOPE_CYCLE_COUNTER(STAT_ChaosPBDVelocityDampUpdateState);
diff --git a/Engine/Source/Runtime/Experimental/Chaos/Public/Chaos/PBDEvolution.h b/Engine/Source/Runtime/Experimental/Chaos/Public/Chaos/PBDEvolution.h
--- a/Engine/Source/Runtime/Experimental/Chaos/Public/Chaos/PBDEvolution.h
+++ b/Engine/Source/Runtime/Experimental/Chaos/Public/Chaos/PBDEvolution.h
@@ -29,7 +29,10 @@ class CHAOS_API TPBDEvolution
 	void AddXPBDConstraintFunctions(TFunction<void()> InitConstraintFunction, TFunction<void(TPBDParticles<T, d>&, const T)> ConstraintFunction) { MInitConstraintRules.Add(InitConstraintFunction); MConstraintRules.Add(ConstraintFunction); }
 	void AddForceFunction(TFunction<void(TPBDParticles<T, d>&, const T, const int32)> ForceFunction) { MForceRules.Add(ForceFunction); }
 
-	void AddParticles(uint32 Num, uint32 GroupId = 0);
+	// Add particles and initialize group ids. Return the index of the first added particle.
+	uint32 AddParticles(uint32 Num, uint32 GroupId = 0);
+	// Add collision particles and initialize group ids. Return the index of the first added particle.
+	uint32 AddCollisionParticles(uint32 Num, uint32 GroupId = 0);
 
 	const TPBDParticles<T, d>& Particles() const { return MParticles; }
 	TPBDParticles<T, d>& Particles() { return MParticles; }
@@ -75,7 +78,8 @@ class CHAOS_API TPBDEvolution
 	TArray<TVector<int32, 3>> MCollisionTriangles;       // Used for self-collisions
 	TSet<TVector<int32, 2>> MDisabledCollisionElements;  // 
 	TArrayCollectionArray<bool> MCollided;
-	TArrayCollectionArray<uint32> MParticleGroupIds;  // Used for per group parameters
+	TArrayCollectionArray<uint32> MCollisionParticleGroupIds;  // Used for per group parameters for collision particles
+	TArrayCollectionArray<uint32> MParticleGroupIds;  // Used for per group parameters for particles
 	TArray<T> MPerGroupDamping;
 	int32 MNumIterations;
 	T MCollisionThickness;
diff --git a/Engine/Source/Runtime/Experimental/Chaos/Public/Chaos/PerParticleGravity.h b/Engine/Source/Runtime/Experimental/Chaos/Public/Chaos/PerParticleGravity.h
--- a/Engine/Source/Runtime/Experimental/Chaos/Public/Chaos/PerParticleGravity.h
+++ b/Engine/Source/Runtime/Experimental/Chaos/Public/Chaos/PerParticleGravity.h
@@ -58,9 +58,10 @@ namespace Chaos
 		}
 
 		void SetAcceleration(const TVector<T, d>& Acceleration)
-		{
-			MAcceleration = Acceleration;
-		}
+		{ MAcceleration = Acceleration; }
+
+		const TVector<T, d>& GetAcceleration() const
+		{ return MAcceleration; }
 
 	private:
 		TVector<T, d> MAcceleration;
diff --git a/Engine/Source/Runtime/Experimental/Chaos/Public/Chaos/PerParticlePBDCollisionConstraint.h b/Engine/Source/Runtime/Experimental/Chaos/Public/Chaos/PerParticlePBDCollisionConstraint.h
--- a/Engine/Source/Runtime/Experimental/Chaos/Public/Chaos/PerParticlePBDCollisionConstraint.h
+++ b/Engine/Source/Runtime/Experimental/Chaos/Public/Chaos/PerParticlePBDCollisionConstraint.h
@@ -18,8 +18,8 @@ class TPerParticlePBDCollisionConstraint : public TPerParticleRule<T, d>
 	};
 
   public:
-	TPerParticlePBDCollisionConstraint(const TKinematicGeometryParticlesImp<T, d, SimType>& InParticles, TArrayCollectionArray<bool>& Collided, const T Thickness = (T)0, const T Friction = (T)0)
-	    : bFastPositionBasedFriction(true), MParticles(InParticles), MCollided(Collided), MThickness(Thickness), MFriction(Friction) {}
+	TPerParticlePBDCollisionConstraint(const TKinematicGeometryParticlesImp<T, d, SimType>& InParticles, TArrayCollectionArray<bool>& Collided, TArrayCollectionArray<uint32>& DynamicGroupIds, TArrayCollectionArray<uint32>& KinematicGroupIds, const T Thickness = (T)0, const T Friction = (T)0)
+	    : bFastPositionBasedFriction(true), MParticles(InParticles), MCollided(Collided), MDynamicGroupIds(DynamicGroupIds), MKinematicGroupIds(KinematicGroupIds), MThickness(Thickness), MFriction(Friction) {}
 	virtual ~TPerParticlePBDCollisionConstraint() {}
 
 	inline void Apply(TPBDParticles<T, d>& InParticles, const T Dt, const int32 Index) const override //-V762
@@ -28,6 +28,7 @@ class TPerParticlePBDCollisionConstraint : public TPerParticleRule<T, d>
 			return;
 		for (uint32 i = 0; i < MParticles.Size(); ++i)
 		{
+			if (MDynamicGroupIds[Index] != MKinematicGroupIds[i]) { continue; }  // Bail out if the collision groups doesn't match the particle group id
 			TVector<T, d> Normal;
 			TRigidTransform<T, d> Frame(MParticles.X(i), MParticles.R(i));
 			T Phi = MParticles.Geometry(i)->PhiWithNormal(Frame.InverseTransformPosition(InParticles.P(Index)), Normal);
@@ -96,6 +97,8 @@ class TPerParticlePBDCollisionConstraint : public TPerParticleRule<T, d>
 	// TODO(mlentine): Need a bb hierarchy
 	const TKinematicGeometryParticlesImp<T, d, SimType>& MParticles;
 	TArrayCollectionArray<bool>& MCollided;
+	TArrayCollectionArray<uint32>& MDynamicGroupIds;
+	TArrayCollectionArray<uint32>& MKinematicGroupIds;
 	mutable TMap<int32, VelocityConstraint> MVelocityConstraints;
 	const T MThickness, MFriction;
 };
