Fix CIS on Android due to overzealous variable shadowing warning in lambdas.

diff --git a/Engine/Plugins/Experimental/ChaosCloth/Source/Chaos/Private/ChaosCloth/ChaosClothingSimulation.cpp b/Engine/Plugins/Experimental/ChaosCloth/Source/Chaos/Private/ChaosCloth/ChaosClothingSimulation.cpp
--- a/Engine/Plugins/Experimental/ChaosCloth/Source/Chaos/Private/ChaosCloth/ChaosClothingSimulation.cpp
+++ b/Engine/Plugins/Experimental/ChaosCloth/Source/Chaos/Private/ChaosCloth/ChaosClothingSimulation.cpp
@@ -768,12 +768,10 @@ void ClothingSimulation::UpdateCollisionTransforms(const ClothingSimulationConte
 {
 	SCOPE_CYCLE_COUNTER(STAT_ChaosClothUpdateCollisionTransforms);
 
-	TGeometryClothParticles<float, 3>& CollisionParticles = Evolution->CollisionParticles();
-
 	// Resize the transform arrays if collision have changed
 	const int32 PrevNumCollisions = OldCollisionTransforms.Num();
 	const int32 NumCollisions = BaseTransforms.Num();
-	check(NumCollisions == int32(CollisionParticles.Size()));  // BaseTransforms should always automatically grow with the number of collision particles (collection array)
+	check(NumCollisions == int32(Evolution->CollisionParticles().Size()));  // BaseTransforms should always automatically grow with the number of collision particles (collection array)
 
 	const bool bHasNumCollisionsChanged = (NumCollisions != PrevNumCollisions);
 	if (bHasNumCollisionsChanged)
@@ -1280,7 +1278,6 @@ void ClothingSimulation::Simulate(IClothingSimulationContext* InContext)
 
 		// Iterate all cloths
 		TPBDParticles<float, 3>& Particles = Evolution->Particles();
-		TGeometryClothParticles<float, 3>& CollisionParticles = Evolution->CollisionParticles();
 
 		for (int32 Index = 0; Index < Assets.Num(); ++Index)
 		{
