Chaos Cloth - Enable the SkeletalMeshComponent SetMaxDistanceScale BP function in Chaos Cloth.

Merging CL 12116735

diff --git a/Engine/Plugins/Experimental/ChaosCloth/Source/Chaos/Private/ChaosCloth/ChaosClothingSimulation.cpp b/Engine/Plugins/Experimental/ChaosCloth/Source/Chaos/Private/ChaosCloth/ChaosClothingSimulation.cpp
--- a/Engine/Plugins/Experimental/ChaosCloth/Source/Chaos/Private/ChaosCloth/ChaosClothingSimulation.cpp
+++ b/Engine/Plugins/Experimental/ChaosCloth/Source/Chaos/Private/ChaosCloth/ChaosClothingSimulation.cpp
@@ -157,6 +157,7 @@ void ClothingSimulation::Shutdown()
 {
 	Assets.Reset();
 	AnimDriveSpringStiffness.Reset();
+	MaxDistancesMultipliers.Reset();
 	ExternalCollisions.Reset();
 	OldCollisionTransforms.Reset();
 	CollisionTransforms.Reset();
@@ -211,6 +212,7 @@ void ClothingSimulation::CreateActor(USkeletalMeshComponent* InOwnerComponent, U
 		// TODO: Refactor all these arrays into a single cloth runtime asset structure
 		Assets.SetNumZeroed(NumAssets);
 		AnimDriveSpringStiffness.SetNumZeroed(NumAssets);
+		MaxDistancesMultipliers.SetNumZeroed(NumAssets);
 
 		Meshes.SetNum(NumAssets);
 		FaceNormals.SetNum(NumAssets);
@@ -687,10 +689,16 @@ void ClothingSimulation::AddConstraints(const UChaosClothConfig* ChaosClothSimCo
 	const FPointWeightMap& MaxDistances = PhysMesh.GetWeightMap(EChaosWeightMapTarget::MaxDistance);
 	if (MaxDistances.Num() > 0)
 	{
+		// Initialize the interactor's multiplier
+		MaxDistancesMultipliers[InSimDataIndex] = 1.f;
+
 		check(Mesh.GetNumElements() > 0);
 		Chaos::PBDSphericalConstraint<float, 3> SphericalContraint(Offset, MaxDistances.Num(), true, &AnimationPositions, &MaxDistances.Values);
-		Evolution->AddPBDConstraintFunction([SphericalContraint = MoveTemp(SphericalContraint)](TPBDParticles<float, 3>& InParticles, const float Dt)
+		Evolution->AddPBDConstraintFunction([
+			&SphereRadiiMultiplier = MaxDistancesMultipliers[InSimDataIndex],
+			SphericalContraint = MoveTemp(SphericalContraint)](TPBDParticles<float, 3>& InParticles, const float Dt) mutable
 		{
+			SphericalContraint.SetSphereRadiiMultiplier(SphereRadiiMultiplier);
 			SphericalContraint.Apply(InParticles, Dt);
 		});
 	}
@@ -1424,6 +1432,9 @@ void ClothingSimulation::Simulate(IClothingSimulationContext* InContext)
 					CollisionParticles.X(i) = OldCollisionTransforms[i].GetTranslation();
 				}, Index);
 			}
+
+			// Update max distance multiplier
+			MaxDistancesMultipliers[Index] = Context->MaxDistanceScale;
 		}
 	}
 
diff --git a/Engine/Plugins/Experimental/ChaosCloth/Source/Chaos/Public/ChaosCloth/ChaosClothingSimulation.h b/Engine/Plugins/Experimental/ChaosCloth/Source/Chaos/Public/ChaosCloth/ChaosClothingSimulation.h
--- a/Engine/Plugins/Experimental/ChaosCloth/Source/Chaos/Public/ChaosCloth/ChaosClothingSimulation.h
+++ b/Engine/Plugins/Experimental/ChaosCloth/Source/Chaos/Public/ChaosCloth/ChaosClothingSimulation.h
@@ -129,6 +129,7 @@ namespace Chaos
 		// These simulation parameters can be changed through blueprints
 		// They will only be updated when the simulation is not running (so are safe to use on any cloth thread)
 		TArray<float> AnimDriveSpringStiffness; // One for every Asset
+		TArray<float> MaxDistancesMultipliers;
 
 		// Collision Data
 		FClothCollisionData ExternalCollisions;  // External collisions
diff --git a/Engine/Source/Runtime/ClothingSystemRuntimeCommon/Private/ClothingSimulation.cpp b/Engine/Source/Runtime/ClothingSystemRuntimeCommon/Private/ClothingSimulation.cpp
--- a/Engine/Source/Runtime/ClothingSystemRuntimeCommon/Private/ClothingSimulation.cpp
+++ b/Engine/Source/Runtime/ClothingSystemRuntimeCommon/Private/ClothingSimulation.cpp
@@ -35,6 +35,7 @@ void FClothingSimulationContextCommon::Fill(const USkeletalMeshComponent* InComp
 	FillWindVelocity(InComponent);
 	FillDeltaSeconds(InDeltaSeconds, InMaxPhysicsDelta);
 	FillTeleportMode(InComponent, InDeltaSeconds, InMaxPhysicsDelta);
+	FillMaxDistanceScale(InComponent);
 }
 
 void FClothingSimulationContextCommon::FillBoneTransforms(const USkeletalMeshComponent* InComponent)
@@ -136,6 +137,11 @@ float FClothingSimulationContextCommon::SetWindFromComponent(const USkeletalMesh
 	return WindAdaption;
 }
 
+void FClothingSimulationContextCommon::FillMaxDistanceScale(const USkeletalMeshComponent* Component)
+{
+	MaxDistanceScale = Component->GetClothMaxDistanceScale();
+}
+
 //==============================================================================
 // FClothingSimulationCommon
 //==============================================================================
diff --git a/Engine/Source/Runtime/ClothingSystemRuntimeCommon/Public/ClothingSimulation.h b/Engine/Source/Runtime/ClothingSystemRuntimeCommon/Public/ClothingSimulation.h
--- a/Engine/Source/Runtime/ClothingSystemRuntimeCommon/Public/ClothingSimulation.h
+++ b/Engine/Source/Runtime/ClothingSystemRuntimeCommon/Public/ClothingSimulation.h
@@ -27,6 +27,7 @@ protected:
 	virtual void FillWindVelocity(const USkeletalMeshComponent* InComponent);
 	virtual void FillDeltaSeconds(float InDeltaSeconds, float InMaxPhysicsDelta);
 	virtual void FillTeleportMode(const USkeletalMeshComponent* InComponent, float InDeltaSeconds, float InMaxPhysicsDelta);
+	virtual void FillMaxDistanceScale(const USkeletalMeshComponent* Component);
 
 	// Set the wind velocity and return the wind adaptation if required
 	float SetWindFromComponent(const USkeletalMeshComponent* Component);
@@ -52,6 +53,9 @@ public:
 
 	// Whether and how we should teleport the simulation this tick
 	EClothingTeleportMode TeleportMode;
+
+	// Scale for the max distance constraints of the simulation mesh
+	float MaxDistanceScale;
 };
 
 // Base simulation to fill in common data for the base context
diff --git a/Engine/Source/Runtime/ClothingSystemRuntimeNv/Private/ClothingSimulationNv.cpp b/Engine/Source/Runtime/ClothingSystemRuntimeNv/Private/ClothingSimulationNv.cpp
--- a/Engine/Source/Runtime/ClothingSystemRuntimeNv/Private/ClothingSimulationNv.cpp
+++ b/Engine/Source/Runtime/ClothingSystemRuntimeNv/Private/ClothingSimulationNv.cpp
@@ -57,8 +57,6 @@ void FClothingSimulationContextNv::Fill(const USkeletalMeshComponent* InComponen
 	FClothingSimulationContextCommon::Fill(InComponent, InDeltaSeconds, InMaxPhysicsDelta);
 
 	PredictedLod = InComponent->PredictedLODLevel;
-
-	MaxDistanceScale = InComponent->GetClothMaxDistanceScale();
 }
 
 void FClothingSimulationContextNv::FillRefToLocals(const USkeletalMeshComponent* InComponent)
diff --git a/Engine/Source/Runtime/ClothingSystemRuntimeNv/Public/ClothingSimulationNv.h b/Engine/Source/Runtime/ClothingSystemRuntimeNv/Public/ClothingSimulationNv.h
--- a/Engine/Source/Runtime/ClothingSystemRuntimeNv/Public/ClothingSimulationNv.h
+++ b/Engine/Source/Runtime/ClothingSystemRuntimeNv/Public/ClothingSimulationNv.h
@@ -54,9 +54,6 @@ public:
 	// Wind adaption, a measure of how quickly to adapt to the wind speed
 	// when using the legacy wind calculation mode
 	float WindAdaption;
-
-	// Scale for the max distance constraints of the simulation mesh
-	float MaxDistanceScale;
 };
 
 // Scratch data for simulation to avoid allocations while processing, per actor data
diff --git a/Engine/Source/Runtime/Experimental/Chaos/Public/Chaos/PBDSphericalConstraint.h b/Engine/Source/Runtime/Experimental/Chaos/Public/Chaos/PBDSphericalConstraint.h
--- a/Engine/Source/Runtime/Experimental/Chaos/Public/Chaos/PBDSphericalConstraint.h
+++ b/Engine/Source/Runtime/Experimental/Chaos/Public/Chaos/PBDSphericalConstraint.h
@@ -28,6 +28,7 @@ public:
 		, SphereRadii(InSphereRadii)
 		, SphereOffsetDistances(InSphereOffsetDistances)
 		, SphereOffsetDirections(InSphereOffsetDirections)
+		, SphereRadiiMultiplier(1.f)
 	{
 		check(SphereRadii);
 		check(SpherePositions);
@@ -46,7 +47,7 @@ public:
 				return;
 			}
 			
-			const T Radius = (*SphereRadii)[Index];
+			const T Radius = (*SphereRadii)[Index] * SphereRadiiMultiplier;
 			TVector<T, d> Centre = (*SpherePositions)[Index + FirstParticleIndex];
 			if (SphereOffsetDistances && SphereOffsetDirections)
 			{
@@ -66,6 +67,11 @@ public:
 		});
 	}
 
+	inline void SetSphereRadiiMultiplier(const T InSphereRadiiMultiplier)
+	{
+		SphereRadiiMultiplier = InSphereRadiiMultiplier;
+	}
+
 private:	
 	const uint32 FirstParticleIndex; // TODO: Get rid of this member variable (this is specific to a temporary cloth instancing implementation)
 	const uint32 ParticleCount;
@@ -76,6 +82,8 @@ private:
 
 	const TArray<T>* const SphereOffsetDistances;  // Optional Sphere position offsets // start at index 0
 	const TArray<TVector<T, d>>* const SphereOffsetDirections; //	Optional Sphere offset directions (TBD: This is really implementation specific to backstop constraints, consider creating a specialized constraint instead).
+
+	T SphereRadiiMultiplier;
 };
 
 }
