Chaos Cloth - Add Subdivision parameter.
- Optimize edge constraints by removing all unused kinematic segments.
- Update damping and velocity scale migration functions.
- Update non obvious default config parameters and tooltips to try reduce user support queries.
- Debug draw kinematic particles in a different color.

diff --git a/Engine/Plugins/Experimental/ChaosCloth/Source/Chaos/Private/ChaosCloth/ChaosClothConfig.cpp b/Engine/Plugins/Experimental/ChaosCloth/Source/Chaos/Private/ChaosCloth/ChaosClothConfig.cpp
--- a/Engine/Plugins/Experimental/ChaosCloth/Source/Chaos/Private/ChaosCloth/ChaosClothConfig.cpp
+++ b/Engine/Plugins/Experimental/ChaosCloth/Source/Chaos/Private/ChaosCloth/ChaosClothConfig.cpp
@@ -64,12 +64,12 @@ void UChaosClothConfig::MigrateFrom(const FClothConfig_Legacy& ClothConfig)
 	DragCoefficient = (ClothConfig.WindMethod == EClothingWindMethod_Legacy::Accurate) ? ClothConfig.WindDragCoefficient: 0.07f;  // Only Accurate wind uses the WindDragCoefficient
 
 	const float Damping = (ClothConfig.Damping.X + ClothConfig.Damping.Y + ClothConfig.Damping.Z) / 3.f;
-	DampingCoefficient = FMath::Clamp(Damping * Damping * 0.95f, 0.f, 1.f);  // Nv Cloth seems to have a different damping formulation.
+	DampingCoefficient = FMath::Clamp(Damping * Damping * 0.7f, 0.f, 1.f);  // Nv Cloth seems to have a different damping formulation.
 
 	CollisionThickness = FMath::Clamp(ClothConfig.CollisionThickness, 0.f, 1000.f);
 
-	LinearVelocityScale = ClothConfig.LinearInertiaScale;
-	const FVector AngularInertiaScale = ClothConfig.AngularInertiaScale * ClothConfig.CentrifugalInertiaScale;
+	LinearVelocityScale = ClothConfig.LinearInertiaScale * 0.75f;
+	const FVector AngularInertiaScale = ClothConfig.AngularInertiaScale * ClothConfig.CentrifugalInertiaScale * 0.75f;
 	AngularVelocityScale = (AngularInertiaScale.X + AngularInertiaScale.Y + AngularInertiaScale.Z) / 3.f;
 }
 
diff --git a/Engine/Plugins/Experimental/ChaosCloth/Source/Chaos/Private/ChaosCloth/ChaosClothingSimulation.cpp b/Engine/Plugins/Experimental/ChaosCloth/Source/Chaos/Private/ChaosCloth/ChaosClothingSimulation.cpp
--- a/Engine/Plugins/Experimental/ChaosCloth/Source/Chaos/Private/ChaosCloth/ChaosClothingSimulation.cpp
+++ b/Engine/Plugins/Experimental/ChaosCloth/Source/Chaos/Private/ChaosCloth/ChaosClothingSimulation.cpp
@@ -89,10 +89,14 @@ ClothingSimulation::ClothingSimulation()
 	: ClothSharedSimConfig(nullptr)
 	, ExternalCollisionsOffset(0)
 	, bOverrideGravity(false)
+	, bUseConfigGravity(false)
+	, GravityScale(1.f)
 	, Gravity(ChaosClothingSimulationDefault::Gravity)
+	, ConfigGravity(ChaosClothingSimulationDefault::Gravity)
 	, WindVelocity(FVector::ZeroVector)
 	, bUseLocalSpaceSimulation(false)
 	, LocalSpaceLocation(FVector::ZeroVector)
+	, NumSubsteps(1)
 {
 #if WITH_EDITOR
 	DebugClothMaterial = LoadObject<UMaterial>(nullptr, TEXT("/Engine/EditorMaterials/Cloth/CameraLitDoubleSided.CameraLitDoubleSided"), nullptr, LOAD_None, nullptr);  // LOAD_EditorOnly
@@ -149,7 +153,7 @@ void ClothingSimulation::Initialize()
 			ParticlesInput.R(Index) = NewR;
 		});
 
-    Time = 0.f;
+    Time = Evolution->GetTime();
 	DeltaTime = 1.f / 30.f;  // Initialize filtered timestep at 30fps 
 }
 
@@ -272,7 +276,7 @@ void ClothingSimulation::CreateActor(USkeletalMeshComponent* InOwnerComponent, U
 	}
 	ComponentToLocalSpace.AddToTranslation(-LocalSpaceLocation);
 
-	// Init local cloth sim space & teleport tranform
+	// Init local cloth sim space & teleport transform
 	const FTransform RootBoneTransform = Context.BoneTransforms[Asset->ReferenceBoneIndex];
 	RootBoneWorldTransforms[InSimDataIndex] = RootBoneTransform * Context.ComponentToWorld;  // Velocity scale deltas are calculated in world space
 	LinearDeltaRatios[InSimDataIndex] = FVector::OneVector - ChaosClothSimConfig->LinearVelocityScale.BoundToBox(FVector::ZeroVector, FVector::OneVector);
@@ -372,7 +376,13 @@ void ClothingSimulation::UpdateSimulationFromSharedSimConfig()
 		// Update local space simulation switch
 		bUseLocalSpaceSimulation = ClothSharedSimConfig->bUseLocalSpaceSimulation;
 
+		// Update gravity related config values
+		ConfigGravity = ClothSharedSimConfig->Gravity;
+		GravityScale = ClothSharedSimConfig->GravityScale;
+		bUseConfigGravity = ClothSharedSimConfig->bUseGravityOverride;
+
 		// Now set all the common parameters on the simulation
+		NumSubsteps = ClothSharedSimConfig->SubdivisionCount;
 		Evolution->SetIterations(ClothSharedSimConfig->IterationCount);
 		Evolution->SetSelfCollisionThickness(ClothSharedSimConfig->SelfCollisionThickness);
 	}
@@ -518,10 +528,48 @@ void ClothingSimulation::AddConstraints(const UChaosClothConfig* ChaosClothSimCo
 		}
 		else
 		{
-			Evolution->AddPBDConstraintFunction([SpringConstraints = Chaos::TPBDSpringConstraints<float, 3>(Evolution->Particles(), SurfaceElements, ChaosClothSimConfig->EdgeStiffness)](TPBDParticles<float, 3>& InParticles, const float Dt)
+			TArray<TVector<int32, 3>> DynamicSurfaceElements;
+			TArray<TVector<int32, 2>> Attachments;
+			for (const TVector<int32, 3>& SurfaceElement : SurfaceElements)
 			{
-				SpringConstraints.Apply(InParticles, Dt);
-			});
+				const bool bIsKinematic0 = (Evolution->Particles().InvM(SurfaceElement[0]) == 0.f);
+				const bool bIsKinematic1 = (Evolution->Particles().InvM(SurfaceElement[1]) == 0.f);
+				const bool bIsKinematic2 = (Evolution->Particles().InvM(SurfaceElement[2]) == 0.f);
+				bool bIsAttachment = false;
+				if (bIsKinematic0 != bIsKinematic1)
+				{
+					Attachments.Emplace(SurfaceElement[0], SurfaceElement[1]);
+					bIsAttachment = true;
+				}
+				if (bIsKinematic1 != bIsKinematic2)
+				{
+					Attachments.Emplace(SurfaceElement[1], SurfaceElement[2]);
+					bIsAttachment = true;
+				}
+				if (bIsKinematic2 != bIsKinematic0)
+				{
+					Attachments.Emplace(SurfaceElement[2], SurfaceElement[0]);
+					bIsAttachment = true;
+				}
+				if (!bIsAttachment)
+				{
+					DynamicSurfaceElements.Add(SurfaceElement);
+				}
+			}
+			if (Attachments.Num())
+			{
+				Evolution->AddPBDConstraintFunction([AttachmentConstraints = Chaos::TPBDSpringConstraints<float, 3>(Evolution->Particles(), MoveTemp(Attachments), ChaosClothSimConfig->EdgeStiffness)](TPBDParticles<float, 3>& InParticles, const float Dt)
+				{
+					AttachmentConstraints.Apply(InParticles, Dt);
+				});
+			}
+			if (DynamicSurfaceElements.Num())
+			{
+				Evolution->AddPBDConstraintFunction([SpringConstraints = Chaos::TPBDSpringConstraints<float, 3>(Evolution->Particles(), DynamicSurfaceElements, ChaosClothSimConfig->EdgeStiffness)](TPBDParticles<float, 3>& InParticles, const float Dt)
+				{
+					SpringConstraints.Apply(InParticles, Dt);
+				});
+			}
 		}
 	}
 	if (ChaosClothSimConfig->BendingStiffness)
@@ -1246,12 +1294,11 @@ void ClothingSimulation::Simulate(IClothingSimulationContext* InContext)
 	static const float Decay = 0.1f;
 	DeltaTime = DeltaTime + (Context->DeltaSeconds - DeltaTime) * Decay;
 
-	// Set gravity, using the legacy priority: 1) config override, 2) game override, 3) world gravity
+	// Set gravity, using the legacy priority: 1) game override, 2) config override, 3) world gravity
 	Evolution->GetGravityForces().SetAcceleration(Chaos::TVector<float, 3>(
-		bOverrideGravity ? Gravity * ClothSharedSimConfig->GravityScale :
-		(ClothSharedSimConfig && ClothSharedSimConfig->bUseGravityOverride) ? ClothSharedSimConfig->Gravity :
-		ClothSharedSimConfig ? Context->WorldGravity * ClothSharedSimConfig->GravityScale :
-		Context->WorldGravity));
+		bOverrideGravity ? Gravity * GravityScale :
+		bUseConfigGravity ? ConfigGravity :  // Config gravity is not subject to scale
+		Context->WorldGravity * GravityScale));
 
 	// Set wind velocity, used by the velocity field lambda
 	WindVelocity = Context->WindVelocity * ChaosClothingSimulationDefault::WorldScale;  // Wind speed is set in m/s and need to be converted to cm/s
@@ -1383,7 +1430,7 @@ void ClothingSimulation::Simulate(IClothingSimulationContext* InContext)
 				if (DeltaAngle > PI) { DeltaAngle -= 2.f * PI; }
 				DeltaAngle *= AngularDeltaRatios[Index];
 				DeltaRotation = FQuat(Axis, DeltaAngle);
-				DeltaRotation.Normalize();  // ToMatrixNoScale does not like quaternions built straight from axis angles without being normilized (although they should have been already).
+				DeltaRotation.Normalize();  // ToMatrixNoScale does not like quaternions built straight from axis angles without being normalized (although they should have been already).
 
 				// Transform points back into the previous frame of reference before applying the adjusted deltas 
 				const FMatrix Matrix = (PrevRootBoneLocalTransform.Inverse() * FTransform(DeltaRotation, DeltaPosition) * PrevRootBoneLocalTransform).ToMatrixNoScale();
@@ -1433,8 +1480,14 @@ void ClothingSimulation::Simulate(IClothingSimulationContext* InContext)
 	}
 
 	// Advance Sim
-	Evolution->AdvanceOneTimeStep(DeltaTime);
-	Time += DeltaTime;
+	const float SubstepDeltaTime = DeltaTime / (float)NumSubsteps;
+	
+	for (int32 i = 0; i < NumSubsteps; ++i)
+	{
+		Evolution->AdvanceOneTimeStep(SubstepDeltaTime);
+	}
+
+	Time = Evolution->GetTime();
 	UE_LOG(LogChaosCloth, VeryVerbose, TEXT("DeltaTime: %.6f, FilteredDeltaTime: %.6f, Time = %.6f,  MaxPhysicsDelta = %.6f"), Context->DeltaSeconds, DeltaTime, Time, FClothingSimulationCommon::MaxPhysicsDelta);
 
 	// Debug draw
@@ -1974,6 +2027,9 @@ void ClothingSimulation::DebugDrawGravity() const
 
 void ClothingSimulation::DebugDrawPhysMeshWired(USkeletalMeshComponent* /*OwnerComponent*/, FPrimitiveDrawInterface* PDI) const
 {
+	static const FLinearColor DynamicColor = FColor::White;
+	static const FLinearColor KinematicColor = FColor::Purple;
+
 	const TPBDParticles<float, 3>& Particles = Evolution->Particles();
 
 	for (int32 Index = 0; Index < Assets.Num(); ++Index)
@@ -1990,9 +2046,13 @@ void ClothingSimulation::DebugDrawPhysMeshWired(USkeletalMeshComponent* /*OwnerC
 				const FVector Pos1 = LocalSpaceLocation + Particles.X(Element.Y);
 				const FVector Pos2 = LocalSpaceLocation + Particles.X(Element.Z);
 
-				DrawLine(PDI, Pos0, Pos1, FLinearColor::White);
-				DrawLine(PDI, Pos1, Pos2, FLinearColor::White);
-				DrawLine(PDI, Pos2, Pos0, FLinearColor::White);
+				const bool bIsKinematic0 = Particles.InvM(Element.X) == 0.f;
+				const bool bIsKinematic1 = Particles.InvM(Element.Y) == 0.f;
+				const bool bIsKinematic2 = Particles.InvM(Element.Z) == 0.f;
+
+				DrawLine(PDI, Pos0, Pos1, bIsKinematic0 && bIsKinematic1 ? KinematicColor : DynamicColor);
+				DrawLine(PDI, Pos1, Pos2, bIsKinematic1 && bIsKinematic2 ? KinematicColor : DynamicColor);
+				DrawLine(PDI, Pos2, Pos0, bIsKinematic2 && bIsKinematic0 ? KinematicColor : DynamicColor);
 			}
 		}
 	}
@@ -2128,7 +2188,7 @@ void ClothingSimulation::DebugDrawCollision(USkeletalMeshComponent* /*OwnerCompo
 					DrawCapsule(PDI, Object->GetObjectChecked<TCapsule<float>>(), Rotation, Position, Color);
 					break;
 
-				case ImplicitObjectType::Union:  // Union only used as collision tappered capsules
+				case ImplicitObjectType::Union:  // Union only used as collision tapered capsules
 					for (const TUniquePtr<FImplicitObject>& SubObjectPtr : Object->GetObjectChecked<FImplicitObjectUnion>().GetObjects())
 					{
 						if (const FImplicitObject* const SubObject = SubObjectPtr.Get())
@@ -2270,7 +2330,7 @@ void ClothingSimulation::DebugDrawAnimDrive(USkeletalMeshComponent* /*OwnerCompo
 		const UClothingAssetCommon* const Asset = Assets[Index];
 		if (!Asset) { continue; }
 
-		// Get Animdrive Multiplier
+		// Get anim drive multiplier
 		const UClothLODDataCommon* const AssetLodData = Asset->ClothLodData[0];
 		check(AssetLodData);
 		const FClothPhysicalMeshData& PhysMesh = AssetLodData->ClothPhysicalMeshData;
diff --git a/Engine/Plugins/Experimental/ChaosCloth/Source/Chaos/Public/ChaosCloth/ChaosClothConfig.h b/Engine/Plugins/Experimental/ChaosCloth/Source/Chaos/Public/ChaosCloth/ChaosClothConfig.h
--- a/Engine/Plugins/Experimental/ChaosCloth/Source/Chaos/Public/ChaosCloth/ChaosClothConfig.h
+++ b/Engine/Plugins/Experimental/ChaosCloth/Source/Chaos/Public/ChaosCloth/ChaosClothConfig.h
@@ -30,9 +30,9 @@ public:
 
 	/**
 	 * How cloth particle mass is determined
-	 * -	Uniform Mass: Every particle's mass will be set to the value specified in the UniformMass setting
-	 * -	Total Mass: The total mass is distributed equally over all the particles
-	 * -	Density: A constant mass density is used	 	 
+	 * -	Uniform Mass: Every particle's mass will be set to the value specified in the UniformMass setting. Mostly to avoid as it can cause some serious issues with irregular tessellations.
+	 * -	Total Mass: The total mass is distributed equally over all the particles. Useful when referencing a specific garment size and feel.
+	 * -	Density: A constant mass density is used. Density is usually the preferred way of setting mass since it allows matching real life materials' density values.
 	 */
 	UPROPERTY(EditAnywhere, Category = "Mass Config")
 	EClothMassMode MassMode = EClothMassMode::Density;
@@ -46,7 +46,7 @@ public:
 	float TotalMass = 0.5f;
 
 	/**
-	 * The value used when Mass Mode is set to Density 
+	 * The value used when Mass Mode is set to Density.
 	 * Melton Wool: 0.7
 	 * Heavy leather: 0.6
 	 * Polyurethane: 0.5
@@ -58,60 +58,63 @@ public:
 	UPROPERTY(EditAnywhere, Category = "Mass Config", meta = (UIMin = "0.001", UIMax = "1", ClampMin = "0"))
 	float Density = 0.35f;
 
-	// This is a lower bound to cloth particle masses
+	// This is a lower bound to cloth particle masses.
 	UPROPERTY(EditAnywhere, Category = "Mass Config")
 	float MinPerParticleMass = 0.0001f;	
 
-	// The Stiffness of the Edge constraints
+	// The Stiffness of the Edge constraints, only use lower than 1 values for very stretchy materials. Increase the iteration count for stiffer materials.
 	UPROPERTY(EditAnywhere, Category = Stiffness, meta = (UIMin = "0", UIMax = "1", ClampMin = "0", ClampMax = "1"))
-	float EdgeStiffness = 0.5f;
+	float EdgeStiffness = 1.f;
 
-	// The Stiffness of the bending constraints
+	// The Stiffness of the bending constraints. Increase the iteration count for stiffer materials.
 	UPROPERTY(EditAnywhere, Category = Stiffness, meta = (UIMin = "0", UIMax = "1", ClampMin = "0", ClampMax = "1"))
-	float BendingStiffness = 0.5f;
+	float BendingStiffness = 1.f;
 
-	// The stiffness of the area preservation constraints
+	// The stiffness of the area preservation constraints. Increase the iteration count for stiffer materials.
 	UPROPERTY(EditAnywhere, Category = Stiffness, meta = (UIMin = "0", UIMax = "1", ClampMin = "0", ClampMax = "1"))
-	float AreaStiffness = 0.5f;
+	float AreaStiffness = 1.f;
 
-	// The stiffness of the volume preservation constraints
+	// The stiffness of the volume preservation constraints.
 	UPROPERTY()
 	float VolumeStiffness = 0.f;
 
-	// The strain limiting stiffness of the long range attachment constraints (aka tether stiffness)
+	// The strain limiting stiffness of the long range attachment constraints (aka tether stiffness).
+	// The long range attachment connects each of the cloth particles to its closest fixed point with a spring constraint.
+	// This can be used to compensate for a lack of stretch resistance when the iterations count is kept low for performance reasons.
+	// Can lead to an unnatural pull string puppet like behavior. Use 0 to disable. 
 	UPROPERTY(EditAnywhere, Category = "Long Range Attachment", meta = (UIMin = "0", UIMax = "1", ClampMin = "0", ClampMax = "1"))
 	float StrainLimitingStiffness = 0.5f;
 
-	// The limit scale of the long range attachment constraints (aka tether limit)
-	UPROPERTY(EditAnywhere, Category = "Long Range Attachment", meta = (UIMin = "0.01", UIMax = "2", ClampMin = "0.01", ClampMax = "10"))
+	// The limit scale of the long range attachment constraints (aka tether limit).
+	UPROPERTY(EditAnywhere, Category = "Long Range Attachment", meta = (UIMin = "1.", UIMax = "1.1", ClampMin = "0.01", ClampMax = "10"))
 	float LimitScale = 1.f;
 
 	// Use geodesic instead of euclidean distance calculations in the long range attachment constraint,
-	// which is slower at setup but less prone to artifacts during simulation
+	// which is slower at setup but less prone to artifacts during simulation.
 	UPROPERTY(EditAnywhere, Category = "Long Range Attachment")
-	bool bUseGeodesicDistance = false;
+	bool bUseGeodesicDistance = true;
 
 	// The stiffness of the shape target constraints
 	UPROPERTY()
 	float ShapeTargetStiffness = 0.f;
 
-	// The radius of cloth points when considering collisions against collider shapes
+	// The radius of cloth points when considering collisions against collider shapes.
 	UPROPERTY(EditAnywhere, Category = "Cloth Dynamics", meta = (UIMin = "0", UIMax = "100", ClampMin = "0", ClampMax = "1000"))
 	float CollisionThickness = 1.0f;
 
-	// Friction coefficient for cloth - collider interaction
+	// Friction coefficient for cloth - collider interaction.
 	UPROPERTY(EditAnywhere, Category = "Cloth Dynamics", meta = (UIMin = "0", UIMax = "1", ClampMin = "0", ClampMax = "10"))
 	float FrictionCoefficient = 0.2f;
 
-	// The amount of damping applied to the cloth velocities
+	// The amount of damping applied to the cloth velocities.
 	UPROPERTY(EditAnywhere, Category = "Cloth Dynamics", meta = (UIMin = "0", UIMax = "1", ClampMin = "0", ClampMax = "1"))
 	float DampingCoefficient = 0.01f;
 
-	// The drag coefficient applying on each particle
+	// The drag coefficient applying on each particle.
 	UPROPERTY(EditAnywhere, Category = "Cloth Dynamics", meta = (UIMin = "0", UIMax = "1", ClampMin = "0", ClampMax = "10"))
-	float DragCoefficient = 0.07f;
+	float DragCoefficient = 0.1f;
 
-	// Default spring stiffness for anim drive if an anim drive is in use
+	// Default spring stiffness for anim drive if an anim drive is in use.
 	UPROPERTY(EditAnywhere, Category = Stiffness, meta = (UIMin = "0", UIMax = "1", ClampMin = "0", ClampMax = "1"))
 	float AnimDriveSpringStiffness = 0.001f;
 
@@ -119,29 +122,31 @@ public:
 	UPROPERTY(EditAnywhere, Category = "Cloth Enable Flags")
 	bool bUseBendingElements = false;
 
-	// Enable tetrahedral constraints
+	// Enable tetrahedral constraints.
 	UPROPERTY()
 	bool bUseTetrahedralConstraints = false;
 
-	// Enable thin shell volume constraints 
+	// Enable thin shell volume constraints.
 	UPROPERTY()
 	bool bUseThinShellVolumeConstraints = false;
 
-	// Enable self collision
+	// Enable self collision.
 	UPROPERTY(EditAnywhere, Category = "Cloth Enable Flags")
 	bool bUseSelfCollisions = false;
 
-	// Enable continuous collision detection
+	// Enable continuous collision detection.
 	UPROPERTY()
 	bool bUseContinuousCollisionDetection = false;
 
-	// The amount of linear velocities sent to the local cloth space from the reference bone (the closest bone to the root on which the cloth section has been skinned, or the root itself if none has)
+	// The amount of linear velocities sent to the local cloth space from the reference bone
+	// (the closest bone to the root on which the cloth section has been skinned, or the root itself if the cloth isn't skinned).
 	UPROPERTY(EditAnywhere, Category = "Velocity Scales", meta = (UIMin = "0", UIMax = "1", ClampMin = "0", ClampMax = "1"))
-	FVector LinearVelocityScale = { 1.f, 1.f, 1.f };
+	FVector LinearVelocityScale = { 0.75f, 0.75f, 0.75f };
 
-	// The amount of angular velocities sent to the local cloth space from the reference bone (the closest bone to the root on which the cloth section has been skinned, or the root itself if none has)
+	// The amount of angular velocities sent to the local cloth space from the reference bone
+	// (the closest bone to the root on which the cloth section has been skinned, or the root itself if the cloth isn't skinned).
 	UPROPERTY(EditAnywhere, Category = "Velocity Scales", meta = (UIMin = "0", UIMax = "1", ClampMin = "0", ClampMax = "1"))
-	float AngularVelocityScale = 1.f;
+	float AngularVelocityScale = 0.75f;
 };
 
 /**
@@ -172,20 +177,25 @@ public:
 	virtual void PostEditChangeChainProperty(FPropertyChangedChainEvent& ChainEvent) override;
 #endif
 
-	// The number of solver iterations
-	// This will increase the stiffness of all constraints but will increase the CPU cost
-	UPROPERTY(EditAnywhere, Category = Simulation, meta = (UIMin = "0", UIMax = "20", ClampMin = "0", ClampMax = "100"))
+	// The number of solver iterations.
+	// This will increase the stiffness of all constraints but will increase the CPU cost.
+	UPROPERTY(EditAnywhere, Category = Simulation, meta = (UIMin = "1", UIMax = "10", ClampMin = "0", ClampMax = "100"))
 	int32 IterationCount = 1;
 
+	// The number of solver substeps.
+	// This will increase the precision of the collision inputs and help with constraint resolutions but will increase the CPU cost.
+	UPROPERTY(EditAnywhere, Category = Simulation, meta = (UIMin = "1", UIMax = "10", ClampMin = "1", ClampMax = "100"))
+	int32 SubdivisionCount = 1;
+
 	// The radius of the spheres used in self collision 
 	UPROPERTY(EditAnywhere, Category = Collision, meta = (UIMin = "0", UIMax = "100", ClampMin = "0", ClampMax = "1000"))
 	float SelfCollisionThickness = 2.0f;
 
-	// The radius of cloth points when considering collisions against collider shapes
+	// The radius of cloth points when considering collisions against collider shapes.
 	UPROPERTY()
 	float CollisionThickness_DEPRECATED = 1.0f;
 
-	// Use shared config damping rather than per cloth damping
+	// Use shared config damping rather than per cloth damping.
 	UPROPERTY()
 	bool bUseDampingOverride_DEPRECATED = true;
 
@@ -193,7 +203,7 @@ public:
 	UPROPERTY()
 	float Damping_DEPRECATED = 0.01f;
 
-	// Use gravity value vs world gravity
+	// Use the config gravity value instead of world gravity.
 	UPROPERTY(EditAnywhere, Category = Simulation, meta = (InlineEditConditionToggle))
 	bool bUseGravityOverride = false;
 
@@ -207,10 +217,10 @@ public:
 
 	// Enable local space simulation to help with floating point precision errors if the character is far away form the world origin
 	UPROPERTY(EditAnywhere, Category = Simulation)
-	bool bUseLocalSpaceSimulation = false;
+	bool bUseLocalSpaceSimulation = true;
 
 	// Enable the XPBD constraints that resolve stiffness independently from the number of iterations
 	// Experimental, this feature might be removed without warning, not for production use
-	UPROPERTY(EditAnywhere, Category = "Experimental", meta = (DisplayName="Use XPBD Constraints (Experimental)"))
+	UPROPERTY()
 	bool bUseXPBDConstraints = false;
 };
diff --git a/Engine/Plugins/Experimental/ChaosCloth/Source/Chaos/Public/ChaosCloth/ChaosClothingSimulation.h b/Engine/Plugins/Experimental/ChaosCloth/Source/Chaos/Public/ChaosCloth/ChaosClothingSimulation.h
--- a/Engine/Plugins/Experimental/ChaosCloth/Source/Chaos/Public/ChaosCloth/ChaosClothingSimulation.h
+++ b/Engine/Plugins/Experimental/ChaosCloth/Source/Chaos/Public/ChaosCloth/ChaosClothingSimulation.h
@@ -158,9 +158,13 @@ namespace Chaos
 
 		float Time;
 		float DeltaTime;
+		int32 NumSubsteps;
 
 		bool bOverrideGravity;
+		bool bUseConfigGravity;
+		float GravityScale;
 		FVector Gravity;
+		FVector ConfigGravity;
 		FVector WindVelocity;
 
 		TArray<TSharedPtr<TPBDLongRangeConstraintsBase<float, 3>>> LongRangeConstraints;
