Chaos Cloth - Set the default number of iterations to 2 in the config to Just Work.

diff --git a/Engine/Plugins/Experimental/ChaosCloth/Source/Chaos/Public/ChaosCloth/ChaosClothConfig.h b/Engine/Plugins/Experimental/ChaosCloth/Source/Chaos/Public/ChaosCloth/ChaosClothConfig.h
--- a/Engine/Plugins/Experimental/ChaosCloth/Source/Chaos/Public/ChaosCloth/ChaosClothConfig.h
+++ b/Engine/Plugins/Experimental/ChaosCloth/Source/Chaos/Public/ChaosCloth/ChaosClothConfig.h
@@ -180,7 +180,7 @@ public:
 	// The number of solver iterations.
 	// This will increase the stiffness of all constraints but will increase the CPU cost.
 	UPROPERTY(EditAnywhere, Category = Simulation, meta = (UIMin = "1", UIMax = "10", ClampMin = "0", ClampMax = "100"))
-	int32 IterationCount = 1;
+	int32 IterationCount = 2;
 
 	// The number of solver substeps.
 	// This will increase the precision of the collision inputs and help with constraint resolutions but will increase the CPU cost.
