IWYU ChaosCloth

diff --git a/Engine/Plugins/Experimental/ChaosCloth/Source/Chaos/Private/ChaosCloth/ChaosClothingSimulation.cpp b/Engine/Plugins/Experimental/ChaosCloth/Source/Chaos/Private/ChaosCloth/ChaosClothingSimulation.cpp
--- a/Engine/Plugins/Experimental/ChaosCloth/Source/Chaos/Private/ChaosCloth/ChaosClothingSimulation.cpp
+++ b/Engine/Plugins/Experimental/ChaosCloth/Source/Chaos/Private/ChaosCloth/ChaosClothingSimulation.cpp
@@ -40,6 +40,9 @@
 #include "Chaos/XPBDSpringConstraints.h"
 #include "Chaos/XPBDAxialSpringConstraints.h"
 
+#include "Chaos/PBDEvolution.h"
+#include "Chaos/TriangleMesh.h"
+
 #if PHYSICS_INTERFACE_PHYSX && !PLATFORM_LUMIN && !PLATFORM_ANDROID
 #include "PhysXIncludes.h"
 #endif
diff --git a/Engine/Plugins/Experimental/ChaosCloth/Source/Chaos/Public/ChaosCloth/ChaosClothingSimulation.h b/Engine/Plugins/Experimental/ChaosCloth/Source/Chaos/Public/ChaosCloth/ChaosClothingSimulation.h
--- a/Engine/Plugins/Experimental/ChaosCloth/Source/Chaos/Public/ChaosCloth/ChaosClothingSimulation.h
+++ b/Engine/Plugins/Experimental/ChaosCloth/Source/Chaos/Public/ChaosCloth/ChaosClothingSimulation.h
@@ -4,15 +4,20 @@
 #include "ClothingSimulation.h"
 #include "ClothingAsset.h"
 #include "Chaos/ArrayCollectionArray.h"
-#include "Chaos/PBDEvolution.h"
 #include "Chaos/Transform.h"
-#include "Chaos/TriangleMesh.h"
 #include "Chaos/ChaosDebugDrawDeclares.h"
 #include "ChaosCloth/ChaosClothConfig.h"
 #include "Components/SkeletalMeshComponent.h"
 
 namespace Chaos
 {
+
+	template <typename T>
+	class TTriangleMesh;
+
+	template <typename T, int d>
+	class TPBDEvolution;
+
 	typedef FClothingSimulationContextCommon ClothingSimulationContext;
 	template<class T, int d> class TPBDLongRangeConstraintsBase;
 
