Chaos Cloth - Add particle counts and simulation time to the cloth editor's viewport HUD.

diff --git a/Engine/Plugins/Experimental/ChaosCloth/Source/Chaos/Private/ChaosCloth/ChaosClothingSimulation.cpp b/Engine/Plugins/Experimental/ChaosCloth/Source/Chaos/Private/ChaosCloth/ChaosClothingSimulation.cpp
--- a/Engine/Plugins/Experimental/ChaosCloth/Source/Chaos/Private/ChaosCloth/ChaosClothingSimulation.cpp
+++ b/Engine/Plugins/Experimental/ChaosCloth/Source/Chaos/Private/ChaosCloth/ChaosClothingSimulation.cpp
@@ -101,6 +101,8 @@ ClothingSimulation::ClothingSimulation()
 	, bUseLocalSpaceSimulation(false)
 	, LocalSpaceLocation(FVector::ZeroVector)
 {
+	ResetStats();
+
 #if WITH_EDITOR
 	DebugClothMaterial = LoadObject<UMaterial>(nullptr, TEXT("/Engine/EditorMaterials/Cloth/CameraLitDoubleSided.CameraLitDoubleSided"), nullptr, LOAD_None, nullptr);  // LOAD_EditorOnly
 	DebugClothMaterialVertex = LoadObject<UMaterial>(nullptr, TEXT("/Engine/EditorMaterials/WidgetVertexColorMaterial"), nullptr, LOAD_None, nullptr);  // LOAD_EditorOnly
@@ -186,6 +188,8 @@ void ClothingSimulation::Shutdown()
 	ExternalCollisionsOffset = 0;
 	ClothSharedSimConfig = nullptr;
 	LongRangeConstraints.Reset();
+
+	ResetStats();
 }
 
 void ClothingSimulation::DestroyActors()
@@ -198,10 +202,7 @@ void ClothingSimulation::CreateActor(USkeletalMeshComponent* InOwnerComponent, U
 {
 	UE_LOG(LogChaosCloth, Verbose, TEXT("Adding Cloth LOD asset to %s in sim slot %d"), InOwnerComponent->GetOwner() ? *InOwnerComponent->GetOwner()->GetName() : TEXT("None"), InSimDataIndex);
 
-	//Evolution->SetCCD(ChaosClothSimConfig->bUseContinuousCollisionDetection);
-	//Evolution->SetCCD(true); // ryan!!!
-
-	UClothingAssetCommon* Asset = Cast<UClothingAssetCommon>(InAsset);
+	UClothingAssetCommon* const Asset = Cast<UClothingAssetCommon>(InAsset);
 	const UChaosClothConfig* const ChaosClothSimConfig = Asset->GetClothConfig<UChaosClothConfig>();
 	if (!ChaosClothSimConfig)
 	{
@@ -351,6 +352,48 @@ void ClothingSimulation::CreateActor(USkeletalMeshComponent* InOwnerComponent, U
 
 	// Update collision transforms, including initial state for particles' X & R
 	UpdateCollisionTransforms(Context, InSimDataIndex);
+
+	// Update stats
+	UpdateStats(InSimDataIndex);
+}
+
+void ClothingSimulation::ResetStats()
+{
+#if WITH_EDITOR
+	NumCloths = 0;
+	NumKinemamicParticles = 0;
+	NumDynamicParticles = 0;
+	SimulationTime = 0.f;
+#endif  // #if WITH_EDITOR
+}
+
+void ClothingSimulation::UpdateStats(int32 InSimDataIndex)
+{
+#if WITH_EDITOR
+	const uint32 Offset = IndexToRangeMap[InSimDataIndex][0];
+	const uint32 Range = IndexToRangeMap[InSimDataIndex][1];
+	if (const int32 AddedParticleCount = Range - Offset)
+	{
+		++NumCloths;
+	}
+
+	TPBDParticles<float, 3>& Particles = Evolution->Particles();
+	int32 NumAddedKinenamicParticles = 0;
+	int32 NumAddedDynamicParticles = 0;
+	for (uint32 i = Offset; i < Range; ++i)
+	{
+		if (Particles.InvM(i) == 0.f)
+		{
+			++NumAddedKinenamicParticles;
+		}
+		else
+		{
+			++NumAddedDynamicParticles;
+		}
+	}
+	NumKinemamicParticles += NumAddedKinenamicParticles;
+	NumDynamicParticles += NumAddedDynamicParticles;
+#endif  // #if WITH_EDITOR
 }
 
 void ClothingSimulation::ExtractCollisions(const UClothingAssetCommon* Asset, int32 InSimDataIndex)
@@ -487,7 +530,7 @@ void ClothingSimulation::SetParticleMasses(const UChaosClothConfig* ChaosClothCo
 	const uint32 Range = IndexToRangeMap[InSimDataIndex][1];
 
 	check(Particles.Size() >= Range);
-	for (uint32 i = Offset; i < Range; i++)
+	for (uint32 i = Offset; i < Range; ++i)
 	{
 		Particles.M(i) = FMath::Max(Particles.M(i), ChaosClothConfig->MinPerParticleMass);
 		Particles.InvM(i) = MaxDistances.IsBelowThreshold(i - Offset) ? 0.f : 1.f / Particles.M(i);
@@ -1291,9 +1334,13 @@ void ClothingSimulation::Simulate(IClothingSimulationContext* InContext)
 		return;
 	}
 
+#if WITH_EDITOR
+	const double StartTime = FPlatformTime::Seconds();
+#endif 
+
 	// Filter delta time to smoothen time variations and prevent unwanted vibrations
-	static const float Decay = 0.1f;
-	DeltaTime = DeltaTime + (Context->DeltaSeconds - DeltaTime) * Decay;
+	static const float DeltaTimeDecay = 0.1f;
+	DeltaTime = DeltaTime + (Context->DeltaSeconds - DeltaTime) * DeltaTimeDecay;
 
 	// Set gravity, using the legacy priority: 1) game override, 2) config override, 3) world gravity
 	Evolution->GetGravityForces().SetAcceleration(Chaos::TVector<float, 3>(
@@ -1491,6 +1538,14 @@ void ClothingSimulation::Simulate(IClothingSimulationContext* InContext)
 	Time = Evolution->GetTime();
 	UE_LOG(LogChaosCloth, VeryVerbose, TEXT("DeltaTime: %.6f, FilteredDeltaTime: %.6f, Time = %.6f,  MaxPhysicsDelta = %.6f"), Context->DeltaSeconds, DeltaTime, Time, FClothingSimulationCommon::MaxPhysicsDelta);
 
+#if WITH_EDITOR
+	// Update simulation time in ms (and provide an instant average instead of the value in real-time)
+	const float PrevSimulationTime = SimulationTime;  // Copy the atomic to prevent a re-read
+	const float CurrSimulationTime = (float)((FPlatformTime::Seconds() - StartTime) * 1000.);
+	static const float SimulationTimeDecay = 0.03f; // 0.03 seems to provide a good rate of update for the instant average
+	SimulationTime = PrevSimulationTime ? PrevSimulationTime + (CurrSimulationTime - PrevSimulationTime) * SimulationTimeDecay : CurrSimulationTime;
+#endif  // #if WITH_EDITOR
+
 	// Debug draw
 #if CHAOS_DEBUG_DRAW
 	if (ChaosClothingSimulationConsoleVariables::CVarDebugDrawLocalSpace      .GetValueOnAnyThread()) { DebugDrawLocalSpace          (); }
@@ -1567,7 +1622,7 @@ FBoxSphereBounds ClothingSimulation::GetBounds(const USkeletalMeshComponent* InO
 	FBoxSphereBounds Bounds(EForceInit::ForceInit);
 
 	// Calculate simulation bounds (in world space)
-	uint32 NumCloths = 0;
+	uint32 NumBoundedCloths = 0;
 	const TPBDParticles<float, 3>& Particles = Evolution->Particles();
 
 	for (int32 Index = 0; Index < Assets.Num(); ++Index)
@@ -1593,11 +1648,11 @@ FBoxSphereBounds ClothingSimulation::GetBounds(const USkeletalMeshComponent* InO
 
 			// Update bounds with this cloth
 			const FBoxSphereBounds ClothBounds(BoundingBox.Center(), BoundingBox.Extents() * 0.5f, FMath::Sqrt(SquaredRadius));
-			Bounds = (NumCloths++ == 0) ? ClothBounds : Bounds + ClothBounds;
+			Bounds = (NumBoundedCloths++ == 0) ? ClothBounds : Bounds + ClothBounds;
 		}
 	}
 
-	if (!bUseLocalSpaceSimulation && NumCloths && InOwnerComponent)
+	if (!bUseLocalSpaceSimulation && NumBoundedCloths && InOwnerComponent)
 	{
 		// Retrieve the master component (unlike the one passed to the context, this could be a slave component)
 		const bool bIsUsingMaster = InOwnerComponent->MasterPoseComponent.IsValid();
@@ -1693,6 +1748,9 @@ void ClothingSimulation::RefreshClothConfig()
 	Evolution->ResetSelfCollision();
 	Evolution->ResetVelocityFields();
 
+	// Reset stats
+	ResetStats();
+
 	for (int32 SimDataIndex = 0; SimDataIndex < Assets.Num(); ++SimDataIndex)
 	{
 		if (const UClothingAssetCommon* const Asset = Assets[SimDataIndex])
@@ -1732,6 +1790,9 @@ void ClothingSimulation::RefreshClothConfig()
 				{
 					AddSelfCollisions(SimDataIndex);
 				}
+
+				// Update stats
+				UpdateStats(SimDataIndex);
 			}
 		}
 	}
diff --git a/Engine/Plugins/Experimental/ChaosCloth/Source/Chaos/Public/ChaosCloth/ChaosClothingSimulation.h b/Engine/Plugins/Experimental/ChaosCloth/Source/Chaos/Public/ChaosCloth/ChaosClothingSimulation.h
--- a/Engine/Plugins/Experimental/ChaosCloth/Source/Chaos/Public/ChaosCloth/ChaosClothingSimulation.h
+++ b/Engine/Plugins/Experimental/ChaosCloth/Source/Chaos/Public/ChaosCloth/ChaosClothingSimulation.h
@@ -47,6 +47,13 @@ namespace Chaos
 		virtual void AddReferencedObjects(FReferenceCollector& Collector) override;
 		// End of FGCObject interface
 
+		// IClothingSimulation interface
+		virtual int32 GetNumCloths() const { return NumCloths; }
+		virtual int32 GetNumKinematicParticles() const { return NumKinemamicParticles; }
+		virtual int32 GetNumDynamicParticles() const { return NumDynamicParticles; }
+		virtual float GetSimulationTime() const { return SimulationTime; }
+		// End of IClothingSimulation interface
+
 		// Editor only debug draw function
 		CHAOSCLOTH_API void DebugDrawPhysMeshShaded(USkeletalMeshComponent* OwnerComponent, FPrimitiveDrawInterface* PDI) const;
 #endif  // #if WITH_EDITOR
@@ -89,6 +96,9 @@ namespace Chaos
 		// End of IClothingSimulation interface
 
 	private:
+		void ResetStats();
+		void UpdateStats(int32 InSimDataIndex);
+
 #if CHAOS_DEBUG_DRAW
 		// Runtime only debug draw functions
 		void DebugDrawBounds() const;
@@ -181,6 +191,12 @@ namespace Chaos
 		TArray<float> AngularDeltaRatios;  // Angular ratios factor applied to the reference bone transforms
 
 #if WITH_EDITOR
+		// Stats
+		TAtomic<int32> NumCloths;
+		TAtomic<int32> NumKinemamicParticles;
+		TAtomic<int32> NumDynamicParticles;
+		TAtomic<float> SimulationTime;
+
 		// Visualization material
 		UMaterial* DebugClothMaterial;
 		UMaterial* DebugClothMaterialVertex;
diff --git a/Engine/Source/Editor/Persona/Private/AnimationEditorViewportClient.cpp b/Engine/Source/Editor/Persona/Private/AnimationEditorViewportClient.cpp
--- a/Engine/Source/Editor/Persona/Private/AnimationEditorViewportClient.cpp
+++ b/Engine/Source/Editor/Persona/Private/AnimationEditorViewportClient.cpp
@@ -968,6 +968,35 @@ FText FAnimationViewportClient::GetDisplayInfo(bool bDisplayAllInfo) const
 		}
 	}
 
+	if (const IClothingSimulation* const ClothingSimulation = PreviewMeshComponent->GetClothingSimulation())
+	{
+		// Cloth stats
+		if (const int32 NumActiveCloths = ClothingSimulation->GetNumCloths())
+		{
+			TextValue = ConcatenateLine(TextValue, FText::Format(LOCTEXT("NumActiveCloths", "Active Cloths: {0}"), NumActiveCloths));
+		}
+		if (const int32 NumKinematicParticles = ClothingSimulation->GetNumKinematicParticles())
+		{
+			TextValue = ConcatenateLine(TextValue, FText::Format(LOCTEXT("NumKinematicParticles", "Kinematic Particles: {0}"), NumKinematicParticles));
+		}
+		if (const int32 NumDynamicParticles = ClothingSimulation->GetNumDynamicParticles())
+		{
+			TextValue = ConcatenateLine(TextValue, FText::Format(LOCTEXT("NumDynamicParticles", "Dynamic Particles: {0}"), NumDynamicParticles));
+		}
+		if (const float SimulationTime = ClothingSimulation->GetSimulationTime())
+		{
+			FNumberFormattingOptions NumberFormatOptions;
+			NumberFormatOptions.AlwaysSign = false;
+			NumberFormatOptions.UseGrouping = false;
+			NumberFormatOptions.RoundingMode = ERoundingMode::HalfFromZero;
+			NumberFormatOptions.MinimumIntegralDigits = 1;
+			NumberFormatOptions.MaximumIntegralDigits = 6;
+			NumberFormatOptions.MinimumFractionalDigits = 2;
+			NumberFormatOptions.MaximumFractionalDigits = 2;
+			TextValue = ConcatenateLine(TextValue, FText::Format(LOCTEXT("SimulationTime", "Simulation Time: {0}ms"), FText::AsNumber(SimulationTime, &NumberFormatOptions)));
+		}
+	}
+
 	if (PreviewMeshComponent->GetSectionPreview() != INDEX_NONE)
 	{
 		// Notify the user if they are isolating a mesh section.
diff --git a/Engine/Source/Runtime/ClothingSystemRuntimeInterface/Public/ClothingSimulationInterface.h b/Engine/Source/Runtime/ClothingSystemRuntimeInterface/Public/ClothingSimulationInterface.h
--- a/Engine/Source/Runtime/ClothingSystemRuntimeInterface/Public/ClothingSimulationInterface.h
+++ b/Engine/Source/Runtime/ClothingSystemRuntimeInterface/Public/ClothingSimulationInterface.h
@@ -139,6 +139,16 @@ public:
 	 * This can be called on any thread, so derived implementations should make sure the gather is
 	 * safe
 	 */
-	virtual void GatherStats() const
-	{}
+	virtual void GatherStats() const {}
+
+#if WITH_EDITOR
+	/** Return the number of simulated cloths. Implementation must be thread safe. */
+	virtual int32 GetNumCloths() const { return 0;  }
+	/** Return the number of kinematic (fixed) particles. Implementation must be thread safe. */
+	virtual int32 GetNumKinematicParticles() const { return 0; }
+	/** Return the number of dynamic (simulated) particles. Implementation must be thread safe. */
+	virtual int32 GetNumDynamicParticles() const { return 0; }
+	/** Return the simulation time in ms. Implementation must be thread safe. */
+	virtual float GetSimulationTime() const { return 0.f; }
+#endif
 };
