Chaos Cloth - Fix cloth rendering during world composition.
- Also tie up a few other loose ends in GetSimulationData to prevent other similar issues.

diff --git a/Engine/Plugins/Experimental/ChaosCloth/Source/Chaos/Private/ChaosCloth/ChaosClothingSimulation.cpp b/Engine/Plugins/Experimental/ChaosCloth/Source/Chaos/Private/ChaosCloth/ChaosClothingSimulation.cpp
--- a/Engine/Plugins/Experimental/ChaosCloth/Source/Chaos/Private/ChaosCloth/ChaosClothingSimulation.cpp
+++ b/Engine/Plugins/Experimental/ChaosCloth/Source/Chaos/Private/ChaosCloth/ChaosClothingSimulation.cpp
@@ -1570,31 +1570,42 @@ void ClothingSimulation::GetSimulationData(
 	USkeletalMeshComponent* InOwnerComponent,
 	USkinnedMeshComponent* InOverrideComponent) const
 {
+	const int32 NumSimulatedCloth = IndexToRangeMap.Num();
+	if (OutData.Num() != NumSimulatedCloth)
+	{
+		OutData.Reset();
+	}
+
 	const FTransform& OwnerTransform = InOwnerComponent->GetComponentTransform();
-	for (int32 i = 0; i < IndexToRangeMap.Num(); ++i)
+	for (int32 i = 0; i < NumSimulatedCloth; ++i)
 	{
+		FClothSimulData& Data = OutData.FindOrAdd(i);
+
 		const TUniquePtr<Chaos::TTriangleMesh<float>>& Mesh = Meshes[i];
 		if (!Mesh)
+		{
+			Data.Reset();
 			continue;
+		}
 		Mesh->GetFaceNormals(FaceNormals[i], Evolution->Particles().X(), false);  // No need to add a point index offset here since that is baked into the triangles
 		Mesh->GetPointNormals(PointNormals[i], FaceNormals[i], /*bReturnEmptyOnError =*/ false, /*bFillAtStartIndex =*/ false);
 
-		FClothSimulData& Data = OutData.FindOrAdd(i);
-		Data.Reset();
-
 		const UClothingAssetCommon* const Asset = Assets[i];
 		if (!Asset)
+		{
+			Data.Reset();
 			continue;
+		}
 
 		const TArray<FTransform>& ComponentSpaceTransforms = InOverrideComponent ?
 			InOverrideComponent->GetComponentSpaceTransforms() :
 			InOwnerComponent->GetComponentSpaceTransforms();
-		if (!ComponentSpaceTransforms.IsValidIndex(Asset->ReferenceBoneIndex))
+		if (!ensure(ComponentSpaceTransforms.IsValidIndex(Asset->ReferenceBoneIndex)))
 		{
 			UE_LOG(LogSkeletalMesh, Warning,
 				TEXT("Failed to write back clothing simulation data for component '%s' as bone transforms are invalid."),
 				*InOwnerComponent->GetName());
-			check(false);
+			Data.Reset();
 			continue;
 		}
 
diff --git a/Engine/Source/Runtime/Engine/Private/SkeletalMeshComponentPhysics.cpp b/Engine/Source/Runtime/Engine/Private/SkeletalMeshComponentPhysics.cpp
--- a/Engine/Source/Runtime/Engine/Private/SkeletalMeshComponentPhysics.cpp
+++ b/Engine/Source/Runtime/Engine/Private/SkeletalMeshComponentPhysics.cpp
@@ -2402,22 +2402,28 @@ void USkeletalMeshComponent::RecreateClothingActors()
 			}
 		}
 
-		if(ClothingSimulation)
+		if (ClothingSimulation)
 		{
-			TArray<UClothingAssetBase*> AssetsInUse;
-			SkeletalMesh->GetClothingAssetsInUse(AssetsInUse);
-
-			const int32 NumMeshAssets = SkeletalMesh->MeshClothingAssets.Num();
-			for(int32 BaseAssetIndex = 0; BaseAssetIndex < NumMeshAssets; ++BaseAssetIndex)
+			// Only create cloth sim actors when the world is ready for it
+			const UWorld* const World = GetWorld();
+			if (World && World->bShouldSimulatePhysics && World->GetPhysicsScene())
 			{
-				UClothingAssetBase* Asset = SkeletalMesh->MeshClothingAssets[BaseAssetIndex];
-				if(!Asset || !AssetsInUse.Contains(Asset))
+				TArray<UClothingAssetBase*> AssetsInUse;
+				SkeletalMesh->GetClothingAssetsInUse(AssetsInUse);
+
+				const int32 NumMeshAssets = SkeletalMesh->MeshClothingAssets.Num();
+				for (int32 BaseAssetIndex = 0; BaseAssetIndex < NumMeshAssets; ++BaseAssetIndex)
 				{
-					continue;
+					UClothingAssetBase* const Asset = SkeletalMesh->MeshClothingAssets[BaseAssetIndex];
+					if (Asset && AssetsInUse.Contains(Asset))
+					{
+						ClothingSimulation->CreateActor(this, Asset, BaseAssetIndex);
+					}
 				}
-				ClothingSimulation->CreateActor(this, Asset, BaseAssetIndex);
+				ClothingSimulation->PostActorCreationInitialize();
 			}
-			ClothingSimulation->PostActorCreationInitialize();
+
+			// Retrieve the cloth sim data, or clear the data if the world isn't ready to sim
 			WritebackClothingSimulationData();
 		}
 	}
