Chaos Cloth - Add per cloth gravity/wind/self-collision.
- Update the Chaos Cloth config categories.
- Replace the velocity field point to triangle map access by an array.
- Add make PBDEvolution privately inherit an array collection to better deal with particle group settings.

diff --git a/Engine/Plugins/Experimental/ChaosCloth/Source/Chaos/Private/ChaosCloth/ChaosClothConfig.cpp b/Engine/Plugins/Experimental/ChaosCloth/Source/Chaos/Private/ChaosCloth/ChaosClothConfig.cpp
--- a/Engine/Plugins/Experimental/ChaosCloth/Source/Chaos/Private/ChaosCloth/ChaosClothConfig.cpp
+++ b/Engine/Plugins/Experimental/ChaosCloth/Source/Chaos/Private/ChaosCloth/ChaosClothConfig.cpp
@@ -67,10 +67,15 @@ void UChaosClothConfig::MigrateFrom(const FClothConfig_Legacy& ClothConfig)
 	DampingCoefficient = FMath::Clamp(Damping * Damping * 0.7f, 0.f, 1.f);  // Nv Cloth seems to have a different damping formulation.
 
 	CollisionThickness = FMath::Clamp(ClothConfig.CollisionThickness, 0.f, 1000.f);
+	SelfCollisionThickness = FMath::Clamp(ClothConfig.SelfCollisionRadius, 0.f, 1000.f);
 
 	LinearVelocityScale = ClothConfig.LinearInertiaScale * 0.75f;
 	const FVector AngularInertiaScale = ClothConfig.AngularInertiaScale * ClothConfig.CentrifugalInertiaScale * 0.75f;
 	AngularVelocityScale = (AngularInertiaScale.X + AngularInertiaScale.Y + AngularInertiaScale.Z) / 3.f;
+
+	bUseGravityOverride = ClothConfig.bUseGravityOverride;
+	GravityScale = ClothConfig.GravityScale;
+	Gravity = ClothConfig.GravityOverride;
 }
 
 void UChaosClothConfig::MigrateFrom(const UClothSharedConfigCommon* ClothSharedConfig)
@@ -87,6 +92,13 @@ void UChaosClothConfig::MigrateFrom(const UClothSharedConfigCommon* ClothSharedC
 			}
 			CollisionThickness = ChaosClothSharedSimConfig->CollisionThickness_DEPRECATED;
 		}
+		if (ChaosClothConfigCustomVersion < FChaosClothConfigCustomVersion::AddGravitySelfCollisionMigration)
+		{
+			SelfCollisionThickness = ChaosClothSharedSimConfig->SelfCollisionThickness_DEPRECATED;
+			bUseGravityOverride = ChaosClothSharedSimConfig->bUseGravityOverride_DEPRECATED;
+			GravityScale = ChaosClothSharedSimConfig->GravityScale_DEPRECATED;
+			Gravity = ChaosClothSharedSimConfig->Gravity_DEPRECATED;
+		}
 	}
 }
 
@@ -103,7 +115,7 @@ void UChaosClothConfig::PostLoad()
 
 	if (ChaosClothConfigCustomVersion < FChaosClothConfigCustomVersion::UpdateDragDefault)
 	{
-		DragCoefficient = 0.07f;  // Reset to a more appropirate default for chaos cloth assets saved before this custom version
+		DragCoefficient = 0.07f;  // Reset to a more appropriate default for chaos cloth assets saved before this custom version
 	}
 }
 
@@ -117,14 +129,6 @@ void UChaosClothSharedSimConfig::MigrateFrom(const FClothConfig_Legacy& ClothCon
 {
 	IterationCount = FMath::Clamp(int32(ClothConfig.SolverFrequency / 60.f), 1, 100);
 
-	SelfCollisionThickness = FMath::Clamp(ClothConfig.SelfCollisionRadius, 0.f, 1000.f);
-
-	bUseGravityOverride = ClothConfig.bUseGravityOverride;
-
-	GravityScale = ClothConfig.GravityScale;
-
-	Gravity = ClothConfig.GravityOverride;
-
 	bUseDampingOverride_DEPRECATED = false;  // Damping is migrated to per cloth configs
 }
 
@@ -141,7 +145,7 @@ void UChaosClothSharedSimConfig::PostLoad()
 
 	if (ChaosClothSharedConfigCustomVersion < FChaosClothSharedConfigCustomVersion::AddGravityOverride)
 	{
-		bUseGravityOverride = true;  // Default gravity override would otherwise disable the currently set gravity on older versions
+		bUseGravityOverride_DEPRECATED = true;  // Default gravity override would otherwise disable the currently set gravity on older versions
 	}
 }
 
diff --git a/Engine/Plugins/Experimental/ChaosCloth/Source/Chaos/Private/ChaosCloth/ChaosClothConfigCustomVersion.h b/Engine/Plugins/Experimental/ChaosCloth/Source/Chaos/Private/ChaosCloth/ChaosClothConfigCustomVersion.h
--- a/Engine/Plugins/Experimental/ChaosCloth/Source/Chaos/Private/ChaosCloth/ChaosClothConfigCustomVersion.h
+++ b/Engine/Plugins/Experimental/ChaosCloth/Source/Chaos/Private/ChaosCloth/ChaosClothConfigCustomVersion.h
@@ -16,6 +16,8 @@ struct FChaosClothConfigCustomVersion
 		UpdateDragDefault = 1,
 		// Added damping and collision thickness per cloth
 		AddDampingThicknessMigration = 2,
+		// Added gravity and self collision thickness per cloth
+		AddGravitySelfCollisionMigration = 3,
 
 		// -----<new versions can be added above this line>-------------------------------------------------
 		VersionPlusOne,
diff --git a/Engine/Plugins/Experimental/ChaosCloth/Source/Chaos/Private/ChaosCloth/ChaosClothingSimulation.cpp b/Engine/Plugins/Experimental/ChaosCloth/Source/Chaos/Private/ChaosCloth/ChaosClothingSimulation.cpp
--- a/Engine/Plugins/Experimental/ChaosCloth/Source/Chaos/Private/ChaosCloth/ChaosClothingSimulation.cpp
+++ b/Engine/Plugins/Experimental/ChaosCloth/Source/Chaos/Private/ChaosCloth/ChaosClothingSimulation.cpp
@@ -9,12 +9,14 @@
 #include "Components/SkeletalMeshComponent.h"
 #include "Materials/Material.h"
 #include "PhysicsEngine/PhysicsAsset.h"
+#include "Chaos/TriangleMesh.h"
 #include "Chaos/Box.h"
 #include "Chaos/Capsule.h"
 #include "Chaos/Cylinder.h"
 #include "Chaos/ImplicitObjectIntersection.h"
 #include "Chaos/ImplicitObjectUnion.h"
 #include "Chaos/Levelset.h"
+#include "Chaos/PBDEvolution.h"
 #include "Chaos/PBDAnimDriveConstraint.h"
 #include "Chaos/PBDAxialSpringConstraints.h"
 #include "Chaos/PBDBendingConstraints.h"
@@ -40,9 +42,6 @@
 #include "Chaos/XPBDSpringConstraints.h"
 #include "Chaos/XPBDAxialSpringConstraints.h"
 
-#include "Chaos/PBDEvolution.h"
-#include "Chaos/TriangleMesh.h"
-
 #if PHYSICS_INTERFACE_PHYSX && !PLATFORM_LUMIN && !PLATFORM_ANDROID
 #include "PhysXIncludes.h"
 #endif
@@ -98,12 +97,8 @@ ClothingSimulation::ClothingSimulation()
 	: ClothSharedSimConfig(nullptr)
 	, ExternalCollisionsOffset(0)
 	, NumSubsteps(1)
-	, bOverrideGravity(false)
-	, bUseConfigGravity(false)
-	, GravityScale(1.f)
-	, Gravity(ChaosClothingSimulationDefault::Gravity)
-	, ConfigGravity(ChaosClothingSimulationDefault::Gravity)
-	, WindVelocity(FVector::ZeroVector)
+	, bUseGravityOverride(false)
+	, GravityOverride(ChaosClothingSimulationDefault::Gravity)
 	, bUseLocalSpaceSimulation(false)
 	, LocalSpaceLocation(FVector::ZeroVector)
 {
@@ -134,7 +129,6 @@ void ClothingSimulation::Initialize()
 			ChaosClothingSimulationDefault::DampingCoefficient));
     Evolution->CollisionParticles().AddArray(&BoneIndices);
 	Evolution->CollisionParticles().AddArray(&BaseTransforms);
-    Evolution->GetGravityForces().SetAcceleration(Gravity);
 
     Evolution->SetKinematicUpdateFunction(
 		[this](Chaos::TPBDParticles<float, 3>& ParticlesInput, const float Dt, const float LocalTime, const int32 Index)
@@ -194,6 +188,10 @@ void ClothingSimulation::Shutdown()
 	ExternalCollisionsOffset = 0;
 	ClothSharedSimConfig = nullptr;
 	LongRangeConstraints.Reset();
+	bUseGravityConfigs.Reset();
+	GravityScales.Reset();
+	GravityConfigs.Reset();
+	WindVelocities.Reset();
 
 	ResetStats();
 }
@@ -241,6 +239,11 @@ void ClothingSimulation::CreateActor(USkeletalMeshComponent* InOwnerComponent, U
 		AngularDeltaRatios.SetNum(NumAssets);
 
 		CollisionsRangeMap.SetNumZeroed(NumAssets);
+
+		GravityScales.SetNum(NumAssets);
+		bUseGravityConfigs.SetNum(NumAssets);
+		GravityConfigs.SetNum(NumAssets);
+		WindVelocities.SetNum(NumAssets);
 	}
 	Assets[InSimDataIndex] = Asset;
 
@@ -319,20 +322,26 @@ void ClothingSimulation::CreateActor(USkeletalMeshComponent* InOwnerComponent, U
 	Evolution->SetCoefficientOfFriction(ChaosClothSimConfig->FrictionCoefficient, InSimDataIndex);
 
 	// Add velocity field
-	auto GetVelocity = [this](const TVector<float, 3>&)->TVector<float, 3>
+	auto GetVelocity = [this, InSimDataIndex](const TVector<float, 3>&)->TVector<float, 3>
 	{
-		return WindVelocity;
+		return WindVelocities[InSimDataIndex];
 	};
-	Evolution->GetVelocityFields().Emplace(
+	Evolution->GetVelocityField(InSimDataIndex) = MakeUnique<TVelocityField<float, 3>>(
 		*Meshes[InSimDataIndex],
 		GetVelocity,
 		/*bInIsUniform =*/ true,
 		ChaosClothSimConfig->DragCoefficient);
 
+	// Update gravity related config values
+	GravityScales[InSimDataIndex] = ChaosClothSimConfig->GravityScale;
+	bUseGravityConfigs[InSimDataIndex] = ChaosClothSimConfig->bUseGravityOverride;
+	GravityConfigs[InSimDataIndex] = ChaosClothSimConfig->Gravity;
+
 	// Add Self Collisions
 	if (ChaosClothSimConfig->bUseSelfCollisions)
 	{
 		AddSelfCollisions(InSimDataIndex);
+		Evolution->SetSelfCollisionThickness(ChaosClothSimConfig->SelfCollisionThickness, InSimDataIndex);
 	}
 
 	// Warn about legacy apex collisions
@@ -428,15 +437,9 @@ void ClothingSimulation::UpdateSimulationFromSharedSimConfig()
 		// Update local space simulation switch
 		bUseLocalSpaceSimulation = ClothSharedSimConfig->bUseLocalSpaceSimulation;
 
-		// Update gravity related config values
-		ConfigGravity = ClothSharedSimConfig->Gravity;
-		GravityScale = ClothSharedSimConfig->GravityScale;
-		bUseConfigGravity = ClothSharedSimConfig->bUseGravityOverride;
-
 		// Now set all the common parameters on the simulation
 		NumSubsteps = ClothSharedSimConfig->SubdivisionCount;
 		Evolution->SetIterations(ClothSharedSimConfig->IterationCount);
-		Evolution->SetSelfCollisionThickness(ClothSharedSimConfig->SelfCollisionThickness);
 	}
 }
 
@@ -1348,15 +1351,6 @@ void ClothingSimulation::Simulate(IClothingSimulationContext* InContext)
 	static const float DeltaTimeDecay = 0.1f;
 	DeltaTime = DeltaTime + (Context->DeltaSeconds - DeltaTime) * DeltaTimeDecay;
 
-	// Set gravity, using the legacy priority: 1) game override, 2) config override, 3) world gravity
-	Evolution->GetGravityForces().SetAcceleration(Chaos::TVector<float, 3>(
-		bOverrideGravity ? Gravity * GravityScale :
-		bUseConfigGravity ? ConfigGravity :  // Config gravity is not subject to scale
-		Context->WorldGravity * GravityScale));
-
-	// Set wind velocity, used by the velocity field lambda
-	WindVelocity = Context->WindVelocity * ChaosClothingSimulationDefault::WorldScale;  // Wind speed is set in m/s and need to be converted to cm/s
-
 	// Check teleport modes
 	const bool bTeleport = (Context->TeleportMode > EClothingTeleportMode::None);
 	const bool bTeleportAndReset = (Context->TeleportMode == EClothingTeleportMode::TeleportAndReset);
@@ -1387,6 +1381,18 @@ void ClothingSimulation::Simulate(IClothingSimulationContext* InContext)
 			const UClothingAssetCommon* const Asset = Assets[Index];
 			if (!Asset) { continue; }
 
+			// Set wind velocity, used by the velocity field lambda
+			WindVelocities[Index] = Context->WindVelocity * ChaosClothingSimulationDefault::WorldScale;  // Wind speed is set in m/s and need to be converted to cm/s
+
+			// Set gravity, using the legacy priority: 1) game override, 2) config override, 3) world gravity
+			const TVector<float, 3> Gravity = TVector<float, 3>(
+				bUseGravityOverride ? GravityOverride * GravityScales[Index] :
+				bUseGravityConfigs[Index] ? GravityConfigs[Index] :  // Config gravity is not subject to scale
+				Context->WorldGravity * GravityScales[Index]);
+
+			// Update gravity
+			Evolution->GetGravityForces(Index).SetAcceleration(Gravity);
+
 			const uint32 Offset = IndexToRangeMap[Index][0];
 			const uint32 Range = IndexToRangeMap[Index][1];
 
@@ -1763,7 +1769,6 @@ void ClothingSimulation::RefreshClothConfig()
 
 	Evolution->ResetConstraintRules();
 	Evolution->ResetSelfCollision();
-	Evolution->ResetVelocityFields();
 
 	// Reset stats
 	ResetStats();
@@ -1792,20 +1797,26 @@ void ClothingSimulation::RefreshClothConfig()
 				Evolution->SetCoefficientOfFriction(ChaosClothConfig->FrictionCoefficient, SimDataIndex);
 
 				// Add Velocity field
-				auto GetVelocity = [this](const TVector<float, 3>&)->TVector<float, 3>
+				auto GetVelocity = [this, SimDataIndex](const TVector<float, 3>&)->TVector<float, 3>
 				{
-					return WindVelocity;
+					return WindVelocities[SimDataIndex];
 				};
-				Evolution->GetVelocityFields().Emplace(
+				Evolution->GetVelocityField(SimDataIndex) = MakeUnique<TVelocityField<float, 3>>(
 					*Meshes[SimDataIndex],
 					GetVelocity,
 					/*bInIsUniform =*/ true,
 					ChaosClothConfig->DragCoefficient);
 
+				// Update gravity related config values
+				GravityScales[SimDataIndex] = ChaosClothConfig->GravityScale;
+				bUseGravityConfigs[SimDataIndex] = ChaosClothConfig->bUseGravityOverride;
+				GravityConfigs[SimDataIndex] = ChaosClothConfig->Gravity;
+
 				// Add Self Collisions
 				if (ChaosClothConfig->bUseSelfCollisions)
 				{
 					AddSelfCollisions(SimDataIndex);
+					Evolution->SetSelfCollisionThickness(ChaosClothConfig->SelfCollisionThickness, SimDataIndex);
 				}
 
 				// Update stats
@@ -1846,13 +1857,13 @@ void ClothingSimulation::SetAnimDriveSpringStiffness(float InStiffness)
 
 void ClothingSimulation::SetGravityOverride(const FVector& InGravityOverride)
 {
-	bOverrideGravity = true;
-	Gravity = InGravityOverride;
+	bUseGravityOverride = true;
+	GravityOverride = InGravityOverride;
 }
 
 void ClothingSimulation::DisableGravityOverride()
 {
-	bOverrideGravity = false;
+	bUseGravityOverride = false;
 }
 
 #if WITH_EDITOR
@@ -2536,12 +2547,16 @@ void ClothingSimulation::DebugDrawLongRangeConstraint(USkeletalMeshComponent* /*
 void ClothingSimulation::DebugDrawWindDragForces(USkeletalMeshComponent* /*OwnerComponent*/, FPrimitiveDrawInterface* PDI) const
 {
 	const TPBDParticles<float, 3>& Particles = Evolution->Particles();
-	const TArray<TVelocityField<float, 3>>& VelocityFields = Evolution->GetVelocityFields();
 
-	for (const TVelocityField<float, 3>& VelocityField : VelocityFields)
+	for (int32 Index = 0; Index < Assets.Num(); ++Index)
 	{
-		const TArray<TVector<int32, 3>>& Elements = VelocityField.GetElements();
-		const TArray<TVector<float, 3>>& Forces = VelocityField.GetForces();
+		const UClothingAssetCommon* const Asset = Assets[Index];
+		if (!Asset) { continue; }
+
+		const TUniquePtr<TVelocityField<float, 3>>& VelocityField = Evolution->GetVelocityField(Index);
+
+		const TArray<TVector<int32, 3>>& Elements = VelocityField->GetElements();
+		const TArray<TVector<float, 3>>& Forces = VelocityField->GetForces();
 
 		for (int32 ElementIndex = 0; ElementIndex < Elements.Num(); ++ElementIndex)
 		{
diff --git a/Engine/Plugins/Experimental/ChaosCloth/Source/Chaos/Public/ChaosCloth/ChaosClothConfig.h b/Engine/Plugins/Experimental/ChaosCloth/Source/Chaos/Public/ChaosCloth/ChaosClothConfig.h
--- a/Engine/Plugins/Experimental/ChaosCloth/Source/Chaos/Public/ChaosCloth/ChaosClothConfig.h
+++ b/Engine/Plugins/Experimental/ChaosCloth/Source/Chaos/Public/ChaosCloth/ChaosClothConfig.h
@@ -34,15 +34,15 @@ public:
 	 * -	Total Mass: The total mass is distributed equally over all the particles. Useful when referencing a specific garment size and feel.
 	 * -	Density: A constant mass density is used. Density is usually the preferred way of setting mass since it allows matching real life materials' density values.
 	 */
-	UPROPERTY(EditAnywhere, Category = "Mass Config")
+	UPROPERTY(EditAnywhere, Category = "Mass Properties")
 	EClothMassMode MassMode = EClothMassMode::Density;
 
 	// The value used when the Mass Mode is set to Uniform Mass
-	UPROPERTY(EditAnywhere, Category = "Mass Config", meta = (UIMin = "0.000001", UIMax = "0.001", ClampMin = "0"))
+	UPROPERTY(EditAnywhere, Category = "Mass Properties", meta = (UIMin = "0.000001", UIMax = "0.001", ClampMin = "0", EditCondition = "MassMode == EClothMassMode::UniformMass"))
 	float UniformMass = 0.00015f;
 
 	// The value used when Mass Mode is set to TotalMass
-	UPROPERTY(EditAnywhere, Category = "Mass Config", meta = (UIMin = "0.001", UIMax = "10", ClampMin = "0"))
+	UPROPERTY(EditAnywhere, Category = "Mass Properties", meta = (UIMin = "0.001", UIMax = "10", ClampMin = "0", EditCondition = "MassMode == EClothMassMode::TotalMass"))
 	float TotalMass = 0.5f;
 
 	/**
@@ -55,23 +55,27 @@ public:
 	 * Cotton: 0.2
 	 * Silk: 0.1
 	 */
-	UPROPERTY(EditAnywhere, Category = "Mass Config", meta = (UIMin = "0.001", UIMax = "1", ClampMin = "0"))
+	UPROPERTY(EditAnywhere, Category = "Mass Properties", meta = (UIMin = "0.001", UIMax = "1", ClampMin = "0", EditCondition = "MassMode == EClothMassMode::Density"))
 	float Density = 0.35f;
 
 	// This is a lower bound to cloth particle masses.
-	UPROPERTY(EditAnywhere, Category = "Mass Config")
+	UPROPERTY()
 	float MinPerParticleMass = 0.0001f;	
 
 	// The Stiffness of the Edge constraints, only use lower than 1 values for very stretchy materials. Increase the iteration count for stiffer materials.
-	UPROPERTY(EditAnywhere, Category = Stiffness, meta = (UIMin = "0", UIMax = "1", ClampMin = "0", ClampMax = "1"))
+	UPROPERTY(EditAnywhere, Category = "Material Properties", meta = (UIMin = "0", UIMax = "1", ClampMin = "0", ClampMax = "1"))
 	float EdgeStiffness = 1.f;
 
 	// The Stiffness of the bending constraints. Increase the iteration count for stiffer materials.
-	UPROPERTY(EditAnywhere, Category = Stiffness, meta = (UIMin = "0", UIMax = "1", ClampMin = "0", ClampMax = "1"))
+	UPROPERTY(EditAnywhere, Category = "Material Properties", meta = (UIMin = "0", UIMax = "1", ClampMin = "0", ClampMax = "1"))
 	float BendingStiffness = 1.f;
 
+	// Enable the more accurate bending element constraints instead of the faster cross-edge spring constraints used for controlling bending stiffness.
+	UPROPERTY(EditAnywhere, Category = "Material Properties")
+	bool bUseBendingElements = false;
+
 	// The stiffness of the area preservation constraints. Increase the iteration count for stiffer materials.
-	UPROPERTY(EditAnywhere, Category = Stiffness, meta = (UIMin = "0", UIMax = "1", ClampMin = "0", ClampMax = "1"))
+	UPROPERTY(EditAnywhere, Category = "Material Properties", meta = (UIMin = "0", UIMax = "1", ClampMin = "0", ClampMax = "1"))
 	float AreaStiffness = 1.f;
 
 	// The stiffness of the volume preservation constraints.
@@ -99,28 +103,54 @@ public:
 	float ShapeTargetStiffness = 0.f;
 
 	// The radius of cloth points when considering collisions against collider shapes.
-	UPROPERTY(EditAnywhere, Category = "Cloth Dynamics", meta = (UIMin = "0", UIMax = "100", ClampMin = "0", ClampMax = "1000"))
+	UPROPERTY(EditAnywhere, Category = "Collision Properties", meta = (UIMin = "0", UIMax = "100", ClampMin = "0", ClampMax = "1000"))
 	float CollisionThickness = 1.0f;
 
 	// Friction coefficient for cloth - collider interaction.
-	UPROPERTY(EditAnywhere, Category = "Cloth Dynamics", meta = (UIMin = "0", UIMax = "1", ClampMin = "0", ClampMax = "10"))
+	UPROPERTY(EditAnywhere, Category = "Collision Properties", meta = (UIMin = "0", UIMax = "1", ClampMin = "0", ClampMax = "10"))
 	float FrictionCoefficient = 0.2f;
 
+	// Enable self collision.
+	UPROPERTY(EditAnywhere, Category = "Collision Properties", meta = (InlineEditConditionToggle))
+	bool bUseSelfCollisions = false;
+
+	// The radius of the spheres used in self collision 
+	UPROPERTY(EditAnywhere, Category = "Collision Properties", meta = (UIMin = "0", UIMax = "100", ClampMin = "0", ClampMax = "1000", EditCondition = "bUseSelfCollisions"))
+	float SelfCollisionThickness = 2.0f;
+
 	// The amount of damping applied to the cloth velocities.
-	UPROPERTY(EditAnywhere, Category = "Cloth Dynamics", meta = (UIMin = "0", UIMax = "1", ClampMin = "0", ClampMax = "1"))
+	UPROPERTY(EditAnywhere, Category = "Environmental Properties", meta = (UIMin = "0", UIMax = "1", ClampMin = "0", ClampMax = "1"))
 	float DampingCoefficient = 0.01f;
 
 	// The drag coefficient applying on each particle.
-	UPROPERTY(EditAnywhere, Category = "Cloth Dynamics", meta = (UIMin = "0", UIMax = "1", ClampMin = "0", ClampMax = "10"))
+	UPROPERTY(EditAnywhere, Category = "Environmental Properties", meta = (UIMin = "0", UIMax = "1", ClampMin = "0", ClampMax = "10"))
 	float DragCoefficient = 0.1f;
 
-	// Default spring stiffness for anim drive if an anim drive is in use.
-	UPROPERTY(EditAnywhere, Category = Stiffness, meta = (UIMin = "0", UIMax = "1", ClampMin = "0", ClampMax = "1"))
-	float AnimDriveSpringStiffness = 0.001f;
+	// Use the config gravity value instead of world gravity.
+	UPROPERTY(EditAnywhere, Category = "Environmental Properties", meta = (InlineEditConditionToggle))
+	bool bUseGravityOverride = false;
+
+	// Scale factor applied to the world gravity and also to the clothing simulation interactor gravity. Does not affect the gravity if set using the override below.
+	UPROPERTY(EditAnywhere, Category = "Environmental Properties")
+	float GravityScale = 1.f;
 
-	// Enable the more accurate bending element constraints instead of the faster cross-edge spring constraints used for controlling bending stiffness.
-	UPROPERTY(EditAnywhere, Category = "Cloth Enable Flags")
-	bool bUseBendingElements = false;
+	// The gravitational acceleration vector [cm/s^2]
+	UPROPERTY(EditAnywhere, Category = "Environmental Properties", meta = (EditCondition = "bUseGravityOverride"))
+	FVector Gravity = { 0.f, 0.f, -980.665f };
+
+	// The strength of the constraint driving the cloth towards the animated goal mesh.
+	UPROPERTY(EditAnywhere, Category = "Animation Properties", meta = (UIMin = "0", UIMax = "1", ClampMin = "0", ClampMax = "1"))
+	float AnimDriveSpringStiffness = 0.f;
+
+	// The amount of linear velocities sent to the local cloth space from the reference bone
+	// (the closest bone to the root on which the cloth section has been skinned, or the root itself if the cloth isn't skinned).
+	UPROPERTY(EditAnywhere, Category = "Animation Properties", meta = (UIMin = "0", UIMax = "1", ClampMin = "0", ClampMax = "1"))
+	FVector LinearVelocityScale = { 0.75f, 0.75f, 0.75f };
+
+	// The amount of angular velocities sent to the local cloth space from the reference bone
+	// (the closest bone to the root on which the cloth section has been skinned, or the root itself if the cloth isn't skinned).
+	UPROPERTY(EditAnywhere, Category = "Animation Properties", meta = (UIMin = "0", UIMax = "1", ClampMin = "0", ClampMax = "1"))
+	float AngularVelocityScale = 0.75f;
 
 	// Enable tetrahedral constraints.
 	UPROPERTY()
@@ -130,23 +160,10 @@ public:
 	UPROPERTY()
 	bool bUseThinShellVolumeConstraints = false;
 
-	// Enable self collision.
-	UPROPERTY(EditAnywhere, Category = "Cloth Enable Flags")
-	bool bUseSelfCollisions = false;
-
 	// Enable continuous collision detection.
 	UPROPERTY()
 	bool bUseContinuousCollisionDetection = false;
 
-	// The amount of linear velocities sent to the local cloth space from the reference bone
-	// (the closest bone to the root on which the cloth section has been skinned, or the root itself if the cloth isn't skinned).
-	UPROPERTY(EditAnywhere, Category = "Velocity Scales", meta = (UIMin = "0", UIMax = "1", ClampMin = "0", ClampMax = "1"))
-	FVector LinearVelocityScale = { 0.75f, 0.75f, 0.75f };
-
-	// The amount of angular velocities sent to the local cloth space from the reference bone
-	// (the closest bone to the root on which the cloth section has been skinned, or the root itself if the cloth isn't skinned).
-	UPROPERTY(EditAnywhere, Category = "Velocity Scales", meta = (UIMin = "0", UIMax = "1", ClampMin = "0", ClampMax = "1"))
-	float AngularVelocityScale = 0.75f;
 };
 
 /**
@@ -188,8 +205,8 @@ public:
 	int32 SubdivisionCount = 1;
 
 	// The radius of the spheres used in self collision 
-	UPROPERTY(EditAnywhere, Category = Collision, meta = (UIMin = "0", UIMax = "100", ClampMin = "0", ClampMax = "1000"))
-	float SelfCollisionThickness = 2.0f;
+	UPROPERTY()
+	float SelfCollisionThickness_DEPRECATED  = 2.0f;
 
 	// The radius of cloth points when considering collisions against collider shapes.
 	UPROPERTY()
@@ -204,19 +221,19 @@ public:
 	float Damping_DEPRECATED = 0.01f;
 
 	// Use the config gravity value instead of world gravity.
-	UPROPERTY(EditAnywhere, Category = Simulation, meta = (InlineEditConditionToggle))
-	bool bUseGravityOverride = false;
+	UPROPERTY()
+	bool bUseGravityOverride_DEPRECATED = false;
 
 	// Scale factor applied to the world gravity and also to the clothing simulation interactor gravity. Does not affect the gravity if set using the override below.
-	UPROPERTY(EditAnywhere, Category = Simulation)
-	float GravityScale = 1.f;
+	UPROPERTY()
+	float GravityScale_DEPRECATED = 1.f;
 
 	// The gravitational acceleration vector [cm/s^2]
-	UPROPERTY(EditAnywhere, Category = Simulation, meta = (EditCondition = "bUseGravityOverride"))
-	FVector Gravity = { 0.f, 0.f, -980.665f };
+	UPROPERTY()
+	FVector Gravity_DEPRECATED = { 0.f, 0.f, -980.665f };
 
-	// Enable local space simulation to help with floating point precision errors if the character is far away form the world origin
-	UPROPERTY(EditAnywhere, Category = Simulation)
+	// Enable local space simulation to help with jitter due to floating point precision errors if the character is far away from the world origin
+	UPROPERTY()
 	bool bUseLocalSpaceSimulation = true;
 
 	// Enable the XPBD constraints that resolve stiffness independently from the number of iterations
diff --git a/Engine/Plugins/Experimental/ChaosCloth/Source/Chaos/Public/ChaosCloth/ChaosClothingSimulation.h b/Engine/Plugins/Experimental/ChaosCloth/Source/Chaos/Public/ChaosCloth/ChaosClothingSimulation.h
--- a/Engine/Plugins/Experimental/ChaosCloth/Source/Chaos/Public/ChaosCloth/ChaosClothingSimulation.h
+++ b/Engine/Plugins/Experimental/ChaosCloth/Source/Chaos/Public/ChaosCloth/ChaosClothingSimulation.h
@@ -11,15 +11,10 @@
 
 namespace Chaos
 {
-
-	template <typename T>
-	class TTriangleMesh;
-
-	template <typename T, int d>
-	class TPBDEvolution;
-
 	typedef FClothingSimulationContextCommon ClothingSimulationContext;
 	template<class T, int d> class TPBDLongRangeConstraintsBase;
+	template<typename T> class TTriangleMesh;
+	template<typename T, int d> class TPBDEvolution;
 
 	class ClothingSimulation : public FClothingSimulationCommon
 #if WITH_EDITOR
@@ -177,12 +172,12 @@ namespace Chaos
 		float DeltaTime;
 		int32 NumSubsteps;
 
-		bool bOverrideGravity;
-		bool bUseConfigGravity;
-		float GravityScale;
-		FVector Gravity;
-		FVector ConfigGravity;
-		FVector WindVelocity;
+		bool bUseGravityOverride;
+		FVector GravityOverride;
+		TArray<float> GravityScales;
+		TArray<bool> bUseGravityConfigs; 
+		TArray<FVector> GravityConfigs;
+		TArray<FVector> WindVelocities;
 
 		TArray<TSharedPtr<TPBDLongRangeConstraintsBase<float, 3>>> LongRangeConstraints;
 
diff --git a/Engine/Source/Runtime/Experimental/Chaos/Private/Chaos/PBDCollisionSpringConstraintsBase.cpp b/Engine/Source/Runtime/Experimental/Chaos/Private/Chaos/PBDCollisionSpringConstraintsBase.cpp
--- a/Engine/Source/Runtime/Experimental/Chaos/Private/Chaos/PBDCollisionSpringConstraintsBase.cpp
+++ b/Engine/Source/Runtime/Experimental/Chaos/Private/Chaos/PBDCollisionSpringConstraintsBase.cpp
@@ -54,12 +54,20 @@ class FMeshBuildDataProvider
 
 template<class T, int d>
 PBDCollisionSpringConstraintsBase<T, d>::PBDCollisionSpringConstraintsBase(
-    const TDynamicParticles<T, d>& InParticles, const TArray<TVector<int32, 3>>& Elements, const TSet<TVector<int32, 2>>& DisabledCollisionElements, const T Dt, const T Height, const T Stiffness)
-    : MH(Height), MStiffness(Stiffness)
+    const TDynamicParticles<T, d>& InParticles, const TArray<TVector<int32, 3>>& Elements, const TSet<TVector<int32, 2>>& DisabledCollisionElements, const TArray<uint32>& DynamicGroupIds, const TArray<T>& PerGroupThicknesses, const T Dt, const T Stiffness)
+    : MDynamicGroupIds(DynamicGroupIds), MPerGroupThicknesses(PerGroupThicknesses), MStiffness(Stiffness)
 {
 	if (!Elements.Num())
 		return;
 #if PLATFORM_DESKTOP && PLATFORM_64BITS
+	// Calculate the reach of the collision line (= twice the maximum thickness)
+	T MaxThickness = 0.f;
+	for (T Thickness : MPerGroupThicknesses)
+	{
+		MaxThickness = FMath::Max(MaxThickness, Thickness);
+	}
+	const T Height = MaxThickness + MaxThickness;
+
 	TkDOPTree<const FMeshBuildDataProvider, uint32> DopTree;
 	TArray<FkDOPBuildCollisionTriangle<uint32>> BuildTraingleArray;
 	for (int32 i = 0; i < Elements.Num(); ++i)
@@ -73,7 +81,7 @@ PBDCollisionSpringConstraintsBase<T, d>::PBDCollisionSpringConstraintsBase(
 	PhysicsParallelFor(InParticles.Size(), [&](int32 Index) {
 		FkHitResult Result;
 		const auto& Start = InParticles.X(Index);
-		const auto End = Start + InParticles.V(Index) * Dt + InParticles.V(Index).GetSafeNormal() * MH;
+		const auto End = Start + InParticles.V(Index) * Dt + InParticles.V(Index).GetSafeNormal() * Height;
 		FVector4 Start4(Start.X, Start.Y, Start.Z, 0);
 		FVector4 End4(End.X, End.Y, End.Z, 0);
 		TkDOPLineCollisionCheck<const FMeshBuildDataProvider, uint32> ray(Start4, End4, true, DopDataProvider, &Result);
@@ -124,7 +132,8 @@ TVector<T, d> PBDCollisionSpringConstraintsBase<T, d>::GetDelta(const TPBDPartic
 	const TVector<T, d>& P2 = InParticles.P(i2);
 	const TVector<T, d>& P3 = InParticles.P(i3);
 	const TVector<T, d>& P4 = InParticles.P(i4);
-	const TVector<T, d> P = MBarys[i][0] * P2 + MBarys[i][1] * P3 + MBarys[i][2] * P4 + MH * MNormals[i];
+	const T Height = MPerGroupThicknesses[MDynamicGroupIds[i1]] + MPerGroupThicknesses[MDynamicGroupIds[i2]];  // Assume same height on each point of the constraint element
+	const TVector<T, d> P = MBarys[i][0] * P2 + MBarys[i][1] * P3 + MBarys[i][2] * P4 + Height * MNormals[i];
 	TVector<T, d> Difference = P1 - P;
 	if (TVector<T, d>::DotProduct(Difference, MNormals[i]) > 0)
 		return TVector<T, d>(0);
diff --git a/Engine/Source/Runtime/Experimental/Chaos/Private/Chaos/PBDEvolution.cpp b/Engine/Source/Runtime/Experimental/Chaos/Private/Chaos/PBDEvolution.cpp
--- a/Engine/Source/Runtime/Experimental/Chaos/Private/Chaos/PBDEvolution.cpp
+++ b/Engine/Source/Runtime/Experimental/Chaos/Private/Chaos/PBDEvolution.cpp
@@ -28,6 +28,25 @@ DECLARE_CYCLE_STAT(TEXT("Chaos PBD Collider Kinematic Update"), STAT_CollisionKi
 
 using namespace Chaos;
 
+template<class T, int d>
+void TPBDEvolution<T, d>::AddGroups(int32 Num)
+{
+	// Add elements
+	const uint32 Offset = TArrayCollection::Size();
+	TArrayCollection::AddElementsHelper(Num);
+
+	// Set defaults
+	for (uint32 GroupId = Offset; GroupId < TArrayCollection::Size(); ++GroupId)
+	{
+		MGroupGravityForces[GroupId].SetAcceleration(MGravity);
+		MGroupCollisionThicknesses[GroupId] = MCollisionThickness;
+		MGroupSelfCollisionThicknesses[GroupId] = MSelfCollisionThickness;
+		MGroupCoefficientOfFrictions[GroupId] = MCoefficientOfFriction;
+		MGroupDampings[GroupId] = MDamping;
+		MGroupVelocityFields[GroupId] = MakeUnique<TVelocityField<float, 3>>();
+	}
+}
+
 template<class T, int d>
 TPBDEvolution<T, d>::TPBDEvolution(TPBDParticles<T, d>&& InParticles, TKinematicGeometryClothParticles<T, d>&& InGeometryParticles, TArray<TVector<int32, 3>>&& CollisionTriangles,
     int32 NumIterations, T CollisionThickness, T SelfCollisionThickness, T CoefficientOfFriction, T Damping)
@@ -35,18 +54,26 @@ TPBDEvolution<T, d>::TPBDEvolution(TPBDParticles<T, d>&& InParticles, TKinematic
 	, MCollisionParticles(MoveTemp(InGeometryParticles))
 	, MCollisionTriangles(MoveTemp(CollisionTriangles))
 	, MNumIterations(NumIterations)
+	, MGravity(TVector<T, d>((T)0., (T)0., (T)-980.665))
 	, MCollisionThickness(CollisionThickness)
 	, MSelfCollisionThickness(SelfCollisionThickness)
 	, MCoefficientOfFriction(CoefficientOfFriction)
 	, MDamping(Damping)
 	, MTime(0)
 {
+	// Add group arrays
+	TArrayCollection::AddArray(&MGroupGravityForces);
+	TArrayCollection::AddArray(&MGroupCollisionThicknesses);
+	TArrayCollection::AddArray(&MGroupSelfCollisionThicknesses);
+	TArrayCollection::AddArray(&MGroupCoefficientOfFrictions);
+	TArrayCollection::AddArray(&MGroupDampings);
+	TArrayCollection::AddArray(&MGroupVelocityFields);
+	AddGroups(1);  // Add default group
+
+	// Add particle arrays
+	MParticles.AddArray(&MParticleGroupIds);
 	MCollisionParticles.AddArray(&MCollided);
 	MCollisionParticles.AddArray(&MCollisionParticleGroupIds);
-	MParticles.AddArray(&MParticleGroupIds);
-	MPerGroupDamping.Add(Damping);
-	MPerGroupCollisionThickness.Add(CollisionThickness);
-	MPerGroupCoefficientOfFriction.Add(CoefficientOfFriction);
 
 	SetParticleUpdateFunction(
 		[PBDUpdateRule = 
@@ -74,18 +101,10 @@ uint32 TPBDEvolution<T, d>::AddParticles(uint32 Num, uint32 GroupId)
 	}
 
 	// Resize group parameter arrays
-	const uint32 GroupNum = (uint32)MPerGroupDamping.Num();
-	if (GroupId >= GroupNum)
+	const uint32 GroupSize = TArrayCollection::Size();
+	if (GroupId >= GroupSize)
 	{
-		MPerGroupDamping.SetNum((int32)GroupId + 1);
-		MPerGroupCollisionThickness.SetNum((int32)GroupId + 1);
-		MPerGroupCoefficientOfFriction.SetNum((int32)GroupId + 1);
-		for (uint32 i = GroupNum; i <= GroupId; ++i)
-		{
-			MPerGroupDamping[i] = MDamping;
-			MPerGroupCollisionThickness[i] = MCollisionThickness;
-			MPerGroupCoefficientOfFriction[i] = MCoefficientOfFriction;
-		}
+		AddGroups(GroupId + 1 - GroupSize);
 	}
 	return Offset;
 }
@@ -111,9 +130,9 @@ void TPBDEvolution<T, d>::AdvanceOneTimeStep(const T Dt)
 	SCOPE_CYCLE_COUNTER(STAT_ChaosPBDVAdvanceTime);
 	TPerParticleInitForce<T, d> InitForceRule;
 	TPerParticleEulerStepVelocity<T, d> EulerStepVelocityRule;
-	TPerGroupDampVelocity<T, d> DampVelocityRule(MParticleGroupIds, MPerGroupDamping);
+	TPerGroupDampVelocity<T, d> DampVelocityRule(MParticleGroupIds, MGroupDampings);
 	TPerParticlePBDEulerStep<T, d> EulerStepRule;
-	TPerParticlePBDCollisionConstraint<T, d, EGeometryParticlesSimType::Other> CollisionRule(MCollisionParticles, MCollided, MParticleGroupIds, MCollisionParticleGroupIds, MPerGroupCollisionThickness, MPerGroupCoefficientOfFriction);
+	TPerParticlePBDCollisionConstraint<T, d, EGeometryParticlesSimType::Other> CollisionRule(MCollisionParticles, MCollided, MParticleGroupIds, MCollisionParticleGroupIds, MGroupCollisionThicknesses, MGroupCoefficientOfFrictions);
 
 	{
 		SCOPE_CYCLE_COUNTER(STAT_ChaosPBDVelocityDampUpdateState);
@@ -122,9 +141,9 @@ void TPBDEvolution<T, d>::AdvanceOneTimeStep(const T Dt)
 
 	{
 		SCOPE_CYCLE_COUNTER(STAT_ChaosPBDVelocityFieldUpdateForces);
-		for (FVelocityField& VelocityField : VelocityFields)
+		for (const TUniquePtr<FVelocityField>& VelocityField : MGroupVelocityFields)
 		{
-			VelocityField.UpdateForces(MParticles, Dt);
+			VelocityField->UpdateForces(MParticles, Dt);
 		}
 	}	
 
@@ -140,16 +159,17 @@ void TPBDEvolution<T, d>::AdvanceOneTimeStep(const T Dt)
 		SCOPE_CYCLE_COUNTER(STAT_ChaosPBDPreIterationUpdates);
 		PhysicsParallelFor(MParticles.Size(), [&](int32 Index)
 		{
+			const uint32 ParticleGroupId = MParticleGroupIds[Index];
+
 			InitForceRule.Apply(MParticles, Dt, Index); // F = TV(0)
-			GravityForces.Apply(MParticles, Dt, Index);
+			MGroupGravityForces[ParticleGroupId].Apply(MParticles, Dt, Index); // F += M * G
 			for (TFunction<void(TPBDParticles<T, d>&, const T, const int32)>& ForceRule : MForceRules)
 			{
 				ForceRule(MParticles, Dt, Index); // F += M * A
 			}
-			for (FVelocityField& VelocityField : VelocityFields)
-			{
-				VelocityField.Apply(MParticles, Dt, Index);
-			}
+			
+			MGroupVelocityFields[ParticleGroupId]->Apply(MParticles, Dt, Index);
+
 			if (MKinematicUpdate)
 			{
 				MKinematicUpdate(MParticles, Dt, MTime + Dt, Index); // X = ...
@@ -168,7 +188,7 @@ void TPBDEvolution<T, d>::AdvanceOneTimeStep(const T Dt)
 		}, NonParallelUpdate);
 	}
 #if !COMPILE_WITHOUT_UNREAL_SUPPORT
-	TPBDCollisionSpringConstraints<T, d> SelfCollisionRule(MParticles, MCollisionTriangles, MDisabledCollisionElements, Dt, MSelfCollisionThickness, 1.5f);
+	TPBDCollisionSpringConstraints<T, d> SelfCollisionRule(MParticles, MCollisionTriangles, MDisabledCollisionElements, MParticleGroupIds, MGroupSelfCollisionThicknesses, Dt);
 #endif
 
 	for (TFunction<void()>& InitConstraintRule : MInitConstraintRules)
@@ -176,8 +196,7 @@ void TPBDEvolution<T, d>::AdvanceOneTimeStep(const T Dt)
 		InitConstraintRule();  // Clear XPBD's Lambdas
 	}
 
-
-	// Do one extra collision pass at the start to decrease likelyhood of cloth penetrating- TODO: Add option for more collision passed interleaved between constraints
+	// Do one extra collision pass at the start to decrease likelihood of cloth penetrating- TODO: Add option for more collision passed interleaved between constraints
 	{
 		SCOPE_CYCLE_COUNTER(STAT_ChaosPBDCollisionRule);
 		CollisionRule.ApplyPerParticle(MParticles, Dt);
diff --git a/Engine/Source/Runtime/Experimental/Chaos/Public/Chaos/PBDCollisionSpringConstraints.h b/Engine/Source/Runtime/Experimental/Chaos/Public/Chaos/PBDCollisionSpringConstraints.h
--- a/Engine/Source/Runtime/Experimental/Chaos/Public/Chaos/PBDCollisionSpringConstraints.h
+++ b/Engine/Source/Runtime/Experimental/Chaos/Public/Chaos/PBDCollisionSpringConstraints.h
@@ -19,8 +19,8 @@ class TPBDCollisionSpringConstraints : public PBDCollisionSpringConstraintsBase<
 	using Base::MConstraints;
 
   public:
-	TPBDCollisionSpringConstraints(const TDynamicParticles<T, d>& InParticles, const TArray<TVector<int32, 3>>& Elements, const TSet<TVector<int32, 2>>& DisabledCollisionElements, const T Dt, const T Height = (T)0, const T Stiffness = (T)1)
-	    : Base(InParticles, Elements, DisabledCollisionElements, Dt, Height, Stiffness) {}
+	TPBDCollisionSpringConstraints(const TDynamicParticles<T, d>& InParticles, const TArray<TVector<int32, 3>>& Elements, const TSet<TVector<int32, 2>>& DisabledCollisionElements, const TArray<uint32>& DynamicGroupIds, const TArray<T>& PerGroupThicknesses, const T Dt, const T Stiffness = (T)1.5)
+	    : Base(InParticles, Elements, DisabledCollisionElements, DynamicGroupIds, PerGroupThicknesses, Dt, Stiffness) {}
 	virtual ~TPBDCollisionSpringConstraints() {}
 
 	void Apply(TPBDParticles<T, d>& InParticles, const T Dt, const int32 InConstraintIndex) const
diff --git a/Engine/Source/Runtime/Experimental/Chaos/Public/Chaos/PBDCollisionSpringConstraintsBase.h b/Engine/Source/Runtime/Experimental/Chaos/Public/Chaos/PBDCollisionSpringConstraintsBase.h
--- a/Engine/Source/Runtime/Experimental/Chaos/Public/Chaos/PBDCollisionSpringConstraintsBase.h
+++ b/Engine/Source/Runtime/Experimental/Chaos/Public/Chaos/PBDCollisionSpringConstraintsBase.h
@@ -14,7 +14,7 @@ template<class T, int d>
 class PBDCollisionSpringConstraintsBase
 {
   public:
-	PBDCollisionSpringConstraintsBase(const TDynamicParticles<T, d>& InParticles, const TArray<TVector<int32, 3>>& Elements, const TSet<TVector<int32, 2>>& DisabledCollisionElements, const T Dt, const T Height = (T)0, const T Stiffness = (T)1);
+	PBDCollisionSpringConstraintsBase(const TDynamicParticles<T, d>& InParticles, const TArray<TVector<int32, 3>>& Elements, const TSet<TVector<int32, 2>>& DisabledCollisionElements, const TArray<uint32>& DynamicGroupIds, const TArray<T>& PerGroupThicknesses, const T Dt, const T Stiffness = (T)1);
 	virtual ~PBDCollisionSpringConstraintsBase() {}
 
 	TVector<T, d> GetDelta(const TPBDParticles<T, d>& InParticles, const int32 i) const;
@@ -22,9 +22,10 @@ class PBDCollisionSpringConstraintsBase
   protected:
 	TArray<TVector<int32, 4>> MConstraints;
 	TArray<TVector<T, 3>> MBarys;
-	T MH;
 
   private:
+	const TArray<uint32>& MDynamicGroupIds;
+	const TArray<T>& MPerGroupThicknesses;
 	TArray<TVector<T, d>> MNormals; // per constraint, sign changes depending on orientation of colliding particle
 	T MStiffness;
 };
diff --git a/Engine/Source/Runtime/Experimental/Chaos/Public/Chaos/PBDEvolution.h b/Engine/Source/Runtime/Experimental/Chaos/Public/Chaos/PBDEvolution.h
--- a/Engine/Source/Runtime/Experimental/Chaos/Public/Chaos/PBDEvolution.h
+++ b/Engine/Source/Runtime/Experimental/Chaos/Public/Chaos/PBDEvolution.h
@@ -10,7 +10,7 @@
 namespace Chaos
 {
 template<class T, int d>
-class CHAOS_API TPBDEvolution
+class CHAOS_API TPBDEvolution  : private TArrayCollection
 {
   public:
 	using FGravityForces = TPerParticleGravity<T, d>;
@@ -37,11 +37,11 @@ class CHAOS_API TPBDEvolution
 	const TPBDParticles<T, d>& Particles() const { return MParticles; }
 	TPBDParticles<T, d>& Particles() { return MParticles; }
 
-	FGravityForces& GetGravityForces() { return GravityForces; }
-	const FGravityForces& GetGravityForces() const { return GravityForces; }
+	FGravityForces& GetGravityForces(const uint32 GroupId = 0) { check(GroupId < TArrayCollection::Size()); return MGroupGravityForces[GroupId]; }
+	const FGravityForces& GetGravityForces(const uint32 GroupId = 0) const { check(GroupId < TArrayCollection::Size()); return MGroupGravityForces[GroupId]; }
 
-	TArray<FVelocityField>& GetVelocityFields() { return VelocityFields; }
-	const TArray<FVelocityField>& GetVelocityFields() const { return VelocityFields; }
+	TUniquePtr<FVelocityField>& GetVelocityField(const uint32 GroupId = 0) { check(GroupId < TArrayCollection::Size()); return MGroupVelocityFields[GroupId]; }
+	const TUniquePtr<FVelocityField>& GetVelocityField(const uint32 GroupId = 0) const { check(GroupId < TArrayCollection::Size()); return MGroupVelocityFields[GroupId]; }
 
 	const TGeometryClothParticles<T, d>& CollisionParticles() const { return MCollisionParticles; }
 	TGeometryClothParticles<T, d>& CollisionParticles() { return MCollisionParticles; }
@@ -53,25 +53,29 @@ class CHAOS_API TPBDEvolution
 	int32 GetIterations() const { return MNumIterations; }
 	void SetIterations(const int32 Iterations) { MNumIterations = Iterations; }
 
-	T GetSelfCollisionThickness() const { return MSelfCollisionThickness; }
-	void SetSelfCollisionThickness(const T SelfCollisionThickness) { MSelfCollisionThickness = SelfCollisionThickness; }
+	T GetSelfCollisionThickness(const uint32 GroupId = 0) const { check(GroupId < TArrayCollection::Size()); return MGroupSelfCollisionThicknesses[GroupId]; }
+	void SetSelfCollisionThickness(const T SelfCollisionThickness, const uint32 GroupId = 0) { check(GroupId < TArrayCollection::Size()); MGroupSelfCollisionThicknesses[GroupId] = SelfCollisionThickness; }
 
-	T GetCollisionThickness(const uint32 GroupId = 0) const { return MPerGroupCollisionThickness[GroupId]; }
-	void SetCollisionThickness(const T CollisionThickness, const uint32 GroupId = 0) { MPerGroupCollisionThickness[GroupId] = CollisionThickness; }
+	T GetCollisionThickness(const uint32 GroupId = 0) const { check(GroupId < TArrayCollection::Size()); return MGroupCollisionThicknesses[GroupId]; }
+	void SetCollisionThickness(const T CollisionThickness, const uint32 GroupId = 0) { check(GroupId < TArrayCollection::Size()); MGroupCollisionThicknesses[GroupId] = CollisionThickness; }
 
-	T GetCoefficientOfFriction(const uint32 GroupId = 0) const { return MPerGroupCoefficientOfFriction[GroupId]; }
-	void SetCoefficientOfFriction(const T CoefficientOfFriction, const uint32 GroupId = 0) { MPerGroupCoefficientOfFriction[GroupId] = CoefficientOfFriction; }
+	T GetCoefficientOfFriction(const uint32 GroupId = 0) const { check(GroupId < TArrayCollection::Size()); return MGroupCoefficientOfFrictions[GroupId]; }
+	void SetCoefficientOfFriction(const T CoefficientOfFriction, const uint32 GroupId = 0) { check(GroupId < TArrayCollection::Size()); MGroupCoefficientOfFrictions[GroupId] = CoefficientOfFriction; }
 
-	T GetDamping(const uint32 GroupId = 0) const { return MPerGroupDamping[GroupId]; }
-	void SetDamping(const T Damping, const uint32 GroupId = 0) { MPerGroupDamping[GroupId] = Damping; }
+	T GetDamping(const uint32 GroupId = 0) const { check(GroupId < TArrayCollection::Size()); return MGroupDampings[GroupId]; }
+	void SetDamping(const T Damping, const uint32 GroupId = 0) { check(GroupId < TArrayCollection::Size()); MGroupDampings[GroupId] = Damping; }
 
 	T GetTime() const { return MTime; }
 
 	void ResetConstraintRules() { MInitConstraintRules.Reset(); MConstraintRules.Reset(); };
 	void ResetSelfCollision() { MCollisionTriangles.Reset(); MDisabledCollisionElements.Reset(); };
-	void ResetVelocityFields() { VelocityFields.Reset(); };
+
+	const TArrayCollectionArray<uint32>& ParticleGroupIds() const { return MParticleGroupIds; }
 
   private:
+	// Add simulation groups and set default values
+	void AddGroups(int32 Num);
+
 	TPBDParticles<T, d> MParticles;
 	TKinematicGeometryClothParticles<T, d> MCollisionParticles;
 	TArray<TVector<int32, 3>> MCollisionTriangles;       // Used for self-collisions
@@ -79,19 +83,20 @@ class CHAOS_API TPBDEvolution
 	TArrayCollectionArray<bool> MCollided;
 	TArrayCollectionArray<uint32> MCollisionParticleGroupIds;  // Used for per group parameters for collision particles
 	TArrayCollectionArray<uint32> MParticleGroupIds;  // Used for per group parameters for particles
-	TArray<T> MPerGroupDamping;
-	TArray<T> MPerGroupCollisionThickness;
-	TArray<T> MPerGroupCoefficientOfFriction;
+	TArrayCollectionArray<FGravityForces> MGroupGravityForces;
+	TArrayCollectionArray<T> MGroupCollisionThicknesses;
+	TArrayCollectionArray<T> MGroupSelfCollisionThicknesses;
+	TArrayCollectionArray<T> MGroupCoefficientOfFrictions;
+	TArrayCollectionArray<T> MGroupDampings;
+	TArrayCollectionArray<TUniquePtr<FVelocityField>> MGroupVelocityFields;
 	int32 MNumIterations;
+	TVector<T, d> MGravity;
 	T MCollisionThickness;
 	T MSelfCollisionThickness;
 	T MCoefficientOfFriction;
 	T MDamping;
 	T MTime;
 
-	FGravityForces GravityForces;
-	TArray<FVelocityField> VelocityFields;
-
 	TArray<TFunction<void(TPBDParticles<T, d>&, const T, const int32)>> MForceRules;
 	TArray<TFunction<void()>> MInitConstraintRules;
 	TArray<TFunction<void(TPBDParticles<T, d>&, const T)>> MConstraintRules;
diff --git a/Engine/Source/Runtime/Experimental/Chaos/Public/Chaos/VelocityField.h b/Engine/Source/Runtime/Experimental/Chaos/Public/Chaos/VelocityField.h
--- a/Engine/Source/Runtime/Experimental/Chaos/Public/Chaos/VelocityField.h
+++ b/Engine/Source/Runtime/Experimental/Chaos/Public/Chaos/VelocityField.h
@@ -13,7 +13,7 @@ namespace Chaos
 	{
 	public:
 		TVelocityField()
-			: QuarterRhoDragCoefficient((T)0.), bIsUniform(true) {}
+			: Elements(nullptr), QuarterRhoDragCoefficient((T)0.), bIsUniform(true) {}
 
 		TVelocityField(
 			const TTriangleMesh<T>& TriangleMesh,
@@ -21,13 +21,20 @@ namespace Chaos
 			const bool bInIsUniform,
 			const T DragCoefficient = (T)0.5,
 			const T FluidDensity = (T)1.225e-6)
-			: PointToTriangleMap(TriangleMesh.GetPointToTriangleMap())  // TODO(Kriss.Gossart): Replace map lookup by safer array lookup
-			, Elements(TriangleMesh.GetElements())
+			: Elements(&TriangleMesh.GetElements())
 			, GetVelocity(InGetVelocity)
 			, bIsUniform(bInIsUniform)
+			, Offset(TriangleMesh.GetVertexRange().Key)
 		{
-			Forces.SetNumUninitialized(Elements.Num());
+			Forces.SetNumUninitialized(Elements->Num());
 			SetProperties(DragCoefficient, FluidDensity);
+
+			const int32 NumPoints = TriangleMesh.GetVertexRange().Value + 1 - Offset;
+			PointToTriangleMap.SetNumUninitialized(NumPoints);
+			for (const TPair<int32, TArray<int32>>& PointToTriangle : TriangleMesh.GetPointToTriangleMap())
+			{
+				PointToTriangleMap[PointToTriangle.Key - Offset] = &PointToTriangle.Value;
+			}
 		}
 
 		virtual ~TVelocityField() {}
@@ -39,14 +46,14 @@ namespace Chaos
 				if (bIsUniform)
 				{
 					const TVector<T, d> Velocity = GetVelocity(TVector<T, d>(0.));
-					for (int32 ElementIndex = 0; ElementIndex < Elements.Num(); ++ElementIndex)
+					for (int32 ElementIndex = 0; ElementIndex < Elements->Num(); ++ElementIndex)
 					{
 						UpdateUniformField(InParticles, Dt, ElementIndex, Velocity);
 					}
 				}
 				else
 				{
-					for (int32 ElementIndex = 0; ElementIndex < Elements.Num(); ++ElementIndex)
+					for (int32 ElementIndex = 0; ElementIndex < Elements->Num(); ++ElementIndex)
 					{
 						UpdateVectorField(InParticles, Dt, ElementIndex);
 					}
@@ -56,8 +63,11 @@ namespace Chaos
 
 		inline void Apply(TPBDParticles<T, d>& InParticles, const T Dt, const int32 Index) const
 		{
-			if (const TArray<int32>* ElementIndices = PointToTriangleMap.Find(Index))  // TODO(Kriss.Gossart): Replace map lookup by faster array lookup
+			const int32 PointIndex = Index - Offset;
+			if (PointToTriangleMap.IsValidIndex(PointIndex))
 			{
+				const TArray<int32>* ElementIndices = PointToTriangleMap[PointIndex];
+
 				for (const int32 ElementIndex : *ElementIndices)
 				{
 					InParticles.F(Index) += Forces[ElementIndex];
@@ -70,13 +80,13 @@ namespace Chaos
 			QuarterRhoDragCoefficient = (T)0.25 * FluidDensity * DragCoefficient;
 		}
 
-		const TArray<TVector<int32, 3>>& GetElements() const { return Elements; }
+		const TArray<TVector<int32, 3>>& GetElements() const { return *Elements; }
 		const TArray<TVector<T, d>>& GetForces() const { return Forces; }
 
 	private:
 		inline void UpdateUniformField(const TPBDParticles<T, d>& InParticles, const T Dt, int32 ElementIndex, const TVector<T, d>& Velocity) const
 		{
-			const TVector<int32, 3>& Element = Elements[ElementIndex];
+			const TVector<int32, 3>& Element = (*Elements)[ElementIndex];
 
 			// Calculate direction and the relative velocity of the triangle to the flow
 			const TVector<T, d>& SurfaceVelocity = (T)(1. / 3.) * (
@@ -100,7 +110,7 @@ namespace Chaos
 
 		inline void UpdateVectorField(const TPBDParticles<T, d>& InParticles, const T Dt, int32 ElementIndex) const
 		{
-			const TVector<int32, 3>& Element = Elements[ElementIndex];
+			const TVector<int32, 3>& Element = (*Elements)[ElementIndex];
 
 			// Get the triangle's position
 			const TVector<T, d>& SurfacePosition = (T)(1. / 3.) * (
@@ -129,11 +139,12 @@ namespace Chaos
 		}
 
 	private:
-		const TMap<int32, TArray<int32>>& PointToTriangleMap;
-		const TArray<TVector<int32, 3>>& Elements;
+		TArray<const TArray<int32>*> PointToTriangleMap;
+		const TArray<TVector<int32, 3>>* Elements;
 		TFunction<TVector<T, d>(const TVector<T, d>&)> GetVelocity;
 		mutable TArray<TVector<T, 3>> Forces;
 		float QuarterRhoDragCoefficient;
 		bool bIsUniform;
+		int32 Offset;
 	};
 }
