Chaos Cloth - Remove internal config parameters from UI.

diff --git a/Engine/Plugins/Experimental/ChaosCloth/Source/Chaos/Private/ChaosCloth/ChaosClothConfig.cpp b/Engine/Plugins/Experimental/ChaosCloth/Source/Chaos/Private/ChaosCloth/ChaosClothConfig.cpp
--- a/Engine/Plugins/Experimental/ChaosCloth/Source/Chaos/Private/ChaosCloth/ChaosClothConfig.cpp
+++ b/Engine/Plugins/Experimental/ChaosCloth/Source/Chaos/Private/ChaosCloth/ChaosClothConfig.cpp
@@ -117,6 +117,12 @@ void UChaosClothConfig::PostLoad()
 	{
 		DragCoefficient = 0.07f;  // Reset to a more appropriate default for chaos cloth assets saved before this custom version
 	}
+
+	if (ChaosClothConfigCustomVersion < FChaosClothConfigCustomVersion::RemoveInternalConfigParameters)
+	{
+		MinPerParticleMass = 0.0001f;  // Override these values in case they might have been accidentally
+		bUseGeodesicDistance = true;   // changed from their default at any point in time
+	}
 }
 
 UChaosClothSharedSimConfig::UChaosClothSharedSimConfig()
diff --git a/Engine/Plugins/Experimental/ChaosCloth/Source/Chaos/Private/ChaosCloth/ChaosClothConfigCustomVersion.h b/Engine/Plugins/Experimental/ChaosCloth/Source/Chaos/Private/ChaosCloth/ChaosClothConfigCustomVersion.h
--- a/Engine/Plugins/Experimental/ChaosCloth/Source/Chaos/Private/ChaosCloth/ChaosClothConfigCustomVersion.h
+++ b/Engine/Plugins/Experimental/ChaosCloth/Source/Chaos/Private/ChaosCloth/ChaosClothConfigCustomVersion.h
@@ -18,6 +18,8 @@ struct FChaosClothConfigCustomVersion
 		AddDampingThicknessMigration = 2,
 		// Added gravity and self collision thickness per cloth
 		AddGravitySelfCollisionMigration = 3,
+		// Remove internal config parameters from UI
+		RemoveInternalConfigParameters = 4,
 
 		// -----<new versions can be added above this line>-------------------------------------------------
 		VersionPlusOne,
diff --git a/Engine/Plugins/Experimental/ChaosCloth/Source/Chaos/Public/ChaosCloth/ChaosClothConfig.h b/Engine/Plugins/Experimental/ChaosCloth/Source/Chaos/Public/ChaosCloth/ChaosClothConfig.h
--- a/Engine/Plugins/Experimental/ChaosCloth/Source/Chaos/Public/ChaosCloth/ChaosClothConfig.h
+++ b/Engine/Plugins/Experimental/ChaosCloth/Source/Chaos/Public/ChaosCloth/ChaosClothConfig.h
@@ -60,7 +60,7 @@ public:
 
 	// This is a lower bound to cloth particle masses.
 	UPROPERTY()
-	float MinPerParticleMass = 0.0001f;	
+	float MinPerParticleMass = 0.0001f;
 
 	// The Stiffness of the Edge constraints, only use lower than 1 values for very stretchy materials. Increase the iteration count for stiffer materials.
 	UPROPERTY(EditAnywhere, Category = "Material Properties", meta = (UIMin = "0", UIMax = "1", ClampMin = "0", ClampMax = "1"))
@@ -95,7 +95,7 @@ public:
 
 	// Use geodesic instead of euclidean distance calculations in the long range attachment constraint,
 	// which is slower at setup but less prone to artifacts during simulation.
-	UPROPERTY(EditAnywhere, Category = "Long Range Attachment")
+	UPROPERTY()
 	bool bUseGeodesicDistance = true;
 
 	// The stiffness of the shape target constraints
