Cloth - Add cloth info functions to the Clothing Simulation Interactor to make them accessible from blueprint.
Chaos Cloth - Implement GetNumSubteps() and GetNumIterations().

diff --git a/Engine/Plugins/Experimental/ChaosCloth/Source/Chaos/Private/ChaosCloth/ChaosClothingSimulation.cpp b/Engine/Plugins/Experimental/ChaosCloth/Source/Chaos/Private/ChaosCloth/ChaosClothingSimulation.cpp
--- a/Engine/Plugins/Experimental/ChaosCloth/Source/Chaos/Private/ChaosCloth/ChaosClothingSimulation.cpp
+++ b/Engine/Plugins/Experimental/ChaosCloth/Source/Chaos/Private/ChaosCloth/ChaosClothingSimulation.cpp
@@ -86,6 +86,7 @@ namespace ChaosClothingSimulationDefault
 {
 	static const FVector Gravity(0.f, 0.f, -980.665f);
 	static const int32 NumIterations = 1;
+	static const int32 NumSubsteps = 1;
 	static const float SelfCollisionThickness = 2.f;
 	static const float CollisionThickness = 1.2f;
 	static const float FrictionCoefficient = 0.2f;
@@ -96,7 +97,8 @@ namespace ChaosClothingSimulationDefault
 ClothingSimulation::ClothingSimulation()
 	: ClothSharedSimConfig(nullptr)
 	, ExternalCollisionsOffset(0)
-	, NumSubsteps(1)
+	, NumIterations(ChaosClothingSimulationDefault::NumIterations)
+	, NumSubsteps(ChaosClothingSimulationDefault::NumSubsteps)
 	, bUseGravityOverride(false)
 	, GravityOverride(ChaosClothingSimulationDefault::Gravity)
 	, bUseLocalSpaceSimulation(false)
@@ -374,17 +376,14 @@ void ClothingSimulation::CreateActor(USkeletalMeshComponent* InOwnerComponent, U
 
 void ClothingSimulation::ResetStats()
 {
-#if WITH_EDITOR
 	NumCloths = 0;
 	NumKinemamicParticles = 0;
 	NumDynamicParticles = 0;
 	SimulationTime = 0.f;
-#endif  // #if WITH_EDITOR
 }
 
 void ClothingSimulation::UpdateStats(int32 InSimDataIndex)
 {
-#if WITH_EDITOR
 	const uint32 Offset = IndexToRangeMap[InSimDataIndex][0];
 	const uint32 Range = IndexToRangeMap[InSimDataIndex][1];
 	if (const int32 AddedParticleCount = Range - Offset)
@@ -408,7 +407,6 @@ void ClothingSimulation::UpdateStats(int32 InSimDataIndex)
 	}
 	NumKinemamicParticles += NumAddedKinenamicParticles;
 	NumDynamicParticles += NumAddedDynamicParticles;
-#endif  // #if WITH_EDITOR
 }
 
 void ClothingSimulation::ExtractCollisions(const UClothingAssetCommon* Asset, int32 InSimDataIndex)
@@ -439,6 +437,7 @@ void ClothingSimulation::UpdateSimulationFromSharedSimConfig()
 
 		// Now set all the common parameters on the simulation
 		NumSubsteps = ClothSharedSimConfig->SubdivisionCount;
+		NumIterations = ClothSharedSimConfig->IterationCount;
 		Evolution->SetIterations(ClothSharedSimConfig->IterationCount);
 	}
 }
@@ -1343,9 +1342,7 @@ void ClothingSimulation::Simulate(IClothingSimulationContext* InContext)
 		return;
 	}
 
-#if WITH_EDITOR
 	const double StartTime = FPlatformTime::Seconds();
-#endif 
 
 	// Filter delta time to smoothen time variations and prevent unwanted vibrations
 	static const float DeltaTimeDecay = 0.1f;
@@ -1540,9 +1537,10 @@ void ClothingSimulation::Simulate(IClothingSimulationContext* InContext)
 	}
 
 	// Advance Sim
-	const float SubstepDeltaTime = DeltaTime / (float)NumSubsteps;
+	const int32 NumSubstepsRead = NumSubsteps;  // atomic read
+	const float SubstepDeltaTime = DeltaTime / (float)NumSubstepsRead;
 	
-	for (int32 i = 0; i < NumSubsteps; ++i)
+	for (int32 i = 0; i < NumSubstepsRead; ++i)
 	{
 		Evolution->AdvanceOneTimeStep(SubstepDeltaTime);
 	}
@@ -1550,13 +1548,11 @@ void ClothingSimulation::Simulate(IClothingSimulationContext* InContext)
 	Time = Evolution->GetTime();
 	UE_LOG(LogChaosCloth, VeryVerbose, TEXT("DeltaTime: %.6f, FilteredDeltaTime: %.6f, Time = %.6f,  MaxPhysicsDelta = %.6f"), Context->DeltaSeconds, DeltaTime, Time, FClothingSimulationCommon::MaxPhysicsDelta);
 
-#if WITH_EDITOR
 	// Update simulation time in ms (and provide an instant average instead of the value in real-time)
 	const float PrevSimulationTime = SimulationTime;  // Copy the atomic to prevent a re-read
 	const float CurrSimulationTime = (float)((FPlatformTime::Seconds() - StartTime) * 1000.);
 	static const float SimulationTimeDecay = 0.03f; // 0.03 seems to provide a good rate of update for the instant average
 	SimulationTime = PrevSimulationTime ? PrevSimulationTime + (CurrSimulationTime - PrevSimulationTime) * SimulationTimeDecay : CurrSimulationTime;
-#endif  // #if WITH_EDITOR
 
 	// Debug draw
 #if CHAOS_DEBUG_DRAW
diff --git a/Engine/Plugins/Experimental/ChaosCloth/Source/Chaos/Private/ChaosCloth/ChaosClothingSimulationInteractor.cpp b/Engine/Plugins/Experimental/ChaosCloth/Source/Chaos/Private/ChaosCloth/ChaosClothingSimulationInteractor.cpp
--- a/Engine/Plugins/Experimental/ChaosCloth/Source/Chaos/Private/ChaosCloth/ChaosClothingSimulationInteractor.cpp
+++ b/Engine/Plugins/Experimental/ChaosCloth/Source/Chaos/Private/ChaosCloth/ChaosClothingSimulationInteractor.cpp
@@ -35,6 +35,14 @@ void UChaosClothingSimulationInteractor::Sync(IClothingSimulation* InSimulation,
 		Command.Execute(ChaosSim, ChaosContext);
 	}
 	Commands.Reset();
+
+	NumCloths = InSimulation->GetNumCloths();
+	NumKinematicParticles = InSimulation->GetNumKinematicParticles();
+	NumDynamicParticles = InSimulation->GetNumDynamicParticles();
+	NumIterations = InSimulation->GetNumIterations();
+	NumSubsteps = InSimulation->GetNumSubsteps();
+	SimulationTime = InSimulation->GetSimulationTime();
+
 	bDirty = false;
 }
 
diff --git a/Engine/Plugins/Experimental/ChaosCloth/Source/Chaos/Public/ChaosCloth/ChaosClothingSimulation.h b/Engine/Plugins/Experimental/ChaosCloth/Source/Chaos/Public/ChaosCloth/ChaosClothingSimulation.h
--- a/Engine/Plugins/Experimental/ChaosCloth/Source/Chaos/Public/ChaosCloth/ChaosClothingSimulation.h
+++ b/Engine/Plugins/Experimental/ChaosCloth/Source/Chaos/Public/ChaosCloth/ChaosClothingSimulation.h
@@ -37,18 +37,20 @@ namespace Chaos
 		// bFullClothRefresh set to false during changes created using the viewport manipulators.
 		void RefreshPhysicsAsset();
 
-#if WITH_EDITOR
-		// FGCObject interface
-		virtual void AddReferencedObjects(FReferenceCollector& Collector) override;
-		// End of FGCObject interface
-
 		// IClothingSimulation interface
 		virtual int32 GetNumCloths() const { return NumCloths; }
 		virtual int32 GetNumKinematicParticles() const { return NumKinemamicParticles; }
 		virtual int32 GetNumDynamicParticles() const { return NumDynamicParticles; }
+		virtual int32 GetNumIterations() const { return NumIterations; }
+		virtual int32 GetNumSubsteps() const { return NumSubsteps; }
 		virtual float GetSimulationTime() const { return SimulationTime; }
 		// End of IClothingSimulation interface
 
+#if WITH_EDITOR
+		// FGCObject interface
+		virtual void AddReferencedObjects(FReferenceCollector& Collector) override;
+		// End of FGCObject interface
+
 		// Editor only debug draw function
 		CHAOSCLOTH_API void DebugDrawPhysMeshShaded(USkeletalMeshComponent* OwnerComponent, FPrimitiveDrawInterface* PDI) const;
 		CHAOSCLOTH_API void DebugDrawParticleIndices(USkeletalMeshComponent* OwnerComponent, FCanvas* Canvas, const FSceneView* SceneView) const;
@@ -170,7 +172,14 @@ namespace Chaos
 
 		float Time;
 		float DeltaTime;
-		int32 NumSubsteps;
+
+		// Properties that must be readable from all threads
+		TAtomic<int32> NumCloths;
+		TAtomic<int32> NumKinemamicParticles;
+		TAtomic<int32> NumDynamicParticles;
+		TAtomic<int32> NumIterations;
+		TAtomic<int32> NumSubsteps;
+		TAtomic<float> SimulationTime;
 
 		bool bUseGravityOverride;
 		FVector GravityOverride;
@@ -188,12 +197,6 @@ namespace Chaos
 		TArray<float> AngularDeltaRatios;  // Angular ratios factor applied to the reference bone transforms
 
 #if WITH_EDITOR
-		// Stats
-		TAtomic<int32> NumCloths;
-		TAtomic<int32> NumKinemamicParticles;
-		TAtomic<int32> NumDynamicParticles;
-		TAtomic<float> SimulationTime;
-
 		// Visualization material
 		UMaterial* DebugClothMaterial;
 		UMaterial* DebugClothMaterialVertex;
diff --git a/Engine/Plugins/Experimental/ChaosCloth/Source/Chaos/Public/ChaosCloth/ChaosClothingSimulationInteractor.h b/Engine/Plugins/Experimental/ChaosCloth/Source/Chaos/Public/ChaosCloth/ChaosClothingSimulationInteractor.h
--- a/Engine/Plugins/Experimental/ChaosCloth/Source/Chaos/Public/ChaosCloth/ChaosClothingSimulationInteractor.h
+++ b/Engine/Plugins/Experimental/ChaosCloth/Source/Chaos/Public/ChaosCloth/ChaosClothingSimulationInteractor.h
@@ -15,17 +15,32 @@ class CHAOSCLOTH_API UChaosClothingSimulationInteractor : public UClothingSimula
 	GENERATED_BODY()
 public:
 
-	// UClothingSimulationInteractor Interface
+	// UClothingSimulationInteractor interface
 	virtual void PhysicsAssetUpdated() override;
 	virtual void ClothConfigUpdated() override;
 	virtual void Sync(IClothingSimulation* InSimulation, IClothingSimulationContext* InContext) override;
 	virtual void SetAnimDriveSpringStiffness(float InStiffness) override;
 	virtual void EnableGravityOverride(const FVector& InVector) override;
 	virtual void DisableGravityOverride() override;
-	//////////////////////////////////////////////////////////////////////////
+
+	virtual int32 GetNumCloths() const override { return NumCloths; }
+	virtual int32 GetNumKinematicParticles() const override { return NumKinematicParticles; }
+	virtual int32 GetNumDynamicParticles() const override { return NumDynamicParticles; }
+	virtual int32 GetNumIterations() const override { return NumIterations; }
+	virtual int32 GetNumSubsteps() const override { return NumSubsteps; }
+	virtual float GetSimulationTime() const override { return SimulationTime; }
+	// End of UClothingSimulationInteractor interface
 
 private:
 
 	// Command queue processed when we hit a sync
 	TArray<ChaosClothInteractorCommand> Commands;
+
+	// Simulation infos updated during the sync
+	int32 NumCloths;
+	int32 NumKinematicParticles;
+	int32 NumDynamicParticles;
+	int32 NumIterations;
+	int32 NumSubsteps;
+	float SimulationTime;
 };
diff --git a/Engine/Source/Runtime/ClothingSystemRuntimeInterface/Public/ClothingSimulationInteractor.h b/Engine/Source/Runtime/ClothingSystemRuntimeInterface/Public/ClothingSimulationInteractor.h
--- a/Engine/Source/Runtime/ClothingSystemRuntimeInterface/Public/ClothingSimulationInteractor.h
+++ b/Engine/Source/Runtime/ClothingSystemRuntimeInterface/Public/ClothingSimulationInteractor.h
@@ -57,6 +57,36 @@ public:
 	virtual void DisableGravityOverride()
 	PURE_VIRTUAL(UClothingSimulationInteractor::DisableGravityOverride, );
 
+	// Return the number of cloths run by the simulation.
+	UFUNCTION(BlueprintCallable, Category = ClothingSimulation)
+	virtual int32 GetNumCloths() const
+	PURE_VIRTUAL(UClothingSimulationInteractor::GetNumCloths, return 0;);
+
+	// Return the number of kinematic (animated) particles.
+	UFUNCTION(BlueprintCallable, Category = ClothingSimulation)
+	virtual int32 GetNumKinematicParticles() const
+	PURE_VIRTUAL(UClothingSimulationInteractor::GetNumKinematicParticles, return 0;);
+
+	// Return the number of dynamic (simulated) particles.
+	UFUNCTION(BlueprintCallable, Category = ClothingSimulation)
+	virtual int32 GetNumDynamicParticles() const
+	PURE_VIRTUAL(UClothingSimulationInteractor::GetNumDynamicParticles, return 0;);
+
+	// Return the solver number of iterations.
+	UFUNCTION(BlueprintCallable, Category = ClothingSimulation)
+	virtual int32 GetNumIterations() const
+	PURE_VIRTUAL(UClothingSimulationInteractor::GetNumIterations, return 0;);
+
+	// Return the solver number of subdivisions.
+	UFUNCTION(BlueprintCallable, Category = ClothingSimulation)
+	virtual int32 GetNumSubsteps() const
+	PURE_VIRTUAL(UClothingSimulationInteractor::GetNumSubsteps, return 0;);
+
+	// Return the instant average simulation time in ms. 
+	UFUNCTION(BlueprintCallable, Category = ClothingSimulation)
+	virtual float GetSimulationTime() const
+	PURE_VIRTUAL(UClothingSimulationInteractor::GetSimulationTime, return 0.f;);
+
 protected:
 
 	// Intended to be called by functions on the interactor to message to the 
diff --git a/Engine/Source/Runtime/ClothingSystemRuntimeInterface/Public/ClothingSimulationInterface.h b/Engine/Source/Runtime/ClothingSystemRuntimeInterface/Public/ClothingSimulationInterface.h
--- a/Engine/Source/Runtime/ClothingSystemRuntimeInterface/Public/ClothingSimulationInterface.h
+++ b/Engine/Source/Runtime/ClothingSystemRuntimeInterface/Public/ClothingSimulationInterface.h
@@ -141,14 +141,24 @@ public:
 	 */
 	virtual void GatherStats() const {}
 
-#if WITH_EDITOR
 	/** Return the number of simulated cloths. Implementation must be thread safe. */
 	virtual int32 GetNumCloths() const { return 0;  }
 	/** Return the number of kinematic (fixed) particles. Implementation must be thread safe. */
 	virtual int32 GetNumKinematicParticles() const { return 0; }
 	/** Return the number of dynamic (simulated) particles. Implementation must be thread safe. */
 	virtual int32 GetNumDynamicParticles() const { return 0; }
+	/**
+	 * Return the number of iterations used by the solver.
+	 * This is the maximum used as an indicative value only, as this could vary between cloths.
+	 * Implementation must be thread safe.
+	 */
+	virtual int32 GetNumIterations() const { return 0; }
+	/**
+	 * Return the number of substeps used by the solver.
+	 * This is the maximum used as an indicative value only, as this could vary between cloths.
+	 * Implementation must be thread safe.
+	 */
+	virtual int32 GetNumSubsteps() const { return 0; }
 	/** Return the simulation time in ms. Implementation must be thread safe. */
 	virtual float GetSimulationTime() const { return 0.f; }
-#endif
 };
