Chaos Cloth - Add island coloring to the LongRangeAttachment constraint visualization.
Chaos Cloth - Add kinematic areas to flat shaded cloth visualization.
CIS. Fix lamda argument shadowing a member variable.

Merging CLs 14000760, 14009230, and 14009297 from //UE4/Main

diff --git a/Engine/Plugins/Experimental/ChaosCloth/Source/Chaos/Private/ChaosCloth/ChaosClothingSimulation.cpp b/Engine/Plugins/Experimental/ChaosCloth/Source/Chaos/Private/ChaosCloth/ChaosClothingSimulation.cpp
--- a/Engine/Plugins/Experimental/ChaosCloth/Source/Chaos/Private/ChaosCloth/ChaosClothingSimulation.cpp
+++ b/Engine/Plugins/Experimental/ChaosCloth/Source/Chaos/Private/ChaosCloth/ChaosClothingSimulation.cpp
@@ -1891,7 +1891,14 @@ void ClothingSimulation::DebugDrawPhysMeshShaded(USkeletalMeshComponent* OwnerCo
 
 	for (int32 Index = 0; Index < Assets.Num(); ++Index)
 	{
-		if (!Assets[Index]) { continue; }
+		const UClothingAssetCommon* const Asset = Assets[Index];
+		if (!Asset) { continue; }
+
+		const FClothLODDataCommon& AssetLodData = Asset->LodData[0];
+		const FClothPhysicalMeshData& PhysMesh = AssetLodData.PhysicalMeshData;
+		const FPointWeightMap& MaxDistances = PhysMesh.GetWeightMap(EChaosWeightMapTarget::MaxDistance);
+		if (MaxDistances.Num() == 0) { continue; }
+		const uint32 Offset = IndexToRangeMap[Index][0];
 
 		if (const TUniquePtr<Chaos::TTriangleMesh<float>>& Mesh = Meshes[Index])
 		{
@@ -1908,9 +1915,14 @@ void ClothingSimulation::DebugDrawPhysMeshShaded(USkeletalMeshComponent* OwnerCo
 				const FVector& Normal = FVector::CrossProduct(Pos2 - Pos0, Pos1 - Pos0).GetSafeNormal();
 				const FVector Tangent = ((Pos1 + Pos2) * 0.5f - Pos0).GetSafeNormal();
 
-				MeshBuilder.AddVertex(FDynamicMeshVertex(Pos0, Tangent, Normal, FVector2D(0.f, 0.f), FColor::White));
-				MeshBuilder.AddVertex(FDynamicMeshVertex(Pos1, Tangent, Normal, FVector2D(0.f, 1.f), FColor::White));
-				MeshBuilder.AddVertex(FDynamicMeshVertex(Pos2, Tangent, Normal, FVector2D(1.f, 1.f), FColor::White));
+				static const float KinematicDistanceThreshold = 0.1f;  // Must be the same value set in UI, not ideal really
+				const bool bIsKinematic0 = MaxDistances[Element.X - Offset] < KinematicDistanceThreshold;
+				const bool bIsKinematic1 = MaxDistances[Element.Y - Offset] < KinematicDistanceThreshold;
+				const bool bIsKinematic2 = MaxDistances[Element.Z - Offset] < KinematicDistanceThreshold;
+
+				MeshBuilder.AddVertex(FDynamicMeshVertex(Pos0, Tangent, Normal, FVector2D(0.f, 0.f), bIsKinematic0 ? FColor::Purple : FColor::White));
+				MeshBuilder.AddVertex(FDynamicMeshVertex(Pos1, Tangent, Normal, FVector2D(0.f, 1.f), bIsKinematic1 ? FColor::Purple : FColor::White));
+				MeshBuilder.AddVertex(FDynamicMeshVertex(Pos2, Tangent, Normal, FVector2D(1.f, 1.f), bIsKinematic2 ? FColor::Purple : FColor::White));
 				MeshBuilder.AddTriangle(VertexIndex, VertexIndex + 1, VertexIndex + 2);
 			}
 		}
@@ -2523,32 +2535,118 @@ void ClothingSimulation::DebugDrawAnimDrive(USkeletalMeshComponent* /*OwnerCompo
 
 void ClothingSimulation::DebugDrawLongRangeConstraint(USkeletalMeshComponent* /*OwnerComponent*/, FPrimitiveDrawInterface* PDI) const
 {
+	auto PseudoRandomColor = 
+		[](int32 NumRotations) -> FLinearColor
+		{
+			static const uint8 Spread = 157;  // Prime number that gives a good spread of colors without getting too similar as a rand might do.
+			uint8 Seed = Spread;
+			for (int32 i = 0; i < NumRotations; ++i)
+			{
+				Seed += Spread;
+			}
+			return FLinearColor::MakeFromHSV8(Seed, 160, 128);
+		};
+
+	int32 ColorOffset = 0;
+
 	const TPBDParticles<float, 3>& Particles = Evolution->Particles();
 
 	for (int32 Index = 0; Index < Assets.Num(); ++Index)
 	{
-		const UClothingAssetCommon* const Asset = Assets[Index];
-		if (!Asset || !LongRangeConstraints[Index]) { continue; }
+		if (!Assets[Index] || !Meshes[Index] || !LongRangeConstraints[Index]) { continue; }
 
+		// Recompute islands
+		const TMap<int32, TSet<uint32>>& PointToNeighborsMap = Meshes[Index]->GetPointToNeighborsMap();
+
+		static TArray<uint32> KinematicIndices;  // Make static to prevent constant allocations
+		KinematicIndices.Reset();
+		for (const TPair<int32, TSet<uint32>>& PointNeighbors : PointToNeighborsMap)
+		{
+			const int32 PointIndex = PointNeighbors.Key;
+			if (Particles.InvM(PointIndex) == 0.f)
+			{
+				KinematicIndices.Add(PointIndex);
+			}
+		}
+
+		const TArray<TArray<uint32>> IslandElements = TPBDLongRangeConstraints<float, 3>::ComputeIslands(PointToNeighborsMap, KinematicIndices);
+
+		// Draw constraints
 		const TArray<TArray<uint32>>& Constraints = LongRangeConstraints[Index]->GetConstraints();
-		const TArray<float>& Dists = LongRangeConstraints[Index]->GetDists();
 
 		for (int32 ConstraintIndex = 0; ConstraintIndex < Constraints.Num(); ++ConstraintIndex)
 		{
 			const TArray<uint32>& Path = Constraints[ConstraintIndex];
-			const float RefDist = Dists[ConstraintIndex];
-			const float CurDist = TPBDLongRangeConstraintsBase<float, 3>::ComputeGeodesicDistance(Particles, Path);
-			const float Offset = CurDist - RefDist;
+			const uint32 KinematicIndex = Path[0];
+			const uint32 DynamicIndex = Path[Path.Num() - 1];
+
+			// Find Island
+			int32 ColorIndex = 0;
+			for (int32 IslandIndex = 0; IslandIndex < IslandElements.Num(); ++IslandIndex)
+			{
+				if (IslandElements[IslandIndex].Find(KinematicIndex) != INDEX_NONE)  // TODO: This is O(n^2), it'll be nice to make this faster, even if it is only debug viz. Maybe binary search if the kinematic indices are ordered?
+				{
+					ColorIndex = ColorOffset + IslandIndex;
+					break;
+				}
+			}
+
+			const TVector<float, 3> Pos0 = Particles.X(KinematicIndex) + LocalSpaceLocation;
+			const TVector<float, 3> Pos1 = Particles.X(DynamicIndex) + LocalSpaceLocation;
+
+			DrawLine(PDI, Pos0, Pos1, PseudoRandomColor(ColorIndex));
+		}
 
-			const TVector<float, 3> P0 = Particles.X(Path[0]) + LocalSpaceLocation;  // Kinematic particle
-			const TVector<float, 3> P1 = Particles.X(Path[Path.Num() - 1]) + LocalSpaceLocation;  // Target particle
+		// Draw islands
+		const TArray<TVector<int32, 3>>& Elements = Meshes[Index]->GetElements();
+
+		for (int32 ElementIndex = 0; ElementIndex < Elements.Num(); ++ElementIndex)
+		{
+			const auto& Element = Elements[ElementIndex];
 
-			const TVector<float, 3> Direction = (LocalSpaceLocation + Particles.X(Path[Path.Num() - 2]) - P1).GetSafeNormal();
-			const TVector<float, 3> P2 = P1 + Direction * Offset;
+			const bool bIsKinematic0 = (Particles.InvM(Element.X) == 0.f);
+			const bool bIsKinematic1 = (Particles.InvM(Element.Y) == 0.f);
+			const bool bIsKinematic2 = (Particles.InvM(Element.Z) == 0.f);
 
-			DrawLine(PDI, P0, P1, FLinearColor(FColor::Purple));
-			DrawLine(PDI, P1, P2, FLinearColor::Black);
+			// Lookup for any kinematic point on the triangle element to use for finding the island (it doesn't matter which one, if two kinematic points are on the same triangle they have to be on the same island)
+			const int32 KinematicIndex = bIsKinematic0 ? Element.X : bIsKinematic1 ? Element.Y : bIsKinematic2 ? Element.Z : INDEX_NONE;
+			if (KinematicIndex == INDEX_NONE)
+			{
+				continue;
+			}
+
+			// Find Island Color
+			int32 ColorIndex = 0;
+			for (int32 IslandIndex = 0; IslandIndex < IslandElements.Num(); ++IslandIndex)
+			{
+				if (IslandElements[IslandIndex].Find(KinematicIndex) != INDEX_NONE)  // TODO: This is O(n^2), it'll be nice to make this faster, even if it is only debug viz. Maybe binary search if the kinematic indices are ordered?
+				{
+					ColorIndex = ColorOffset + IslandIndex;
+					break;
+				}
+			}
+			const FLinearColor Color = PseudoRandomColor(ColorIndex);
+
+			const FVector Pos0 = LocalSpaceLocation + Particles.X(Element.X);
+			const FVector Pos1 = LocalSpaceLocation + Particles.X(Element.Y);
+			const FVector Pos2 = LocalSpaceLocation + Particles.X(Element.Z);
+
+			if (bIsKinematic0 && bIsKinematic1)
+			{
+				DrawLine(PDI, Pos0, Pos1, Color);
+			}
+			if (bIsKinematic1 && bIsKinematic2)
+			{
+				DrawLine(PDI, Pos1, Pos2, Color);
+			}
+			if (bIsKinematic2 && bIsKinematic0)
+			{
+				DrawLine(PDI, Pos2, Pos0, Color);
+			}
 		}
+
+		// Rotate the colors for each cloths
+		ColorOffset += IslandElements.Num();
 	}
 }
 
diff --git a/Engine/Source/Runtime/Experimental/Chaos/Private/Chaos/PBDLongRangeConstraintsBase.cpp b/Engine/Source/Runtime/Experimental/Chaos/Private/Chaos/PBDLongRangeConstraintsBase.cpp
--- a/Engine/Source/Runtime/Experimental/Chaos/Private/Chaos/PBDLongRangeConstraintsBase.cpp
+++ b/Engine/Source/Runtime/Experimental/Chaos/Private/Chaos/PBDLongRangeConstraintsBase.cpp
@@ -37,7 +37,6 @@ TPBDLongRangeConstraintsBase<T, d>::TPBDLongRangeConstraintsBase(
 
 template<class T, int d>
 TArray<TArray<uint32>> TPBDLongRangeConstraintsBase<T, d>::ComputeIslands(
-    const TDynamicParticles<T, d>& InParticles,
     const TMap<int32, TSet<uint32>>& PointToNeighbors,
     const TArray<uint32>& KinematicParticles)
 {
@@ -129,7 +128,7 @@ void TPBDLongRangeConstraintsBase<T, d>::ComputeEuclidianConstraints(
 	}
 
 	// Compute the islands of kinematic particles
-	const TArray<TArray<uint32>> IslandElements = ComputeIslands(InParticles, PointToNeighbors, KinematicParticles);
+	const TArray<TArray<uint32>> IslandElements = ComputeIslands(PointToNeighbors, KinematicParticles);
 	int32 NumTotalIslandElements = 0;
 	for(const auto& Elements : IslandElements)
 		NumTotalIslandElements += Elements.Num();
@@ -206,7 +205,7 @@ void TPBDLongRangeConstraintsBase<T, d>::ComputeGeodesicConstraints(
 			KinematicParticles.Add(i);
 		}
 	}
-	TArray<TArray<uint32>> IslandElements = ComputeIslands(InParticles, PointToNeighbors, KinematicParticles);
+	TArray<TArray<uint32>> IslandElements = ComputeIslands(PointToNeighbors, KinematicParticles);
 	// Store distances for all adjacent vertices
 	TMap<TVector<uint32, 2>, T> Distances;
 	for (const uint32 i : UsedIndices)
diff --git a/Engine/Source/Runtime/Experimental/Chaos/Public/Chaos/PBDLongRangeConstraintsBase.h b/Engine/Source/Runtime/Experimental/Chaos/Public/Chaos/PBDLongRangeConstraintsBase.h
--- a/Engine/Source/Runtime/Experimental/Chaos/Public/Chaos/PBDLongRangeConstraintsBase.h
+++ b/Engine/Source/Runtime/Experimental/Chaos/Public/Chaos/PBDLongRangeConstraintsBase.h
@@ -78,8 +78,9 @@ class CHAOS_API TPBDLongRangeConstraintsBase
 		return distance;
 	}
 
+	static TArray<TArray<uint32>> ComputeIslands(const TMap<int32, TSet<uint32>>& PointToNeighbors, const TArray<uint32>& KinematicParticles);
+
   protected:
-	static TArray<TArray<uint32>> ComputeIslands(const TDynamicParticles<T, d>& InParticles, const TMap<int32, TSet<uint32>>& PointToNeighbors,/*const TTriangleMesh<T>& Mesh,*/ const TArray<uint32>& KinematicParticles);
 	void ComputeEuclidianConstraints(const TDynamicParticles<T, d>& InParticles, const TMap<int32, TSet<uint32>>& PointToNeighbors,/*const TTriangleMesh<T>& Mesh,*/ const int32 NumberOfAttachments);
 	void ComputeGeodesicConstraints(const TDynamicParticles<T, d>& InParticles, const TMap<int32, TSet<uint32>>& PointToNeighbors,/*const TTriangleMesh<T>& Mesh,*/ const int32 NumberOfAttachments);
 
