Chaos Cloth - LOD/graph node refactor.
- Add multi LOD support.
- Add simulation Mesh, Cloth, Collider, and Solver graph node ready classes.
- Fix backstop, and improve backstop debug draw.
- Add Anim debug draw.
- Output cloth sim data in local space.

diff --git a/Engine/Plugins/Experimental/ChaosCloth/Source/Chaos/Private/ChaosCloth/ChaosClothConfig.cpp b/Engine/Plugins/Experimental/ChaosCloth/Source/Chaos/Private/ChaosCloth/ChaosClothConfig.cpp
--- a/Engine/Plugins/Experimental/ChaosCloth/Source/Chaos/Private/ChaosCloth/ChaosClothConfig.cpp
+++ b/Engine/Plugins/Experimental/ChaosCloth/Source/Chaos/Private/ChaosCloth/ChaosClothConfig.cpp
@@ -125,6 +125,17 @@ void UChaosClothConfig::PostLoad()
 	}
 }
 
+float UChaosClothConfig::GetMassValue() const
+{
+	switch (MassMode)
+	{
+	default:
+	case EClothMassMode::Density: return Density;
+	case EClothMassMode::TotalMass: return TotalMass;
+	case EClothMassMode::UniformMass: return UniformMass;
+	}
+}
+
 UChaosClothSharedSimConfig::UChaosClothSharedSimConfig()
 {}
 
diff --git a/Engine/Plugins/Experimental/ChaosCloth/Source/Chaos/Private/ChaosCloth/ChaosClothConstraints.cpp b/Engine/Plugins/Experimental/ChaosCloth/Source/Chaos/Private/ChaosCloth/ChaosClothConstraints.cpp
new file mode 100644
--- /dev/null
+++ b/Engine/Plugins/Experimental/ChaosCloth/Source/Chaos/Private/ChaosCloth/ChaosClothConstraints.cpp
@@ -0,0 +1,420 @@
+// Copyright Epic Games, Inc. All Rights Reserved.
+#include "ChaosCloth/ChaosClothConstraints.h"
+
+#include "Chaos/PBDSpringConstraints.h"
+#include "Chaos/XPBDSpringConstraints.h"
+#include "Chaos/PBDBendingConstraints.h"
+#include "Chaos/PBDAxialSpringConstraints.h"
+#include "Chaos/XPBDAxialSpringConstraints.h"
+#include "Chaos/PBDVolumeConstraint.h"
+#include "Chaos/PBDLongRangeConstraints.h"
+#include "Chaos/XPBDLongRangeConstraints.h"
+#include "Chaos/PBDSphericalConstraint.h"
+#include "Chaos/PBDAnimDriveConstraint.h"
+#include "Chaos/PBDShapeConstraints.h"
+
+#include "Chaos/PBDEvolution.h"
+
+using namespace Chaos;
+
+FClothConstraints::FClothConstraints()
+	: Evolution(nullptr)
+	, AnimationPositions(nullptr)
+	, AnimationNormals(nullptr)
+	, ParticleOffset(0)
+	, NumParticles(0)
+	, ConstraintInitOffset(INDEX_NONE)
+	, ConstraintRuleOffset(INDEX_NONE)
+	, NumConstraintInits(0)
+	, NumConstraintRules(0)
+	, MaxDistancesMultiplier(1.f)
+	, AnimDriveSpringStiffness(0.f)
+{
+}
+
+FClothConstraints::~FClothConstraints()
+{
+}
+
+void FClothConstraints::Initialize(
+	TPBDEvolution<float, 3>* InEvolution,
+	const TArray<TVector<float, 3>>& InAnimationPositions,
+	const TArray<TVector<float, 3>>& InAnimationNormals,
+	int32 InParticleOffset,
+	int32 InNumParticles)
+{
+	Evolution = InEvolution;
+	AnimationPositions = &InAnimationPositions;
+	AnimationNormals = &InAnimationNormals;
+	ParticleOffset = InParticleOffset;
+	NumParticles = InNumParticles;
+}
+
+void FClothConstraints::Enable(bool bEnable)
+{
+	check(Evolution);
+	if (ConstraintInitOffset != INDEX_NONE)
+	{
+		Evolution->ActivateConstraintInitRange(ConstraintInitOffset, bEnable);
+	}
+	if (ConstraintRuleOffset != INDEX_NONE)
+	{
+		Evolution->ActivateConstraintRuleRange(ConstraintRuleOffset, bEnable);
+	}
+}
+
+void FClothConstraints::CreateRules()
+{
+	check(Evolution);
+	check(ConstraintInitOffset == INDEX_NONE)
+	if (NumConstraintInits)
+	{
+		ConstraintInitOffset = Evolution->AddConstraintInitRange(NumConstraintInits, false);
+	}
+	check(ConstraintRuleOffset == INDEX_NONE)
+	if (NumConstraintRules)
+	{
+		ConstraintRuleOffset = Evolution->AddConstraintRuleRange(NumConstraintRules, false);
+	}
+
+	TFunction<void()>* const ConstraintInits = Evolution->ConstraintInits().GetData() + ConstraintInitOffset;
+	TFunction<void(TPBDParticles<float, 3>&, const float)>* const ConstraintRules = Evolution->ConstraintRules().GetData() + ConstraintRuleOffset;
+
+	int32 ConstraintInitIndex = 0;
+	int32 ConstraintRuleIndex = 0;
+
+	if (XTwoEdgeConstraints)
+	{
+		ConstraintInits[ConstraintInitIndex++] =
+			[this]()
+			{
+				XTwoEdgeConstraints->Init();
+			};
+
+		ConstraintRules[ConstraintRuleIndex++] = 
+			[this](TPBDParticles<float, 3>& InParticles, const float Dt)
+			{
+				XTwoEdgeConstraints->Apply(InParticles, Dt);
+			};
+	}
+	if (TwoEdgeConstraints)
+	{
+		ConstraintRules[ConstraintRuleIndex++] = 
+			[this](TPBDParticles<float, 3>& InParticles, const float Dt)
+			{
+				TwoEdgeConstraints->Apply(InParticles, Dt);
+			};
+	}
+	if (XThreeEdgeConstraints)
+	{
+		ConstraintInits[ConstraintInitIndex++] =
+			[this]()
+			{
+				XThreeEdgeConstraints->Init();
+			};
+		ConstraintRules[ConstraintRuleIndex++] = 
+			[this](TPBDParticles<float, 3>& InParticles, const float Dt)
+			{
+				XThreeEdgeConstraints->Apply(InParticles, Dt);
+			};
+	}
+	if (ThreeEdgeConstraints)
+	{
+		ConstraintRules[ConstraintRuleIndex++] = 
+			[this](TPBDParticles<float, 3>& InParticles, const float Dt)
+			{
+				ThreeEdgeConstraints->Apply(InParticles, Dt);
+			};
+	}
+	if (XBendingConstraints)
+	{
+		ConstraintInits[ConstraintInitIndex++] =
+			[this]()
+			{
+				XBendingConstraints->Init();
+			};
+		ConstraintRules[ConstraintRuleIndex++] = 
+			[this](TPBDParticles<float, 3>& InParticles, const float Dt)
+			{
+				XBendingConstraints->Apply(InParticles, Dt);
+			};
+	}
+	if (BendingConstraints)
+	{
+		ConstraintRules[ConstraintRuleIndex++] = 
+			[this](TPBDParticles<float, 3>& InParticles, const float Dt)
+			{
+				BendingConstraints->Apply(InParticles, Dt);
+			};
+	}
+	if (BendingElementConstraints)
+	{
+		ConstraintRules[ConstraintRuleIndex++] = 
+			[this](TPBDParticles<float, 3>& InParticles, const float Dt)
+			{
+				BendingElementConstraints->Apply(InParticles, Dt);
+			};
+	}
+	if (XAreaConstraints)
+	{
+		ConstraintInits[ConstraintInitIndex++] =
+			[this]()
+			{
+				XAreaConstraints->Init();
+			};
+		ConstraintRules[ConstraintRuleIndex++] = 
+			[this](TPBDParticles<float, 3>& InParticles, const float Dt)
+			{
+				XAreaConstraints->Apply(InParticles, Dt);
+			};
+	}
+	if (AreaConstraints)
+	{
+		ConstraintRules[ConstraintRuleIndex++] = 
+			[this](TPBDParticles<float, 3>& InParticles, const float Dt)
+			{
+				AreaConstraints->Apply(InParticles, Dt);
+			};
+	}
+	if (ThinShellVolumeConstraints)
+	{
+		ConstraintRules[ConstraintRuleIndex++] = 
+			[this](TPBDParticles<float, 3>& InParticles, const float Dt)
+			{
+				ThinShellVolumeConstraints->Apply(InParticles, Dt);
+			};
+	}
+	if (VolumeConstraints)
+	{
+		ConstraintRules[ConstraintRuleIndex++] =
+			[this](TPBDParticles<float, 3>& InParticles, const float Dt)
+			{
+				VolumeConstraints->Apply(InParticles, Dt);
+			};
+	}
+	if (XLongRangeConstraints)
+	{
+		ConstraintInits[ConstraintInitIndex++] =
+			[this]()
+			{
+				XLongRangeConstraints->Init();
+			};
+		ConstraintRules[ConstraintRuleIndex++] = 
+			[this](TPBDParticles<float, 3>& InParticles, const float Dt)
+			{
+				XLongRangeConstraints->Apply(InParticles, Dt);
+			};
+	}
+	if (LongRangeConstraints)
+	{
+		ConstraintRules[ConstraintRuleIndex++] = 
+			[this](TPBDParticles<float, 3>& InParticles, const float Dt)
+			{
+				LongRangeConstraints->Apply(InParticles, Dt);
+			};
+	}
+	if (MaximumDistanceConstraints)
+	{
+		ConstraintRules[ConstraintRuleIndex++] = 
+			[this](TPBDParticles<float, 3>& InParticles, const float Dt)
+			{
+				MaximumDistanceConstraints->SetSphereRadiiMultiplier(FMath::Max(0.f, MaxDistancesMultiplier));
+				MaximumDistanceConstraints->Apply(InParticles, Dt);
+			};
+	}
+	if (BackstopConstraints)
+	{
+		ConstraintRules[ConstraintRuleIndex++] = 
+			[this](TPBDParticles<float, 3>& InParticles, const float Dt)
+			{
+				BackstopConstraints->Apply(InParticles, Dt);
+			};
+	}
+	if (AnimDriveConstraints)
+	{
+		ConstraintRules[ConstraintRuleIndex++] = 
+			[this](TPBDParticles<float, 3>& InParticles, const float Dt)
+			{
+				AnimDriveConstraints->SetSpringStiffness(FMath::Clamp(AnimDriveSpringStiffness, 0.f, 1.f));
+				AnimDriveConstraints->Apply(InParticles, Dt);
+			};
+	}
+	if (ShapeConstraints)
+	{
+		ConstraintRules[ConstraintRuleIndex++] = 
+			[this](TPBDParticles<float, 3>& InParticles, const float Dt)
+			{
+				ShapeConstraints->Apply(InParticles, Dt);
+			};
+	}
+	check(ConstraintInitIndex == NumConstraintInits);
+	check(ConstraintRuleIndex == NumConstraintRules);
+}
+
+
+void FClothConstraints::SetEdgeConstraints(TArray<TVector<int32, 2>>&& Edges, float EdgeStiffness, bool bUseXPBDConstraints)
+{
+	check(Evolution);
+	check(EdgeStiffness > 0.f && EdgeStiffness <= 1.f);
+
+	if (bUseXPBDConstraints)
+	{
+		XTwoEdgeConstraints = MakeShared<TXPBDSpringConstraints<float, 3>>(Evolution->Particles(), MoveTemp(Edges), EdgeStiffness);
+		++NumConstraintInits;
+	}
+	else
+	{
+		TwoEdgeConstraints = MakeShared<FPBDSpringConstraints>(Evolution->Particles(), MoveTemp(Edges), EdgeStiffness);
+	}
+	++NumConstraintRules;
+}
+
+void FClothConstraints::SetEdgeConstraints(const TArray<TVector<int32, 3>>& SurfaceElements, float EdgeStiffness, bool bUseXPBDConstraints)
+{
+	check(Evolution);
+	check(EdgeStiffness > 0.f && EdgeStiffness <= 1.f);
+
+	if (bUseXPBDConstraints)
+	{
+		XThreeEdgeConstraints = MakeShared<TXPBDSpringConstraints<float, 3>>(Evolution->Particles(), SurfaceElements, EdgeStiffness);
+		++NumConstraintInits;
+	}
+	else
+	{
+		ThreeEdgeConstraints = MakeShared<FPBDSpringConstraints>(Evolution->Particles(), SurfaceElements, EdgeStiffness);
+	}
+	++NumConstraintRules;
+}
+
+void FClothConstraints::SetBendingConstraints(TArray<TVector<int32, 2>>&& Edges, float BendingStiffness, bool bUseXPBDConstraints)
+{
+	check(Evolution);
+
+	if (bUseXPBDConstraints)
+	{
+		XBendingConstraints = MakeShared<TXPBDSpringConstraints<float, 3>>(Evolution->Particles(), MoveTemp(Edges), BendingStiffness);
+		++NumConstraintInits;
+	}
+	else
+	{
+		BendingConstraints = MakeShared<FPBDSpringConstraints>(Evolution->Particles(), MoveTemp(Edges), BendingStiffness);
+	}
+	++NumConstraintRules;
+}
+
+void FClothConstraints::SetBendingConstraints(TArray<TVector<int32, 4>>&& BendingElements, float BendingStiffness)
+{
+	check(Evolution);
+	check(BendingStiffness > 0.f && BendingStiffness <= 1.f);
+
+	BendingElementConstraints = MakeShared<TPBDBendingConstraints<float>>(Evolution->Particles(), MoveTemp(BendingElements), BendingStiffness);
+	++NumConstraintRules;
+}
+
+void FClothConstraints::SetAreaConstraints(TArray<TVector<int32, 3>>&& SurfaceElements, float AreaStiffness, bool bUseXPBDConstraints)
+{
+	check(Evolution);
+	check(AreaStiffness > 0.f && AreaStiffness <= 1.f);
+
+	if (bUseXPBDConstraints)
+	{
+		XAreaConstraints = MakeShared<TXPBDAxialSpringConstraints<float, 3>>(Evolution->Particles(), MoveTemp(SurfaceElements), AreaStiffness);
+		++NumConstraintInits;
+	}
+	else
+	{
+		AreaConstraints = MakeShared<FPBDAxialSpringConstraints>(Evolution->Particles(), MoveTemp(SurfaceElements), AreaStiffness);
+	}
+	++NumConstraintRules;
+}
+
+void FClothConstraints::SetVolumeConstraints(TArray<TVector<int32, 2>>&& DoubleBendingEdges, float VolumeStiffness)
+{
+	check(Evolution);
+	check(VolumeStiffness > 0.f && VolumeStiffness <= 1.f);
+
+	ThinShellVolumeConstraints = MakeShared<FPBDSpringConstraints>(Evolution->Particles(), MoveTemp(DoubleBendingEdges), VolumeStiffness);
+	++NumConstraintRules;
+}
+
+void FClothConstraints::SetVolumeConstraints(TArray<TVector<int32, 3>>&& SurfaceElements, float VolumeStiffness)
+{
+	check(Evolution);
+	check(VolumeStiffness > 0.f && VolumeStiffness <= 1.f);
+
+	VolumeConstraints = MakeShared<TPBDVolumeConstraint<float>>(Evolution->Particles(), MoveTemp(SurfaceElements), VolumeStiffness);
+	++NumConstraintRules;
+}
+
+void FClothConstraints::SetLongRangeConstraints(const TMap<int32, TSet<uint32>>& PointToNeighborsMap, float StrainLimitingStiffness, float LimitScale, bool bUseGeodesicDistance, bool bUseXPBDConstraints)
+{
+	check(Evolution);
+	check(StrainLimitingStiffness > 0.f && StrainLimitingStiffness <= 1.f);
+
+	if (bUseXPBDConstraints)
+	{
+		XLongRangeConstraints = MakeShared<TXPBDLongRangeConstraints<float, 3>>(
+			Evolution->Particles(),
+			PointToNeighborsMap, 
+			10, // The max number of connected neighbors per particle.
+			StrainLimitingStiffness);  // TODO: Add LimitScale to the XPBD constraint
+		++NumConstraintInits;
+	}
+	else
+	{
+		LongRangeConstraints = MakeShared<TPBDLongRangeConstraints<float, 3>>(
+			Evolution->Particles(),
+			PointToNeighborsMap,
+			10, // The max number of connected neighbors per particle.
+			StrainLimitingStiffness,
+			LimitScale,
+			bUseGeodesicDistance);
+	}
+	++NumConstraintRules;
+}
+
+void FClothConstraints::SetMaximumDistanceConstraints(const TConstArrayView<float>& MaxDistances)
+{
+	MaximumDistanceConstraints = MakeShared<TPBDSphericalConstraint<float, 3>>(
+		ParticleOffset,
+		NumParticles,
+		*AnimationPositions,
+		MaxDistances);
+	++NumConstraintRules;
+}
+
+void FClothConstraints::SetBackstopConstraints(const TConstArrayView<float>& BackstopDistances, const TConstArrayView<float>& BackstopRadiuses)
+{
+	BackstopConstraints = MakeShared<TPBDSphericalBackstopConstraint<float, 3>>(
+		ParticleOffset,
+		NumParticles,
+		*AnimationPositions,
+		*AnimationNormals,
+		BackstopRadiuses,
+		BackstopDistances);
+	++NumConstraintRules;
+}
+
+void FClothConstraints::SetAnimDriveConstraints(const TConstArrayView<float>& AnimDriveMultipliers)
+{
+	AnimDriveConstraints = MakeShared<TPBDAnimDriveConstraint<float, 3>>(
+		ParticleOffset,
+		NumParticles,
+		*AnimationPositions,
+		AnimDriveMultipliers);
+	++NumConstraintRules;
+}
+
+void FClothConstraints::SetShapeTargetConstraints(float ShapeTargetStiffness)
+{
+	// TODO: Review this constraint. Currently does nothing more than the anim drive with less controls
+	check(ShapeTargetStiffness > 0.f && ShapeTargetStiffness <= 1.f);
+
+	ShapeConstraints = MakeShared<TPBDShapeConstraints<float, 3>>(
+		ParticleOffset,
+		NumParticles,
+		*AnimationPositions,
+		*AnimationPositions,
+		ShapeTargetStiffness);
+	++NumConstraintRules;
+}
