Chaos Cloth - LOD/graph node refactor.

diff --git a/Engine/Plugins/Experimental/ChaosCloth/Source/Chaos/Private/ChaosCloth/ChaosClothingSimulation.cpp b/Engine/Plugins/Experimental/ChaosCloth/Source/Chaos/Private/ChaosCloth/ChaosClothingSimulation.cpp
--- a/Engine/Plugins/Experimental/ChaosCloth/Source/Chaos/Private/ChaosCloth/ChaosClothingSimulation.cpp
+++ b/Engine/Plugins/Experimental/ChaosCloth/Source/Chaos/Private/ChaosCloth/ChaosClothingSimulation.cpp
@@ -2,72 +2,52 @@
 #include "ChaosCloth/ChaosClothingSimulation.h"
 #include "ChaosCloth/ChaosClothPrivate.h"
 
-#include "Async/ParallelFor.h"
+#include "ClothingSimulation.h"
 #include "ClothingAsset.h"
-#include "ClothingSimulation.h" // ClothingSystemRuntimeInterface
-#include "Utils/ClothingMeshUtils.h" // ClothingSystemRuntimeCommon
-#include "Components/SkeletalMeshComponent.h"
-#include "Materials/Material.h"
-#include "PhysicsEngine/PhysicsAsset.h"
-#include "Chaos/TriangleMesh.h"
-#include "Chaos/Box.h"
-#include "Chaos/Capsule.h"
-#include "Chaos/Cylinder.h"
-#include "Chaos/ImplicitObjectIntersection.h"
-#include "Chaos/ImplicitObjectUnion.h"
-#include "Chaos/Levelset.h"
-#include "Chaos/PBDEvolution.h"
-#include "Chaos/PBDAnimDriveConstraint.h"
-#include "Chaos/PBDAxialSpringConstraints.h"
-#include "Chaos/PBDBendingConstraints.h"
-#include "Chaos/PBDLongRangeConstraints.h"
-#include "Chaos/PBDParticles.h"
-#include "Chaos/PBDSphericalConstraint.h"
-#include "Chaos/PBDSpringConstraints.h"
-#include "Chaos/PBDVolumeConstraint.h"
-#include "Chaos/PBDShapeConstraints.h"
-#include "Chaos/PerParticleGravity.h"
-#include "Chaos/Plane.h"
-#include "Chaos/Sphere.h"
-#include "Chaos/TaperedCylinder.h"
-#include "Chaos/Convex.h"
-#include "Chaos/Transform.h"
-#include "Chaos/Utilities.h"
-#include "Chaos/Vector.h"
 #include "ChaosCloth/ChaosClothConfig.h"
 #include "ChaosCloth/ChaosWeightMapTarget.h"
-#include "ChaosCloth/ChaosClothPrivate.h"
 
-#include "Chaos/XPBDLongRangeConstraints.h"
-#include "Chaos/XPBDSpringConstraints.h"
-#include "Chaos/XPBDAxialSpringConstraints.h"
+#include "ChaosCloth/ChaosClothingSimulationSolver.h"
+#include "ChaosCloth/ChaosClothingSimulationMesh.h"
+#include "ChaosCloth/ChaosClothingSimulationCloth.h"
+#include "ChaosCloth/ChaosClothingSimulationCollider.h"
 
-#if PHYSICS_INTERFACE_PHYSX && !PLATFORM_LUMIN && !PLATFORM_ANDROID
-#include "PhysXIncludes.h"
-#endif
+#if WITH_EDITOR || CHAOS_DEBUG_DRAW
+#include "Chaos/Sphere.h"
+#include "Chaos/TaperedCylinder.h"
+#include "Chaos/Convex.h"
+#include "Chaos/PBDSphericalConstraint.h"
+#include "Chaos/PBDAnimDriveConstraint.h"
+#include "Chaos/PBDLongRangeConstraints.h"
+#endif  // #if WITH_EDITOR || CHAOS_DEBUG_DRAW
 
 #if WITH_EDITOR
+#include "Materials/Material.h"
 #include "Engine/Canvas.h"  // For debug draw text
 #include "CanvasItem.h"     //
 #include "Engine/Engine.h"  //
-#endif
+#endif  // #if WITH_EDITOR
 
 #if CHAOS_DEBUG_DRAW
 #include "Chaos/DebugDrawQueue.h"
 #include "HAL/IConsoleManager.h"
+#endif  // #if CHAOS_DEBUG_DRAW
 
+#if CHAOS_DEBUG_DRAW
 namespace ChaosClothingSimulationConsoleVariables
 {
 	TAutoConsoleVariable<bool> CVarDebugDrawLocalSpace      (TEXT("p.ChaosCloth.DebugDrawLocalSpace"          ), false, TEXT("Whether to debug draw the Chaos Cloth local space"), ECVF_Cheat);
 	TAutoConsoleVariable<bool> CVarDebugDrawBounds          (TEXT("p.ChaosCloth.DebugDrawBounds"              ), false, TEXT("Whether to debug draw the Chaos Cloth bounds"), ECVF_Cheat);
 	TAutoConsoleVariable<bool> CVarDebugDrawGravity         (TEXT("p.ChaosCloth.DebugDrawGravity"             ), false, TEXT("Whether to debug draw the Chaos Cloth gravity acceleration vector"), ECVF_Cheat);
 	TAutoConsoleVariable<bool> CVarDebugDrawPhysMeshWired   (TEXT("p.ChaosCloth.DebugDrawPhysMeshWired"       ), false, TEXT("Whether to debug draw the Chaos Cloth wireframe meshes"), ECVF_Cheat);
+	TAutoConsoleVariable<bool> CVarDebugDrawAnimMeshWired   (TEXT("p.ChaosCloth.DebugDrawAnimMeshWired"       ), false, TEXT("Whether to debug draw the animated/kinematic Cloth wireframe meshes"), ECVF_Cheat);
 	TAutoConsoleVariable<bool> CVarDebugPointNormals        (TEXT("p.ChaosCloth.DebugDrawPointNormals"        ), false, TEXT("Whether to debug draw the Chaos Cloth point normals"), ECVF_Cheat);
 	TAutoConsoleVariable<bool> CVarDebugInversedPointNormals(TEXT("p.ChaosCloth.DebugDrawInversedPointNormals"), false, TEXT("Whether to debug draw the Chaos Cloth inversed point normals"), ECVF_Cheat);
 	TAutoConsoleVariable<bool> CVarDebugFaceNormals         (TEXT("p.ChaosCloth.DebugDrawFaceNormals"         ), false, TEXT("Whether to debug draw the Chaos Cloth face normals"), ECVF_Cheat);
 	TAutoConsoleVariable<bool> CVarDebugInversedFaceNormals (TEXT("p.ChaosCloth.DebugDrawInversedFaceNormals" ), false, TEXT("Whether to debug draw the Chaos Cloth inversed face normals"), ECVF_Cheat);
 	TAutoConsoleVariable<bool> CVarDebugCollision           (TEXT("p.ChaosCloth.DebugDrawCollision"           ), false, TEXT("Whether to debug draw the Chaos Cloth collisions"), ECVF_Cheat);
 	TAutoConsoleVariable<bool> CVarDebugBackstops           (TEXT("p.ChaosCloth.DebugDrawBackstops"           ), false, TEXT("Whether to debug draw the Chaos Cloth backstops"), ECVF_Cheat);
+	TAutoConsoleVariable<bool> CVarDebugBackstopDistances   (TEXT("p.ChaosCloth.DebugDrawBackstopDistances"   ), false, TEXT("Whether to debug draw the Chaos Cloth backstop distances"), ECVF_Cheat);
 	TAutoConsoleVariable<bool> CVarDebugMaxDistances        (TEXT("p.ChaosCloth.DebugDrawMaxDistances"        ), false, TEXT("Whether to debug draw the Chaos Cloth max distances"), ECVF_Cheat);
 	TAutoConsoleVariable<bool> CVarDebugAnimDrive           (TEXT("p.ChaosCloth.DebugDrawAnimDrive"           ), false, TEXT("Whether to debug draw the Chaos Cloth anim drive"), ECVF_Cheat);
 	TAutoConsoleVariable<bool> CVarDebugLongRangeConstraint (TEXT("p.ChaosCloth.DebugDrawLongRangeConstraint" ), false, TEXT("Whether to debug draw the Chaos Cloth long range constraint (aka tether constraint)"), ECVF_Cheat);
@@ -78,1487 +58,233 @@ namespace ChaosClothingSimulationConsoleVariables
 using namespace Chaos;
 
 DECLARE_CYCLE_STAT(TEXT("Chaos Cloth Simulate"), STAT_ChaosClothSimulate, STATGROUP_ChaosCloth);
-DECLARE_CYCLE_STAT(TEXT("Chaos Cloth Get Animation Data"), STAT_ChaosClothGetAnimationData, STATGROUP_ChaosCloth);
-DECLARE_CYCLE_STAT(TEXT("Chaos Cloth Update Collision Transforms"), STAT_ChaosClothUpdateCollisionTransforms, STATGROUP_ChaosCloth);
+DECLARE_CYCLE_STAT(TEXT("Chaos Cloth Create Actor"), STAT_ChaosClothCreateActor, STATGROUP_ChaosCloth);
+DECLARE_CYCLE_STAT(TEXT("Chaos Cloth Get Simulation Data"), STAT_ChaosClothGetSimulationData, STATGROUP_ChaosCloth);
 
 // Default parameters, will be overwritten when cloth assets are loaded
 namespace ChaosClothingSimulationDefault
 {
 	static const FVector Gravity(0.f, 0.f, -980.665f);
-	static const int32 NumIterations = 1;
-	static const int32 NumSubsteps = 1;
-	static const float SelfCollisionThickness = 2.f;
-	static const float CollisionThickness = 1.2f;
-	static const float FrictionCoefficient = 0.2f;
-	static const float DampingCoefficient = 0.01f;
-	static const float WorldScale = 100.f;  // World is in cm, but values like wind speed and density are in SI unit and relates to m.
+	static const float MaxDistancesMultipliers = 1.f;
+	static const float AnimDriveSpringStiffness = 1.f;
 }
 
-ClothingSimulation::ClothingSimulation()
+FClothingSimulation::FClothingSimulation()
 	: ClothSharedSimConfig(nullptr)
-	, ExternalCollisionsOffset(0)
-	, NumIterations(ChaosClothingSimulationDefault::NumIterations)
-	, NumSubsteps(ChaosClothingSimulationDefault::NumSubsteps)
+	, bUseLocalSpaceSimulation(false)
 	, bUseGravityOverride(false)
 	, GravityOverride(ChaosClothingSimulationDefault::Gravity)
-	, bUseLocalSpaceSimulation(false)
-	, LocalSpaceLocation(FVector::ZeroVector)
+	, MaxDistancesMultipliers(ChaosClothingSimulationDefault::MaxDistancesMultipliers)
+	, AnimDriveSpringStiffness(ChaosClothingSimulationDefault::AnimDriveSpringStiffness)
 {
-	ResetStats();
-
 #if WITH_EDITOR
 	DebugClothMaterial = LoadObject<UMaterial>(nullptr, TEXT("/Engine/EditorMaterials/Cloth/CameraLitDoubleSided.CameraLitDoubleSided"), nullptr, LOAD_None, nullptr);  // LOAD_EditorOnly
 	DebugClothMaterialVertex = LoadObject<UMaterial>(nullptr, TEXT("/Engine/EditorMaterials/WidgetVertexColorMaterial"), nullptr, LOAD_None, nullptr);  // LOAD_EditorOnly
 #endif  // #if WITH_EDITOR
 }
 
-ClothingSimulation::~ClothingSimulation()
+FClothingSimulation::~FClothingSimulation()
 {}
 
-void ClothingSimulation::Initialize()
+void FClothingSimulation::Initialize()
 {
-    Chaos::TPBDParticles<float, 3> LocalParticles;
-    Chaos::TKinematicGeometryClothParticles<float, 3> TRigidParticles;
-    Evolution.Reset(
-		new Chaos::TPBDEvolution<float, 3>(
-			MoveTemp(LocalParticles),
-			MoveTemp(TRigidParticles),
-			{}, // CollisionTriangles
-			ChaosClothingSimulationDefault::NumIterations,
-			ChaosClothingSimulationDefault::CollisionThickness,
-			ChaosClothingSimulationDefault::SelfCollisionThickness,
-			ChaosClothingSimulationDefault::FrictionCoefficient,
-			ChaosClothingSimulationDefault::DampingCoefficient));
-    Evolution->CollisionParticles().AddArray(&BoneIndices);
-	Evolution->CollisionParticles().AddArray(&BaseTransforms);
-
-    Evolution->SetKinematicUpdateFunction(
-		[this](Chaos::TPBDParticles<float, 3>& ParticlesInput, const float Dt, const float LocalTime, const int32 Index)
-		{
-			if (!OldAnimationPositions.IsValidIndex(Index) || ParticlesInput.InvM(Index) > 0)
-				return;
-			const float Alpha = (LocalTime - Time) / DeltaTime;
-			ParticlesInput.X(Index) = Alpha * AnimationPositions[Index] + (1.f - Alpha) * OldAnimationPositions[Index];
-		});
-
-	Evolution->SetCollisionKinematicUpdateFunction(
-//		[&](Chaos::TKinematicGeometryParticles<float, 3>& ParticlesInput, const float Dt, const float LocalTime, const int32 Index)
-		[this](Chaos::TKinematicGeometryClothParticles<float, 3>& ParticlesInput, const float Dt, const float LocalTime, const int32 Index)
-		{
-			checkSlow(Dt > SMALL_NUMBER && DeltaTime > SMALL_NUMBER);
-			const float Alpha = (LocalTime - Time) / DeltaTime;
-			const Chaos::TVector<float, 3> NewX =
-				Alpha * CollisionTransforms[Index].GetTranslation() + (1.f - Alpha) * OldCollisionTransforms[Index].GetTranslation();
-			ParticlesInput.V(Index) = (NewX - ParticlesInput.X(Index)) / Dt;
-			ParticlesInput.X(Index) = NewX;
-			Chaos::TRotation<float, 3> NewR = FQuat::Slerp(OldCollisionTransforms[Index].GetRotation(), CollisionTransforms[Index].GetRotation(), Alpha);
-			Chaos::TRotation<float, 3> Delta = NewR * ParticlesInput.R(Index).Inverse();
-			Chaos::TVector<float, 3> Axis;
-			float Angle;
-			Delta.ToAxisAndAngle(Axis, Angle);
-			ParticlesInput.W(Index) = Axis * Angle / Dt;
-			ParticlesInput.R(Index) = NewR;
-		});
-
-    Time = Evolution->GetTime();
-	DeltaTime = 1.f / 30.f;  // Initialize filtered timestep at 30fps 
+	// Create solver
+	Solver = MakeUnique<FClothingSimulationSolver>();
+
+	// Create external collision dynamic collider
+	Colliders.Emplace(MakeUnique<FClothingSimulationCollider>(nullptr, nullptr, false, INDEX_NONE));
+	Colliders[0]->SetCollisionData(&ExternalCollisionData);
+	Solver->AddCollider(Colliders[0].Get());
+
+	ResetStats();
 }
 
-void ClothingSimulation::Shutdown()
+void FClothingSimulation::Shutdown()
 {
-	Assets.Reset();
-	AnimDriveSpringStiffness.Reset();
-	MaxDistancesMultipliers.Reset();
-	ExternalCollisions.Reset();
-	OldCollisionTransforms.Reset();
-	CollisionTransforms.Reset();
-	BoneIndices.Reset();
-	BaseTransforms.Reset();
-	OldAnimationPositions.Reset();
-	AnimationPositions.Reset();
-	AnimationNormals.Reset();
-	IndexToRangeMap.Reset();
-	RootBoneWorldTransforms.Reset();
-	LinearDeltaRatios.Reset();
-	AngularDeltaRatios.Reset();
+	Solver.Reset();
 	Meshes.Reset();
-	FaceNormals.Reset();
-	PointNormals.Reset();
-	Evolution.Reset();
-	CollisionsRangeMap.Reset();
-	ExternalCollisionsRangeMaps.Reset();
-	ExternalCollisionsOffset = 0;
+	Cloths.Reset();
+	Colliders.Reset();
 	ClothSharedSimConfig = nullptr;
-	LongRangeConstraints.Reset();
-	bUseGravityConfigs.Reset();
-	GravityScales.Reset();
-	GravityConfigs.Reset();
-	WindVelocities.Reset();
-
-	ResetStats();
 }
 
-void ClothingSimulation::DestroyActors()
+void FClothingSimulation::DestroyActors()
 {
 	Shutdown();
 	Initialize();
 }
 
-void ClothingSimulation::CreateActor(USkeletalMeshComponent* InOwnerComponent, UClothingAssetBase* InAsset, int32 InSimDataIndex)
+IClothingSimulationContext* FClothingSimulation::CreateContext()
 {
-	UE_LOG(LogChaosCloth, Verbose, TEXT("Adding Cloth LOD asset to %s in sim slot %d"), InOwnerComponent->GetOwner() ? *InOwnerComponent->GetOwner()->GetName() : TEXT("None"), InSimDataIndex);
+	return new FClothingSimulationContext();
+}
 
-	UClothingAssetCommon* const Asset = Cast<UClothingAssetCommon>(InAsset);
-	const UChaosClothConfig* const ChaosClothSimConfig = Asset->GetClothConfig<UChaosClothConfig>();
-	if (!ChaosClothSimConfig)
-	{
-		UE_LOG(LogChaosCloth, Warning, TEXT("Missing Chaos config Cloth LOD asset to %s in sim slot %d"), InOwnerComponent->GetOwner() ? *InOwnerComponent->GetOwner()->GetName() : TEXT("None"), InSimDataIndex);
-		return;
-	}
+void FClothingSimulation::CreateActor(USkeletalMeshComponent* InOwnerComponent, UClothingAssetBase* InAsset, int32 InSimDataIndex)
+{
+	SCOPE_CYCLE_COUNTER(STAT_ChaosClothCreateActor);
 
-	ClothingSimulationContext Context;
-	FillContext(InOwnerComponent, 0, &Context);
+	check(InOwnerComponent);
+	check(Solver);
 
-	if (Assets.Num() <= InSimDataIndex)
+	if (!InAsset)
 	{
-		const int32 NumAssets = InSimDataIndex + 1;
-
-		// TODO: Refactor all these arrays into a single cloth runtime asset structure
-		Assets.SetNumZeroed(NumAssets);
-		AnimDriveSpringStiffness.SetNumZeroed(NumAssets);
-		MaxDistancesMultipliers.SetNumZeroed(NumAssets);
-
-		Meshes.SetNum(NumAssets);
-		FaceNormals.SetNum(NumAssets);
-		PointNormals.SetNum(NumAssets);
-
-		IndexToRangeMap.SetNum(NumAssets);
-
-		LongRangeConstraints.SetNum(NumAssets);
-
-		RootBoneWorldTransforms.SetNum(NumAssets);
-		LinearDeltaRatios.SetNum(NumAssets);
-		AngularDeltaRatios.SetNum(NumAssets);
-
-		CollisionsRangeMap.SetNumZeroed(NumAssets);
-
-		GravityScales.SetNum(NumAssets);
-		bUseGravityConfigs.SetNum(NumAssets);
-		GravityConfigs.SetNum(NumAssets);
-		WindVelocities.SetNum(NumAssets);
+		return;
 	}
-	Assets[InSimDataIndex] = Asset;
-
-	check(Asset->GetNumLods() > 0);
-	UE_CLOG(Asset->GetNumLods() != 1,
-		LogChaosCloth, Warning, TEXT("More than one LOD with the current cloth asset %s in sim slot %d. Only LOD 0 is supported with Chaos Cloth for now."),
-		InOwnerComponent->GetOwner() ? *InOwnerComponent->GetOwner()->GetName() : TEXT("None"), InSimDataIndex);
-	const FClothLODDataCommon& AssetLodData = Asset->LodData[0];
-	const FClothPhysicalMeshData& PhysMesh = AssetLodData.PhysicalMeshData;
-
-	// Add particles
-	TPBDParticles<float, 3>& Particles = Evolution->Particles();
-	const uint32 Offset = Evolution->AddParticles(PhysMesh.Vertices.Num(), (uint32)InSimDataIndex);
-
-	IndexToRangeMap[InSimDataIndex] = Chaos::TVector<uint32, 2>(Offset, Particles.Size());
 
 	// ClothSharedSimConfig should either be a nullptr, or point to an object common to the whole skeletal mesh
-	if (ClothSharedSimConfig == nullptr)
+	UClothingAssetCommon* const Asset = Cast<UClothingAssetCommon>(InAsset);
+	if (!ClothSharedSimConfig)
 	{
 		ClothSharedSimConfig = Asset->GetClothConfig<UChaosClothSharedSimConfig>();
+
+		UpdateSimulationFromSharedSimConfig();
+
+		// Must set the local space location prior to adding any mesh/cloth, as otherwise the start poses would be in the wrong local space
+		const FClothingSimulationContext* const Context = static_cast<const FClothingSimulationContext*>(InOwnerComponent->GetClothingSimulationContext());
+		check(Context);
+		Solver->SetLocalSpaceLocation(bUseLocalSpaceSimulation ? Context->ComponentToWorld.GetLocation() : TVector<float, 3>(0.f));
 	}
 	else
 	{
 		check(ClothSharedSimConfig == Asset->GetClothConfig<UChaosClothSharedSimConfig>());
 	}
 
-	AnimationPositions.SetNum(Particles.Size());
-	AnimationNormals.SetNum(Particles.Size());
-
-	// Initialize the local simulation space transform
-	FTransform ComponentToLocalSpace = Context.ComponentToWorld;
-	if (Offset == 0) // Only initialize this once for all cloth instances
+	// Retrieve the cloth config stored in the asset
+	const UChaosClothConfig* const ClothConfig = Asset->GetClothConfig<UChaosClothConfig>();
+	if (!ClothConfig)
 	{
-		// Set local offset
-		if (bUseLocalSpaceSimulation)
-		{
-			LocalSpaceLocation = ComponentToLocalSpace.GetTranslation();
-		}
-		else
-		{
-			LocalSpaceLocation = FVector::ZeroVector;
-		}
-	}
-	ComponentToLocalSpace.AddToTranslation(-LocalSpaceLocation);
-
-	// Init local cloth sim space & teleport transform
-	const FTransform RootBoneTransform = Context.BoneTransforms[Asset->ReferenceBoneIndex];
-	RootBoneWorldTransforms[InSimDataIndex] = RootBoneTransform * Context.ComponentToWorld;  // Velocity scale deltas are calculated in world space
-	LinearDeltaRatios[InSimDataIndex] = FVector::OneVector - ChaosClothSimConfig->LinearVelocityScale.BoundToBox(FVector::ZeroVector, FVector::OneVector);
-	AngularDeltaRatios[InSimDataIndex] = 1.f - FMath::Clamp(ChaosClothSimConfig->AngularVelocityScale, 0.f, 1.f);
-
-	// Skin start pose
-	ClothingMeshUtils::SkinPhysicsMesh<true, false>(
-		Asset->UsedBoneIndices,
-		PhysMesh, // curr pos and norm
-		ComponentToLocalSpace,
-		Context.RefToLocals.GetData(),
-		Context.RefToLocals.Num(),
-		reinterpret_cast<TArray<FVector>&>(AnimationPositions),
-		reinterpret_cast<TArray<FVector>&>(AnimationNormals),
-		Offset);
-
-	ResetParticles(InSimDataIndex);
-
-	OldAnimationPositions = AnimationPositions;  // Also update the old positions array to avoid any interpolation issues
-
-	BuildMesh(PhysMesh, InSimDataIndex);
-
-	SetParticleMasses(ChaosClothSimConfig, PhysMesh, InSimDataIndex);
-
-	AddConstraints(ChaosClothSimConfig, PhysMesh, InSimDataIndex);
-
-	// Set this cloth damping, collision thickness, friction
-	Evolution->SetDamping(ChaosClothSimConfig->DampingCoefficient, InSimDataIndex);
-	Evolution->SetCollisionThickness(ChaosClothSimConfig->CollisionThickness, InSimDataIndex);
-	Evolution->SetCoefficientOfFriction(ChaosClothSimConfig->FrictionCoefficient, InSimDataIndex);
-
-	// Add velocity field
-	auto GetVelocity = [this, InSimDataIndex](const TVector<float, 3>&)->TVector<float, 3>
-	{
-		return WindVelocities[InSimDataIndex];
-	};
-	Evolution->GetVelocityField(InSimDataIndex) = MakeUnique<TVelocityField<float, 3>>(
-		*Meshes[InSimDataIndex],
-		GetVelocity,
-		/*bInIsUniform =*/ true,
-		ChaosClothSimConfig->DragCoefficient);
-
-	// Update gravity related config values
-	GravityScales[InSimDataIndex] = ChaosClothSimConfig->GravityScale;
-	bUseGravityConfigs[InSimDataIndex] = ChaosClothSimConfig->bUseGravityOverride;
-	GravityConfigs[InSimDataIndex] = ChaosClothSimConfig->Gravity;
-
-	// Add Self Collisions
-	if (ChaosClothSimConfig->bUseSelfCollisions)
-	{
-		AddSelfCollisions(InSimDataIndex);
-		Evolution->SetSelfCollisionThickness(ChaosClothSimConfig->SelfCollisionThickness, InSimDataIndex);
+		UE_LOG(LogChaosCloth, Warning, TEXT("Missing Chaos config Cloth LOD asset to %s in sim slot %d"), InOwnerComponent->GetOwner() ? *InOwnerComponent->GetOwner()->GetName() : TEXT("None"), InSimDataIndex);
+		return;
 	}
 
-	// Warn about legacy apex collisions
-	const FClothCollisionData& LodCollData = AssetLodData.CollisionData;
-	UE_CLOG(LodCollData.Spheres.Num() > 0 || LodCollData.SphereConnections.Num() > 0 || LodCollData.Convexes.Num() > 0,
-		LogChaosCloth, Warning, TEXT(
-			"Actor '%s' component '%s' has %d sphere, %d capsule, and %d "
-			"convex collision objects for physics authored as part of a LOD construct, "
-			"probably by the Apex cloth authoring system.  This is deprecated.  "
-			"Please update your asset!"),
-		InOwnerComponent->GetOwner() ? *InOwnerComponent->GetOwner()->GetName() : TEXT("None"),
-		*InOwnerComponent->GetName(),
-		LodCollData.Spheres.Num(),
-		LodCollData.SphereConnections.Num(),
-		LodCollData.Convexes.Num());
-
-	// Extract all collisions for this asset
-	checkf(ExternalCollisions.Spheres.Num() == 0 &&
-		ExternalCollisions.SphereConnections.Num() == 0 &&
-		ExternalCollisions.Convexes.Num() == 0 &&
-		ExternalCollisions.Boxes.Num() == 0, TEXT("There cannot be any external collisions added before all the cloth assets collisions are processed."));
-	ExtractCollisions(Asset, InSimDataIndex);
-
-	// Update collision transforms, including initial state for particles' X & R
-	UpdateCollisionTransforms(Context, InSimDataIndex);
+	// Create mesh node
+	const int32 MeshIndex = Meshes.Emplace(MakeUnique<FClothingSimulationMesh>(
+		Asset,
+		InOwnerComponent));
+
+	// Create collider node
+	const int32 ColliderIndex = Colliders.Emplace(MakeUnique<FClothingSimulationCollider>(
+		Asset,
+		InOwnerComponent,
+		/*bInUseLODIndexOverride =*/ false,
+		/*InLODIndexOverride =*/ INDEX_NONE));
+
+	// Create cloth node
+	AnimDriveSpringStiffness = ClothConfig->AnimDriveSpringStiffness;
+	const int32 ClothIndex = Cloths.Emplace(MakeUnique<FClothingSimulationCloth>(
+		Meshes[MeshIndex].Get(),
+		TArray<FClothingSimulationCollider*>({ Colliders[ColliderIndex].Get() }),
+		InSimDataIndex,
+		(FClothingSimulationCloth::EMassMode)ClothConfig->MassMode,
+		ClothConfig->GetMassValue(),
+		ClothConfig->MinPerParticleMass,
+		ClothConfig->EdgeStiffness,
+		ClothConfig->BendingStiffness,
+		ClothConfig->bUseBendingElements,
+		ClothConfig->AreaStiffness,
+		ClothConfig->VolumeStiffness,
+		ClothConfig->bUseThinShellVolumeConstraints,
+		ClothConfig->StrainLimitingStiffness,
+		ClothConfig->LimitScale,
+		ClothConfig->bUseGeodesicDistance,
+		/*MaxDistancesMultiplier =*/ 1.f,  // Animatable
+		AnimDriveSpringStiffness,  // Animatable
+		ClothConfig->ShapeTargetStiffness,
+		/*bUseXPBDConstraints =*/ false,  // Experimental
+		ClothConfig->GravityScale,
+		ClothConfig->bUseGravityOverride,
+		ClothConfig->Gravity,
+		ClothConfig->LinearVelocityScale,
+		ClothConfig->AngularVelocityScale,
+		ClothConfig->DragCoefficient,
+		ClothConfig->DampingCoefficient,
+		ClothConfig->CollisionThickness,
+		ClothConfig->FrictionCoefficient,
+		ClothConfig->bUseSelfCollisions,
+		ClothConfig->SelfCollisionThickness,
+		/*bUseLODIndexOverride =*/ false,
+		/*LODIndexOverride =*/ INDEX_NONE));
+
+	// Add cloth to solver
+	Solver->AddCloth(Cloths[ClothIndex].Get());
 
 	// Update stats
-	UpdateStats(InSimDataIndex);
+	UpdateStats(Cloths[ClothIndex].Get());
+
+	UE_LOG(LogChaosCloth, Verbose, TEXT("Added Cloth asset to %s in sim slot %d"), InOwnerComponent->GetOwner() ? *InOwnerComponent->GetOwner()->GetName() : TEXT("None"), InSimDataIndex);
 }
 
-void ClothingSimulation::ResetStats()
+void FClothingSimulation::ResetStats()
 {
+	check(Solver);
 	NumCloths = 0;
 	NumKinemamicParticles = 0;
 	NumDynamicParticles = 0;
 	SimulationTime = 0.f;
+	NumSubsteps = Solver->GetNumSubsteps();
+	NumIterations = Solver->GetNumIterations();
 }
 
-void ClothingSimulation::UpdateStats(int32 InSimDataIndex)
-{
-	const uint32 Offset = IndexToRangeMap[InSimDataIndex][0];
-	const uint32 Range = IndexToRangeMap[InSimDataIndex][1];
-	if (const int32 AddedParticleCount = Range - Offset)
-	{
-		++NumCloths;
-	}
-
-	TPBDParticles<float, 3>& Particles = Evolution->Particles();
-	int32 NumAddedKinenamicParticles = 0;
-	int32 NumAddedDynamicParticles = 0;
-	for (uint32 i = Offset; i < Range; ++i)
-	{
-		if (Particles.InvM(i) == 0.f)
-		{
-			++NumAddedKinenamicParticles;
-		}
-		else
-		{
-			++NumAddedDynamicParticles;
-		}
-	}
-	NumKinemamicParticles += NumAddedKinenamicParticles;
-	NumDynamicParticles += NumAddedDynamicParticles;
-}
-
-void ClothingSimulation::ExtractCollisions(const UClothingAssetCommon* Asset, int32 InSimDataIndex)
+void FClothingSimulation::UpdateStats(const FClothingSimulationCloth* Cloth)
 {
-	if (!Asset)
-	{
-		return;
-	}
-
-	CollisionsRangeMap[InSimDataIndex][0] = Evolution->CollisionParticles().Size();
-
-	// Pull collisions from the specified physics asset inside the clothing asset
-	ExtractPhysicsAssetCollisions(Asset, InSimDataIndex);
-
-	// Extract the legacy Apex collision from the clothing asset
-	ExtractLegacyAssetCollisions(Asset, InSimDataIndex);
-
-	// Update the external collision offset and collision range for this asset
-	CollisionsRangeMap[InSimDataIndex][1] = ExternalCollisionsOffset = Evolution->CollisionParticles().Size();
+	NumCloths = Cloths.Num();
+	NumKinemamicParticles += Cloth->GetNumActiveKinematicParticles();
+	NumDynamicParticles += Cloth->GetNumActiveDynamicParticles();
 }
 
-void ClothingSimulation::PostActorCreationInitialize()
-{
-	UpdateSimulationFromSharedSimConfig();
-}
-
-void ClothingSimulation::UpdateSimulationFromSharedSimConfig()
+void FClothingSimulation::UpdateSimulationFromSharedSimConfig()
 {
+	check(Solver);
 	if (ClothSharedSimConfig) // ClothSharedSimConfig will be a null pointer if all cloth instances are disabled in which case we will use default Evolution parameters
 	{
 		// Update local space simulation switch
 		bUseLocalSpaceSimulation = ClothSharedSimConfig->bUseLocalSpaceSimulation;
 
-		// Now set all the common parameters on the simulation
-		NumSubsteps = ClothSharedSimConfig->SubdivisionCount;
-		NumIterations = ClothSharedSimConfig->IterationCount;
-		Evolution->SetIterations(ClothSharedSimConfig->IterationCount);
-	}
-}
-
-void ClothingSimulation::BuildMesh(const FClothPhysicalMeshData& InPhysMesh, int32 InSimDataIndex)
-{
-	TUniquePtr<Chaos::TTriangleMesh<float>>& Mesh = Meshes[InSimDataIndex];
-
-	const uint32 Offset = IndexToRangeMap[InSimDataIndex][0];
-
-	const int32 NumTriangles = InPhysMesh.Indices.Num() / 3;
-	TArray<Chaos::TVector<int32, 3>> InputSurfaceElements;
-	InputSurfaceElements.Reserve(NumTriangles);
-	for (int i = 0; i < NumTriangles; ++i)
-	{
-		const int32 Index = 3 * i;
-		InputSurfaceElements.Add(
-			{ static_cast<int32>(Offset + InPhysMesh.Indices[Index]),
-			 static_cast<int32>(Offset + InPhysMesh.Indices[Index + 1]),
-			 static_cast<int32>(Offset + InPhysMesh.Indices[Index + 2]) });
-	}
-	check(InputSurfaceElements.Num() == NumTriangles);
-	Mesh.Reset(new Chaos::TTriangleMesh<float>(MoveTemp(InputSurfaceElements)));
-	check(Mesh->GetNumElements() == NumTriangles);
-	Mesh->GetPointToTriangleMap(); // Builds map for later use by GetPointNormals()
-}
-
-void ClothingSimulation::ResetParticles(int32 InSimDataIndex)
-{
-	const uint32 Offset = IndexToRangeMap[InSimDataIndex][0];
-	const uint32 Range = IndexToRangeMap[InSimDataIndex][1];
-
-	for (uint32 i = Offset; i < Range; ++i)
-	{
-		Evolution->Particles().P(i) = Evolution->Particles().X(i) = AnimationPositions[i];
-		Evolution->Particles().V(i) = Chaos::TVector<float, 3>(0.f);
-		Evolution->Particles().M(i) = 0.f;
-	}
-}
-
-void ClothingSimulation::SetParticleMasses(const UChaosClothConfig* ChaosClothConfig, const FClothPhysicalMeshData& PhysMesh, int32 InSimDataIndex)
-{
-	const Chaos::TTriangleMesh<float>& Mesh = *Meshes[InSimDataIndex];
-	TPBDParticles<float, 3>& Particles = Evolution->Particles();
-
-	// Assign per particle mass proportional to connected area.
-	const TArray<TVector<int32, 3>>& SurfaceElements = Mesh.GetSurfaceElements();
-	float TotalArea = 0.f;
-	for (const Chaos::TVector<int32, 3>& Tri : SurfaceElements)
-	{
-		const float TriArea = 0.5f * Chaos::TVector<float, 3>::CrossProduct(
-			Particles.X(Tri[1]) - Particles.X(Tri[0]),
-			Particles.X(Tri[2]) - Particles.X(Tri[0])).Size();
-		TotalArea += TriArea;
-		const float ThirdTriArea = TriArea / 3.f;
-		Particles.M(Tri[0]) += ThirdTriArea;
-		Particles.M(Tri[1]) += ThirdTriArea;
-		Particles.M(Tri[2]) += ThirdTriArea;
-	}
-	const TSet<int32> Vertices = Mesh.GetVertices();
-	float TotalMass = 0.f;
-	switch (ChaosClothConfig->MassMode)
-	{
-	case EClothMassMode::UniformMass:
-		for (const int32 Vertex : Vertices)
-		{
-			Particles.M(Vertex) = ChaosClothConfig->UniformMass;
-			TotalMass += Particles.M(Vertex);
-		}
-		break;
-	case EClothMassMode::TotalMass:
-	{
-		const float MassPerUnitArea = TotalArea > 0.f ? ChaosClothConfig->TotalMass / TotalArea : 1.f;
-		for (const int32 Vertex : Vertices)
-		{
-			Particles.M(Vertex) *= MassPerUnitArea;
-			TotalMass += Particles.M(Vertex);
-		}
-		break;
-	}
-	case EClothMassMode::Density:
-	{
-		const float Density = ChaosClothConfig->Density / FMath::Square(ChaosClothingSimulationDefault::WorldScale);
-		for (const int32 Vertex : Vertices)
-		{
-			Particles.M(Vertex) *= Density;
-			TotalMass += Particles.M(Vertex);
-		}
-		break;
-	}
-	};
-	UE_LOG(LogChaosCloth, Verbose, TEXT("Density: %f, Total surface: %f, Total mass: %f, "), TotalArea > 0.f ? TotalMass / TotalArea : 1.f, TotalArea, TotalMass);
-	UE_LOG(LogChaosCloth, Verbose, TEXT("SI Density: %f, SI Total surface: %f, SI Total mass: %f, "), (TotalArea > 0.f ? TotalMass / TotalArea : 1.f) * FMath::Square(ChaosClothingSimulationDefault::WorldScale), TotalArea / FMath::Square(ChaosClothingSimulationDefault::WorldScale), TotalMass);
-
-	// Clamp and enslave
-	const FPointWeightMap& MaxDistances = PhysMesh.GetWeightMap(EChaosWeightMapTarget::MaxDistance);
-	const uint32 Offset = IndexToRangeMap[InSimDataIndex][0];
-	const uint32 Range = IndexToRangeMap[InSimDataIndex][1];
-
-	check(Particles.Size() >= Range);
-	for (uint32 i = Offset; i < Range; ++i)
-	{
-		Particles.M(i) = FMath::Max(Particles.M(i), ChaosClothConfig->MinPerParticleMass);
-		Particles.InvM(i) = MaxDistances.IsBelowThreshold(i - Offset) ? 0.f : 1.f / Particles.M(i);
-	}
-}
-
-void ClothingSimulation::AddConstraints(const UChaosClothConfig* ChaosClothSimConfig, const FClothPhysicalMeshData& PhysMesh, int32 InSimDataIndex)
-{
-	const Chaos::TTriangleMesh<float>& Mesh = *Meshes[InSimDataIndex];
-	const TArray<TVector<int32, 3>>& SurfaceElements = Mesh.GetSurfaceElements();
-
-	const uint32 Offset = IndexToRangeMap[InSimDataIndex][0];
-	const uint32 ParticleCount = IndexToRangeMap[InSimDataIndex][1] - Offset;
-
-	const bool bUseXPBDConstraints = ClothSharedSimConfig && ClothSharedSimConfig->bUseXPBDConstraints;
-
-	if (ChaosClothSimConfig->ShapeTargetStiffness)
-	{
-		check(ChaosClothSimConfig->ShapeTargetStiffness > 0.f && ChaosClothSimConfig->ShapeTargetStiffness <= 1.f);
-		Evolution->AddPBDConstraintFunction([ShapeConstraints = Chaos::TPBDShapeConstraints<float, 3>(Evolution->Particles(), Offset, ParticleCount, AnimationPositions, ChaosClothSimConfig->ShapeTargetStiffness)](TPBDParticles<float, 3>& InParticles, const float Dt)
-		{
-			ShapeConstraints.Apply(InParticles, Dt);
-		});
-	}
-	if (ChaosClothSimConfig->EdgeStiffness)
-	{
-		check(ChaosClothSimConfig->EdgeStiffness > 0.f && ChaosClothSimConfig->EdgeStiffness <= 1.f);
-		if (bUseXPBDConstraints)
-		{
-			const TSharedPtr<Chaos::TXPBDSpringConstraints<float, 3>> SpringConstraints =
-				MakeShared<Chaos::TXPBDSpringConstraints<float, 3>>(Evolution->Particles(), SurfaceElements, ChaosClothSimConfig->EdgeStiffness);
-			Evolution->AddXPBDConstraintFunctions(
-				[SpringConstraints]()
-				{
-					SpringConstraints->Init();
-				},
-				[SpringConstraints](TPBDParticles<float, 3>& InParticles, const float Dt)
-				{
-					SpringConstraints->Apply(InParticles, Dt);
-				});
-		}
-		else
-		{
-			TArray<TVector<int32, 3>> DynamicSurfaceElements;
-			TArray<TVector<int32, 2>> Attachments;
-			for (const TVector<int32, 3>& SurfaceElement : SurfaceElements)
-			{
-				const bool bIsKinematic0 = (Evolution->Particles().InvM(SurfaceElement[0]) == 0.f);
-				const bool bIsKinematic1 = (Evolution->Particles().InvM(SurfaceElement[1]) == 0.f);
-				const bool bIsKinematic2 = (Evolution->Particles().InvM(SurfaceElement[2]) == 0.f);
-				bool bIsAttachment = false;
-				if (bIsKinematic0 != bIsKinematic1)
-				{
-					Attachments.Emplace(SurfaceElement[0], SurfaceElement[1]);
-					bIsAttachment = true;
-				}
-				if (bIsKinematic1 != bIsKinematic2)
-				{
-					Attachments.Emplace(SurfaceElement[1], SurfaceElement[2]);
-					bIsAttachment = true;
-				}
-				if (bIsKinematic2 != bIsKinematic0)
-				{
-					Attachments.Emplace(SurfaceElement[2], SurfaceElement[0]);
-					bIsAttachment = true;
-				}
-				if (!bIsAttachment)
-				{
-					DynamicSurfaceElements.Add(SurfaceElement);
-				}
-			}
-			if (Attachments.Num())
-			{
-				Evolution->AddPBDConstraintFunction([AttachmentConstraints = Chaos::FPBDSpringConstraints(Evolution->Particles(), MoveTemp(Attachments), ChaosClothSimConfig->EdgeStiffness)](TPBDParticles<float, 3>& InParticles, const float Dt)
-				{
-					AttachmentConstraints.Apply(InParticles, Dt);
-				});
-			}
-			if (DynamicSurfaceElements.Num())
-			{
-				Evolution->AddPBDConstraintFunction([SpringConstraints = Chaos::FPBDSpringConstraints(Evolution->Particles(), DynamicSurfaceElements, ChaosClothSimConfig->EdgeStiffness)](TPBDParticles<float, 3>& InParticles, const float Dt)
-				{
-					SpringConstraints.Apply(InParticles, Dt);
-				});
-			}
-		}
-	}
-	if (ChaosClothSimConfig->BendingStiffness)
-	{
-		check(ChaosClothSimConfig->BendingStiffness > 0.f && ChaosClothSimConfig->BendingStiffness <= 1.f);
-		if (ChaosClothSimConfig->bUseBendingElements)
-		{
-			TArray<Chaos::TVector<int32, 4>> BendingConstraints = Mesh.GetUniqueAdjacentElements();
-			Evolution->AddPBDConstraintFunction([BendConstraints = Chaos::TPBDBendingConstraints<float>(Evolution->Particles(), MoveTemp(BendingConstraints), ChaosClothSimConfig->BendingStiffness)](TPBDParticles<float, 3>& InParticles, const float Dt)
-			{
-				BendConstraints.Apply(InParticles, Dt);
-			});
-		}
-		else
-		{
-			TArray<Chaos::TVector<int32, 2>> BendingConstraints = Mesh.GetUniqueAdjacentPoints();
-			if (bUseXPBDConstraints)
-			{
-				const TSharedPtr<Chaos::TXPBDSpringConstraints<float, 3>> SpringConstraints =
-					MakeShared<Chaos::TXPBDSpringConstraints<float, 3>>(Evolution->Particles(), MoveTemp(BendingConstraints), ChaosClothSimConfig->BendingStiffness);
-				Evolution->AddXPBDConstraintFunctions(
-					[SpringConstraints]()
-					{
-						SpringConstraints->Init();
-					},
-					[SpringConstraints](TPBDParticles<float, 3>& InParticles, const float Dt)
-					{
-						SpringConstraints->Apply(InParticles, Dt);
-					});
-			}
-			else
-			{
-				Evolution->AddPBDConstraintFunction([SpringConstraints = Chaos::FPBDSpringConstraints(Evolution->Particles(), MoveTemp(BendingConstraints), ChaosClothSimConfig->BendingStiffness)](TPBDParticles<float, 3>& InParticles, const float Dt)
-				{
-					SpringConstraints.Apply(InParticles, Dt);
-				});
-			}
-		}
-	}
-	if (ChaosClothSimConfig->AreaStiffness)
-	{
-		TArray<Chaos::TVector<int32, 3>> SurfaceConstraints = SurfaceElements;
-		if (bUseXPBDConstraints)
-		{
-			const TSharedPtr<Chaos::TXPBDAxialSpringConstraints<float, 3>> AxialSpringConstraints =
-				MakeShared<Chaos::TXPBDAxialSpringConstraints<float, 3>>(Evolution->Particles(), MoveTemp(SurfaceConstraints), ChaosClothSimConfig->AreaStiffness);
-			Evolution->AddXPBDConstraintFunctions(
-				[AxialSpringConstraints]()
-				{
-					AxialSpringConstraints->Init();
-				},
-				[AxialSpringConstraints](TPBDParticles<float, 3>& InParticles, const float Dt)
-				{
-					AxialSpringConstraints->Apply(InParticles, Dt);
-				});
-		}
-		else
-		{
-			Evolution->AddPBDConstraintFunction([AxialSpringConstraints = Chaos::FPBDAxialSpringConstraints(Evolution->Particles(), MoveTemp(SurfaceConstraints), ChaosClothSimConfig->AreaStiffness)](TPBDParticles<float, 3>& InParticles, const float Dt)
-			{
-				AxialSpringConstraints.Apply(InParticles, Dt);
-			});
-		}
-	}
-	if (ChaosClothSimConfig->VolumeStiffness)
-	{
-		check(ChaosClothSimConfig->VolumeStiffness > 0.f && ChaosClothSimConfig->VolumeStiffness <= 1.f);
-		if (ChaosClothSimConfig->bUseTetrahedralConstraints)
-		{
-			// TODO(mlentine): Need to tetrahedralize surface to support this
-			check(false);
-		}
-		else if (ChaosClothSimConfig->bUseThinShellVolumeConstraints)
-		{
-			TArray<Chaos::TVector<int32, 2>> BendingConstraints = Mesh.GetUniqueAdjacentPoints();
-			TArray<Chaos::TVector<int32, 2>> DoubleBendingConstraints;
-			{
-				TMap<int32, TArray<int32>> BendingHash;
-				for (int32 i = 0; i < BendingConstraints.Num(); ++i)
-				{
-					BendingHash.FindOrAdd(BendingConstraints[i][0]).Add(BendingConstraints[i][1]);
-					BendingHash.FindOrAdd(BendingConstraints[i][1]).Add(BendingConstraints[i][0]);
-				}
-				TSet<Chaos::TVector<int32, 2>> Visited;
-				for (auto Elem : BendingHash)
-				{
-					for (int32 i = 0; i < Elem.Value.Num(); ++i)
-					{
-						for (int32 j = i + 1; j < Elem.Value.Num(); ++j)
-						{
-							if (Elem.Value[i] == Elem.Value[j])
-								continue;
-							auto NewElem = Chaos::TVector<int32, 2>(Elem.Value[i], Elem.Value[j]);
-							if (!Visited.Contains(NewElem))
-							{
-								DoubleBendingConstraints.Add(NewElem);
-								Visited.Add(NewElem);
-								Visited.Add(Chaos::TVector<int32, 2>(Elem.Value[j], Elem.Value[i]));
-							}
-						}
-					}
-				}
-			}
-			Evolution->AddPBDConstraintFunction([SpringConstraints = Chaos::FPBDSpringConstraints(Evolution->Particles(), MoveTemp(DoubleBendingConstraints), ChaosClothSimConfig->VolumeStiffness)](TPBDParticles<float, 3>& InParticles, const float Dt) {
-				SpringConstraints.Apply(InParticles, Dt);
-			});
-		}
-		else
-		{
-			TArray<Chaos::TVector<int32, 3>> SurfaceConstraints = SurfaceElements;
-			Chaos::TPBDVolumeConstraint<float> PBDVolumeConstraint(Evolution->Particles(), MoveTemp(SurfaceConstraints));
-			Evolution->AddPBDConstraintFunction([PBDVolumeConstraint = MoveTemp(PBDVolumeConstraint)](TPBDParticles<float, 3>& InParticles, const float Dt)
-			{
-				PBDVolumeConstraint.Apply(InParticles, Dt);
-			});
-		}
-	}
-	if (ChaosClothSimConfig->StrainLimitingStiffness)
-	{
-		check(Mesh.GetNumElements() > 0);
-		// PerFormance note: The Per constraint version of this function is quite a bit faster for smaller assets
-		// There might be a cross-over point where the PerParticle version is faster: To be determined
-		if (bUseXPBDConstraints)
-		{
-			LongRangeConstraints[InSimDataIndex] = MakeShared<Chaos::TXPBDLongRangeConstraints<float, 3>>(
-				Evolution->Particles(),
-				Mesh.GetPointToNeighborsMap(),
-				10, // The max number of connected neighbors per particle.  ryan - What should this be?  Was k...
-				ChaosClothSimConfig->StrainLimitingStiffness);  // TODO(Kriss.Gossart): Add LimitScale and Geodesic mode if ever of use
-
-			Evolution->AddXPBDConstraintFunctions(
-				[this, InSimDataIndex]()
-				{
-					static_cast<Chaos::TXPBDLongRangeConstraints<float, 3>&>(*LongRangeConstraints[InSimDataIndex]).Init();
-				},
-				[this, InSimDataIndex](TPBDParticles<float, 3>& InParticles, const float Dt)
-				{
-					static_cast<Chaos::TXPBDLongRangeConstraints<float, 3>&>(*LongRangeConstraints[InSimDataIndex]).Apply(InParticles, Dt);
-				});
-		}
-		else
-		{
-			LongRangeConstraints[InSimDataIndex] = MakeShared<Chaos::TPBDLongRangeConstraints<float, 3>>(
-				Evolution->Particles(),
-				Mesh.GetPointToNeighborsMap(),
-				10, // The max number of connected neighbors per particle.  ryan - What should this be?  Was k...
-				ChaosClothSimConfig->StrainLimitingStiffness,
-				ChaosClothSimConfig->LimitScale,
-				ChaosClothSimConfig->bUseGeodesicDistance);
-
-			Evolution->AddPBDConstraintFunction([this, InSimDataIndex](TPBDParticles<float, 3>& InParticles, const float Dt)
-			{
-				static_cast<Chaos::TPBDLongRangeConstraints<float, 3>&>(*LongRangeConstraints[InSimDataIndex]).Apply(InParticles, Dt);
-			});
-		}
-	}
-
-	// Maximum Distance Constraints
-	const FPointWeightMap& MaxDistances = PhysMesh.GetWeightMap(EChaosWeightMapTarget::MaxDistance);
-	if (MaxDistances.Num() > 0)
-	{
-		// Initialize the interactor's multiplier
-		MaxDistancesMultipliers[InSimDataIndex] = 1.f;
-
-		check(Mesh.GetNumElements() > 0);
-		Chaos::PBDSphericalConstraint<float, 3> SphericalContraint(Offset, MaxDistances.Num(), true, &AnimationPositions, &MaxDistances.Values);
-		Evolution->AddPBDConstraintFunction([
-			this, InSimDataIndex,
-			SphericalContraint = MoveTemp(SphericalContraint)](TPBDParticles<float, 3>& InParticles, const float Dt) mutable
-		{
-			SphericalContraint.SetSphereRadiiMultiplier(MaxDistancesMultipliers[InSimDataIndex]);
-			SphericalContraint.Apply(InParticles, Dt);
-		});
-	}
-
-	// Backstop Constraints
-	const FPointWeightMap& BackstopRadiuses = PhysMesh.GetWeightMap(EChaosWeightMapTarget::BackstopRadius);
-	const FPointWeightMap& BackstopDistances = PhysMesh.GetWeightMap(EChaosWeightMapTarget::BackstopDistance);
-	if (BackstopRadiuses.Num() > 0 && BackstopDistances.Num() > 0)
-	{
-		check(Mesh.GetNumElements() > 0);
-		check(BackstopRadiuses.Num() == BackstopDistances.Num());
-
-		Chaos::PBDSphericalConstraint<float, 3> SphericalContraint(Offset, BackstopRadiuses.Num(), false, &AnimationPositions,
-			&BackstopRadiuses.Values, &BackstopDistances.Values, &AnimationNormals);
-		Evolution->AddPBDConstraintFunction([SphericalContraint = MoveTemp(SphericalContraint)](TPBDParticles<float, 3>& InParticles, const float Dt)
-		{
-			SphericalContraint.Apply(InParticles, Dt);
-		});
-	}
-
-	// Animation Drive Constraints
-	AnimDriveSpringStiffness[InSimDataIndex] = ChaosClothSimConfig->AnimDriveSpringStiffness;
-	const FPointWeightMap& AnimDriveMultipliers = PhysMesh.GetWeightMap(EChaosWeightMapTarget::AnimDriveMultiplier);
-	if (AnimDriveMultipliers.Num() > 0)
-	{
-		check(Mesh.GetNumElements() > 0);
-		TPBDAnimDriveConstraint<float, 3> PBDAnimDriveConstraint(Offset, &AnimationPositions, &AnimDriveMultipliers.Values, AnimDriveSpringStiffness[InSimDataIndex]);
-		Evolution->AddPBDConstraintFunction(
-			[PBDAnimDriveConstraint = MoveTemp(PBDAnimDriveConstraint), &Stiffness = AnimDriveSpringStiffness[InSimDataIndex]](TPBDParticles<float, 3>& InParticles, const float Dt) mutable
-		{
-			PBDAnimDriveConstraint.SetSpringStiffness(Stiffness);
-			PBDAnimDriveConstraint.Apply(InParticles, Dt);
-		});
-	}
-}
-
-void ClothingSimulation::AddSelfCollisions(int32 InSimDataIndex) 
-{
-#ifdef CHAOS_CLOTH_MUST_IMPROVE_SELF_COLLISION  // TODO: Improve self-collision until it can run in engine tests without crashing the simulation.
-	// TODO(mlentine): Parallelize these for multiple meshes
-	const Chaos::TTriangleMesh<float>& Mesh = *Meshes[InSimDataIndex];
-	Evolution->CollisionTriangles().Append(Mesh.GetSurfaceElements());
-
-	const uint32 Offset = IndexToRangeMap[InSimDataIndex][0];
-	const uint32 Range = IndexToRangeMap[InSimDataIndex][1];
-	for (uint32 i = Offset; i < Range; ++i)
-	{
-		const TSet<int32> Neighbors = Mesh.GetNRing(i, 5);
-		for (int32 Element : Neighbors)
-		{
-			check(i != Element);
-			Evolution->DisabledCollisionElements().Add(Chaos::TVector<int32, 2>(i, Element));
-			Evolution->DisabledCollisionElements().Add(Chaos::TVector<int32, 2>(Element, i));
-		}
+		// Set common simulation parameters
+		Solver->SetNumSubsteps(ClothSharedSimConfig->SubdivisionCount);
+		Solver->SetNumIterations(ClothSharedSimConfig->IterationCount);
 	}
-#endif  // #ifdef CHAOS_CLOTH_MUST_IMPROVE_SELF_COLLISION
 }
 
-void ClothingSimulation::ForAllCollisions(TFunction<void(TGeometryClothParticles<float, 3>&, uint32)> CollisionFunction, int32 SimDataIndex)
-{
-	TGeometryClothParticles<float, 3>& CollisionParticles = Evolution->CollisionParticles();
-
-	for (uint32 Index = CollisionsRangeMap[SimDataIndex][0]; Index < CollisionsRangeMap[SimDataIndex][1]; ++Index)
-	{
-		CollisionFunction(CollisionParticles, Index);
-	}
-	for (const TArray<TVector<uint32, 2>>& ExternalCollisionsRangeMap : ExternalCollisionsRangeMaps)
-	{
-		for (uint32 Index = ExternalCollisionsRangeMap[SimDataIndex][0]; Index < ExternalCollisionsRangeMap[SimDataIndex][1]; ++Index)
-		{
-			CollisionFunction(CollisionParticles, Index);
-		}
-	}
-}
-
-void ClothingSimulation::UpdateCollisionTransforms(const ClothingSimulationContext& Context, int32 InSimDataIndex)
-{
-	SCOPE_CYCLE_COUNTER(STAT_ChaosClothUpdateCollisionTransforms);
-
-	// Resize the transform arrays if collision have changed
-	const int32 PrevNumCollisions = OldCollisionTransforms.Num();
-	const int32 NumCollisions = BaseTransforms.Num();
-	check(NumCollisions == int32(Evolution->CollisionParticles().Size()));  // BaseTransforms should always automatically grow with the number of collision particles (collection array)
-
-	const bool bHasNumCollisionsChanged = (NumCollisions != PrevNumCollisions);
-	if (bHasNumCollisionsChanged)
-	{
-		CollisionTransforms.SetNum(NumCollisions);
-		OldCollisionTransforms.SetNum(NumCollisions);
-	}
-
-	// Update the collision transforms
-	FTransform ComponentToLocalSimulationSpace = Context.ComponentToWorld;
-	ComponentToLocalSimulationSpace.AddToTranslation(-LocalSpaceLocation);
-
-	ForAllCollisions([this, &Context, &ComponentToLocalSimulationSpace, bHasNumCollisionsChanged](TGeometryClothParticles<float, 3>& CollisionParticles, uint32 Index)
-	{
-		// Update the collision transforms
-		const int32 BoneIndex = BoneIndices[Index];
-		Chaos::TRigidTransform<float, 3>& CollisionTransform = CollisionTransforms[Index];
-		if (Context.BoneTransforms.IsValidIndex(BoneIndex))
-		{
-			const FTransform& BoneTransform = Context.BoneTransforms[BoneIndex];
-			CollisionTransform = BaseTransforms[Index] * BoneTransform * ComponentToLocalSimulationSpace;
-		}
-		else
-		{
-			CollisionTransform = BaseTransforms[Index] * ComponentToLocalSimulationSpace;  // External collisions often don't map to a bone
-		}
-		// Reset initial states if required
-		if (bHasNumCollisionsChanged)
-		{
-			CollisionParticles.X(Index) = CollisionTransform.GetTranslation();
-			CollisionParticles.R(Index) = CollisionTransform.GetRotation();
-			OldCollisionTransforms[Index] = CollisionTransform;
-		}
-	}, InSimDataIndex);
-}
-
-void ClothingSimulation::ExtractPhysicsAssetCollisions(const UClothingAssetCommon* Asset, int32 InSimDataIndex)
-{
-	FClothCollisionData ExtractedCollisions;
-
-	TGeometryClothParticles<float, 3>& CollisionParticles = Evolution->CollisionParticles();
-
-	if (!Asset)
-	{
-		return;
-	}
-
-	// TODO(mlentine): Support collision body activation on a per particle basis, preferably using a map but also can be a particle attribute
-	if (const UPhysicsAsset* const PhysAsset = Asset->PhysicsAsset)
-	{
-		const USkeletalMesh* const TargetMesh = CastChecked<USkeletalMesh>(Asset->GetOuter());
-
-		TArray<int32> UsedBoneIndices;
-		UsedBoneIndices.Reserve(PhysAsset->SkeletalBodySetups.Num());
-
-		for (const USkeletalBodySetup* BodySetup : PhysAsset->SkeletalBodySetups)
-		{
-			if (!BodySetup)
-				continue;
-
-			const int32 MeshBoneIndex = TargetMesh->RefSkeleton.FindBoneIndex(BodySetup->BoneName);
-			const int32 MappedBoneIndex = UsedBoneIndices.Add(MeshBoneIndex);
-			
-			// Add capsules
-			const FKAggregateGeom& AggGeom = BodySetup->AggGeom;
-			if (AggGeom.SphylElems.Num())
-			{
-				for (const FKSphylElem& SphylElem : AggGeom.SphylElems)
-				{
-					if (SphylElem.Length == 0.0f)
-					{
-						// Add extracted sphere collision data
-						FClothCollisionPrim_Sphere Sphere;
-						Sphere.LocalPosition = SphylElem.Center;
-						Sphere.Radius = SphylElem.Radius;
-						Sphere.BoneIndex = MappedBoneIndex;
-						ExtractedCollisions.Spheres.Add(Sphere);
-					}
-					else
-					{
-						// Add extracted spheres collision data
-						FClothCollisionPrim_Sphere Sphere0;
-						FClothCollisionPrim_Sphere Sphere1;
-						const FVector OrientedDirection = SphylElem.Rotation.RotateVector(FVector::UpVector);
-						const FVector HalfDim = OrientedDirection * (SphylElem.Length / 2.f);
-						Sphere0.LocalPosition = SphylElem.Center - HalfDim;
-						Sphere1.LocalPosition = SphylElem.Center + HalfDim;
-						Sphere0.Radius = SphylElem.Radius;
-						Sphere1.Radius = SphylElem.Radius;
-						Sphere0.BoneIndex = MappedBoneIndex;
-						Sphere1.BoneIndex = MappedBoneIndex;
-
-						// Add extracted sphere connection collision data
-						FClothCollisionPrim_SphereConnection SphereConnection;
-						SphereConnection.SphereIndices[0] = ExtractedCollisions.Spheres.Add(Sphere0);
-						SphereConnection.SphereIndices[1] = ExtractedCollisions.Spheres.Add(Sphere1);
-						ExtractedCollisions.SphereConnections.Add(SphereConnection);
-					}
-				}
-			}
-
-			// Add spheres
-			for (const FKSphereElem& SphereElem : AggGeom.SphereElems)
-			{
-				// Add extracted sphere collision data
-				FClothCollisionPrim_Sphere Sphere;
-				Sphere.LocalPosition = SphereElem.Center;
-				Sphere.Radius = SphereElem.Radius;
-				Sphere.BoneIndex = MappedBoneIndex;
-				ExtractedCollisions.Spheres.Add(Sphere);
-			}
-
-			// Add boxes
-			for (const FKBoxElem& BoxElem : AggGeom.BoxElems)
-			{
-				// Add extracted box collision data
-				FClothCollisionPrim_Box Box;
-				Box.LocalPosition = BoxElem.Center;
-				Box.LocalRotation = BoxElem.Rotation.Quaternion();
-				Box.HalfExtents = FVector(BoxElem.X, BoxElem.Y, BoxElem.Z) * 0.5f;
-				Box.BoneIndex = MappedBoneIndex;
-				ExtractedCollisions.Boxes.Add(Box);
-			}
-
-			// Add tapered capsules
-			for (const FKTaperedCapsuleElem& TaperedCapsuleElem : AggGeom.TaperedCapsuleElems)
-			{
-				if (TaperedCapsuleElem.Length == 0)
-				{
-					// Add extracted sphere collision data
-					FClothCollisionPrim_Sphere Sphere;
-					Sphere.LocalPosition = TaperedCapsuleElem.Center;
-					Sphere.Radius = FMath::Max(TaperedCapsuleElem.Radius0, TaperedCapsuleElem.Radius1);
-					Sphere.BoneIndex = MappedBoneIndex;
-					ExtractedCollisions.Spheres.Add(Sphere);
-				}
-				else
-				{
-					// Add extracted spheres collision data
-					FClothCollisionPrim_Sphere Sphere0;
-					FClothCollisionPrim_Sphere Sphere1;
-					const FVector OrientedDirection = TaperedCapsuleElem.Rotation.RotateVector(FVector::UpVector);
-					const FVector HalfDim = OrientedDirection * (TaperedCapsuleElem.Length / 2.f);
-					Sphere0.LocalPosition = TaperedCapsuleElem.Center + HalfDim;
-					Sphere1.LocalPosition = TaperedCapsuleElem.Center - HalfDim;
-					Sphere0.Radius = TaperedCapsuleElem.Radius0;
-					Sphere1.Radius = TaperedCapsuleElem.Radius1;
-					Sphere0.BoneIndex = MappedBoneIndex;
-					Sphere1.BoneIndex = MappedBoneIndex;
-
-					// Add extracted sphere connection collision data
-					FClothCollisionPrim_SphereConnection SphereConnection;
-					SphereConnection.SphereIndices[0] = ExtractedCollisions.Spheres.Add(Sphere0);
-					SphereConnection.SphereIndices[1] = ExtractedCollisions.Spheres.Add(Sphere1);
-					ExtractedCollisions.SphereConnections.Add(SphereConnection);
-				}
-			}
-
-#if !PLATFORM_LUMIN && !PLATFORM_ANDROID  // TODO(Kriss.Gossart): Compile on Android and fix whatever errors the following code is causing
-			// Add convexes
-			for (const FKConvexElem& ConvexElem : AggGeom.ConvexElems)
-			{
-				// Add stub for extracted collision data
-				FClothCollisionPrim_Convex Convex;
-				Convex.BoneIndex = MappedBoneIndex;
-#if PHYSICS_INTERFACE_PHYSX
-				// Collision bodies are stored in PhysX specific data structures so they can only be imported if we enable PhysX.
-				const physx::PxConvexMesh* const PhysXMesh = ConvexElem.GetConvexMesh();  // TODO(Kriss.Gossart): Deal with this legacy structure in a different place, so that there's only TConvex
-				const int32 NumPolygons = int32(PhysXMesh->getNbPolygons());
-				Convex.Planes.SetNumUninitialized(NumPolygons);
-				for (int32 i = 0; i < NumPolygons; ++i)
-				{
-					physx::PxHullPolygon Poly;
-					PhysXMesh->getPolygonData(i, Poly);
-					check(Poly.mNbVerts == 3);
-					const auto Indices = PhysXMesh->getIndexBuffer() + Poly.mIndexBase;
-
-					Convex.Planes[i] = FPlane(
-						ConvexElem.VertexData[Indices[0]],
-						ConvexElem.VertexData[Indices[1]],
-						ConvexElem.VertexData[Indices[2]]);
-				}
-
-				// Rebuild surface points
-				Convex.RebuildSurfacePoints();
-
-#elif WITH_CHAOS  // #if PHYSICS_INTERFACE_PHYSX
-				const Chaos::FImplicitObject& ChaosConvexMesh = *ConvexElem.GetChaosConvexMesh();
-				const Chaos::FConvex& ChaosConvex = ChaosConvexMesh.GetObjectChecked<Chaos::FConvex>();
-
-				// Copy planes
-				const TArray<TPlaneConcrete<float, 3>>& Planes = ChaosConvex.GetFaces();
-				Convex.Planes.Reserve(Planes.Num());
-				for (const TPlaneConcrete<float, 3>& Plane : Planes)
-				{
-					Convex.Planes.Add(FPlane(Plane.X(), Plane.Normal()));
-				}
-
-				// Copy surface points
-				const uint32 NumSurfacePoints = ChaosConvex.GetSurfaceParticles().Size();
-				Convex.SurfacePoints.Reserve(NumSurfacePoints);
-				for (uint32 ParticleIndex = 0; ParticleIndex < NumSurfacePoints; ++ParticleIndex)
-				{
-					Convex.SurfacePoints.Add(ChaosConvex.GetSurfaceParticles().X(ParticleIndex));
-				}
-#endif  // #if PHYSICS_INTERFACE_PHYSX #elif WITH_CHAOS
-
-				// Add extracted collision data
-				ExtractedCollisions.Convexes.Add(Convex);
-			}
-#endif  // #if !PLATFORM_LUMIN && !PLATFORM_ANDROID
-
-		}  // End for PhysAsset->SkeletalBodySetups
-
-		// Add collisions particles
-		UE_LOG(LogChaosCloth, VeryVerbose, TEXT("Adding physics asset collisions..."));
-		AddCollisions(ExtractedCollisions, UsedBoneIndices, InSimDataIndex);
-
-	}  // End if Asset->PhysicsAsset
-}
-
-void ClothingSimulation::ExtractLegacyAssetCollisions(const UClothingAssetCommon* Asset, int32 InSimDataIndex)
-{
-	const FClothLODDataCommon& AssetLodData = Asset->LodData[0];
-	const FClothCollisionData& LodCollData = AssetLodData.CollisionData;
-	if (LodCollData.Spheres.Num() || LodCollData.SphereConnections.Num() || LodCollData.Convexes.Num())
-	{
-		UE_LOG(LogChaosCloth, VeryVerbose, TEXT("Adding legacy cloth asset collisions..."));
-		AddCollisions(LodCollData, Asset->UsedBoneIndices, InSimDataIndex);
-	}
-}
-
-void ClothingSimulation::AddCollisions(const FClothCollisionData& ClothCollisionData, const TArray<int32>& UsedBoneIndices, int32 InSimDataIndex)
-{
-	TGeometryClothParticles<float, 3>& CollisionParticles = Evolution->CollisionParticles();
-
-	// Capsules
-	TSet<int32> CapsuleEnds;
-	const int32 NumCapsules = ClothCollisionData.SphereConnections.Num();
-	if (NumCapsules)
-	{
-		const uint32 Offset = Evolution->AddCollisionParticles(NumCapsules, InSimDataIndex);
-
-		CapsuleEnds.Reserve(NumCapsules * 2);
-		for (uint32 i = Offset; i < CollisionParticles.Size(); ++i)
-		{
-			// Always initialize the collision particle transforms before setting any geometry as otherwise NaNs gets detected during the bounding box updates
-			CollisionParticles.X(i) = Chaos::TVector<float, 3>(0.f);
-			CollisionParticles.R(i) = Chaos::TRotation<float, 3>::FromIdentity();
-
-			const FClothCollisionPrim_SphereConnection& Connection = ClothCollisionData.SphereConnections[i - Offset];
-
-			const int32 SphereIndex0 = Connection.SphereIndices[0];
-			const int32 SphereIndex1 = Connection.SphereIndices[1];
-			checkSlow(SphereIndex0 != SphereIndex1);
-			const FClothCollisionPrim_Sphere& Sphere0 = ClothCollisionData.Spheres[SphereIndex0];
-			const FClothCollisionPrim_Sphere& Sphere1 = ClothCollisionData.Spheres[SphereIndex1];
-
-			const int32 MappedIndex = UsedBoneIndices.IsValidIndex(Sphere0.BoneIndex) ? UsedBoneIndices[Sphere0.BoneIndex] : INDEX_NONE;
-
-			BoneIndices[i] = GetMappedBoneIndex(UsedBoneIndices, Sphere0.BoneIndex);
-			checkSlow(Sphere0.BoneIndex == Sphere1.BoneIndex);
-			UE_CLOG(Sphere0.BoneIndex != Sphere1.BoneIndex,
-				LogChaosCloth, Warning, TEXT("Found a legacy Apex cloth asset with a collision capsule spanning across two bones. This is not supported with the current system."));
-			UE_LOG(LogChaosCloth, VeryVerbose, TEXT("Found collision capsule on bone index %d."), BoneIndices[i]);
-
-			const Chaos::TVector<float, 3> X0 = Sphere0.LocalPosition;
-			const Chaos::TVector<float, 3> X1 = Sphere1.LocalPosition;
-			const Chaos::TVector<float, 3> Axis = X1 - X0;
-			const float AxisSize = Axis.Size();
-
-			const float Radius0 = Sphere0.Radius;
-			const float Radius1 = Sphere1.Radius;
-			float MinRadius, MaxRadius;
-			if (Radius0 <= Radius1) { MinRadius = Radius0; MaxRadius = Radius1; }
-			else { MinRadius = Radius1; MaxRadius = Radius0; }
-
-			if (AxisSize < KINDA_SMALL_NUMBER)
-			{
-				// Sphere
-				BaseTransforms[i] = Chaos::TRigidTransform<float, 3>(FTransform::Identity);
-
-				CollisionParticles.SetDynamicGeometry(
-					i,
-					MakeUnique<Chaos::TSphere<float, 3>>(
-						X0,
-						MaxRadius));
-			}
-			else if (MaxRadius - MinRadius < KINDA_SMALL_NUMBER)
-			{
-				// Capsule
-				const Chaos::TVector<float, 3> Center = (X0 + X1) * 0.5f;  // Construct a capsule centered at the origin along the Z axis
-				const Chaos::TRotation<float, 3> Rotation = Chaos::TRotation<float, 3>::FromRotatedVector(
-					Chaos::TVector<float, 3>::AxisVector(2),
-					Axis.GetSafeNormal());
-
-				BaseTransforms[i] = Chaos::TRigidTransform<float, 3>(Center, Rotation);
-
-				const float HalfHeight = AxisSize * 0.5f;
-				CollisionParticles.SetDynamicGeometry(
-					i,
-					MakeUnique<Chaos::TCapsule<float>>(
-						Chaos::TVector<float, 3>(0.f, 0.f, -HalfHeight), // Min
-						Chaos::TVector<float, 3>(0.f, 0.f, HalfHeight), // Max
-						MaxRadius));
-			}
-			else
-			{
-				// Tapered capsule
-				BaseTransforms[i] = Chaos::TRigidTransform<float, 3>(FTransform::Identity);
-
-				TArray<TUniquePtr<Chaos::FImplicitObject>> Objects;
-				Objects.Reserve(3);
-				Objects.Add(TUniquePtr<Chaos::FImplicitObject>(
-					new Chaos::TTaperedCylinder<float>(X0, X1, Radius0, Radius1)));
-				Objects.Add(TUniquePtr<Chaos::FImplicitObject>(
-					new Chaos::TSphere<float, 3>(X0, Radius0)));
-				Objects.Add(TUniquePtr<Chaos::FImplicitObject>(
-					new Chaos::TSphere<float, 3>(X1, Radius1)));
-				CollisionParticles.SetDynamicGeometry(
-					i,
-					MakeUnique<Chaos::FImplicitObjectUnion>(MoveTemp(Objects)));  // TODO(Kriss.Gossart): Replace this once a TTaperedCapsule implicit type is implemented (note: this tapered cylinder with spheres is an approximation of a real tapered capsule)
-			}
-
-			// Skip spheres added as end caps for the capsule.
-			CapsuleEnds.Add(SphereIndex0);
-			CapsuleEnds.Add(SphereIndex1);
-		}
-	}
-
-	// Spheres
-	const int32 NumSpheres = ClothCollisionData.Spheres.Num() - CapsuleEnds.Num();
-	if (NumSpheres != 0)
-	{
-		const uint32 Offset = Evolution->AddCollisionParticles(NumSpheres, InSimDataIndex);
-		// i = CollisionParticles index, j = Spheres index
-		for (uint32 i = Offset, j = 0; j < (uint32)ClothCollisionData.Spheres.Num(); ++j)
-		{
-			// Skip spheres that are the end caps of capsules.
-			if (CapsuleEnds.Contains(j))
-				continue;
-
-			const FClothCollisionPrim_Sphere& Sphere = ClothCollisionData.Spheres[j];
-
-			// Always initialize the collision particle transforms before setting any geometry as otherwise NaNs gets detected during the bounding box updates
-			CollisionParticles.X(i) = Chaos::TVector<float, 3>(0.f);
-			CollisionParticles.R(i) = Chaos::TRotation<float, 3>::FromIdentity();
-
-			BoneIndices[i] = GetMappedBoneIndex(UsedBoneIndices, Sphere.BoneIndex);
-			UE_LOG(LogChaosCloth, VeryVerbose, TEXT("Found collision sphere on bone index %d."), BoneIndices[i]);
-
-			BaseTransforms[i] = Chaos::TRigidTransform<float, 3>(FTransform::Identity);
-
-			CollisionParticles.SetDynamicGeometry(
-				i,
-				MakeUnique<Chaos::TSphere<float, 3>>(
-					Sphere.LocalPosition,
-					Sphere.Radius));
-
-			++i;
-		}
-	}
-
-	// Convexes
-	const uint32 NumConvexes = ClothCollisionData.Convexes.Num();
-	if (NumConvexes != 0)
-	{
-		const uint32 Offset = Evolution->AddCollisionParticles(NumConvexes, InSimDataIndex);
-		for (uint32 i = Offset; i < CollisionParticles.Size(); ++i)
-		{
-			const FClothCollisionPrim_Convex& Convex = ClothCollisionData.Convexes[i - Offset];
-
-			// Always initialize the collision particle transforms before setting any geometry as otherwise NaNs gets detected during the bounding box updates
-			CollisionParticles.X(i) = Chaos::TVector<float, 3>(0.f);
-			CollisionParticles.R(i) = Chaos::TRotation<float, 3>::FromIdentity();
-
-			BaseTransforms[i] = Chaos::TRigidTransform<float, 3>(FTransform::Identity);
-
-			BoneIndices[i] = GetMappedBoneIndex(UsedBoneIndices, Convex.BoneIndex);
-			UE_LOG(LogChaosCloth, VeryVerbose, TEXT("Found collision convex on bone index %d."), BoneIndices[i]);
-
-			const int32 NumSurfacePoints = Convex.SurfacePoints.Num();
-			const int32 NumPlanes = Convex.Planes.Num();
-
-			if (NumSurfacePoints < 4)
-			{
-				UE_LOG(LogChaosCloth, Warning, TEXT("Invalid convex collision: not enough surface points."));
-			}
-			else if (NumPlanes < 4)
-			{
-				UE_LOG(LogChaosCloth, Warning, TEXT("Invalid convex collision: not enough planes."));
-			}
-			else
-			{
-				// Retrieve convex planes
-				TArray<TPlaneConcrete<float, 3>> Planes;
-				Planes.Reserve(Convex.Planes.Num());
-				for (const FPlane& Plane : Convex.Planes)
-				{
-					FPlane NormalizedPlane(Plane);
-					if (NormalizedPlane.Normalize())
-					{
-						const Chaos::TVector<float, 3> Normal(static_cast<FVector>(NormalizedPlane));
-						const Chaos::TVector<float, 3> Base = Normal * NormalizedPlane.W;
-
-						Planes.Add(Chaos::TPlaneConcrete<float, 3>(Base, Normal));
-					}
-					else
-					{
-						UE_LOG(LogChaosCloth, Warning, TEXT("Invalid convex collision: bad plane normal."));
-						break;
-					}
-				}
-
-				if (Planes.Num() == Convex.Planes.Num())
-				{
-					// Retrieve particles
-					TParticles<float, 3> SurfaceParticles;
-					SurfaceParticles.Resize(NumSurfacePoints);
-					for (int32 ParticleIndex = 0; ParticleIndex < NumSurfacePoints; ++ParticleIndex)
-					{
-						SurfaceParticles.X(ParticleIndex) = Convex.SurfacePoints[ParticleIndex];
-					}
-
-					// Setup the collision particle geometry
-					CollisionParticles.SetDynamicGeometry(i, MakeUnique<Chaos::FConvex>(MoveTemp(Planes), MoveTemp(SurfaceParticles)));
-				}
-			}
-
-			if (!CollisionParticles.DynamicGeometry(i))
-			{
-				UE_LOG(LogChaosCloth, Warning, TEXT("Replacing invalid convex collision by a default unit sphere."));
-				CollisionParticles.SetDynamicGeometry(i, MakeUnique<Chaos::TSphere<float, 3>>(Chaos::TVector<float, 3>(0.0f), 1.0f));  // Default to a unit sphere to replace the faulty convex
-			}
-		}
-	}
-
-	// Boxes
-	const uint32 NumBoxes = ClothCollisionData.Boxes.Num();
-	if (NumBoxes != 0)
-	{
-		const uint32 Offset = Evolution->AddCollisionParticles(NumBoxes, InSimDataIndex);
-		for (uint32 i = Offset; i < CollisionParticles.Size(); ++i)
-		{
-			const FClothCollisionPrim_Box& Box = ClothCollisionData.Boxes[i - Offset];
-			
-			// Always initialize the collision particle transforms before setting any geometry as otherwise NaNs gets detected during the bounding box updates
-			CollisionParticles.X(i) = Chaos::TVector<float, 3>(0.f);
-			CollisionParticles.R(i) = Chaos::TRotation<float, 3>::FromIdentity();
-
-			BaseTransforms[i] = Chaos::TRigidTransform<float, 3>(Box.LocalPosition, Box.LocalRotation);
-			
-			BoneIndices[i] = GetMappedBoneIndex(UsedBoneIndices, Box.BoneIndex);
-			UE_LOG(LogChaosCloth, VeryVerbose, TEXT("Found collision box on bone index %d."), BoneIndices[i]);
-
-			CollisionParticles.SetDynamicGeometry(i, MakeUnique<Chaos::TBox<float, 3>>(-Box.HalfExtents, Box.HalfExtents));
-		}
-	}
-
-	UE_LOG(LogChaosCloth, VeryVerbose, TEXT("Added collisions: %d spheres, %d capsules, %d convexes, %d boxes."), NumSpheres, NumCapsules, NumConvexes, NumBoxes);
-}
-
-void ClothingSimulation::Simulate(IClothingSimulationContext* InContext)
+void FClothingSimulation::Simulate(IClothingSimulationContext* InContext)
 {
 	SCOPE_CYCLE_COUNTER(STAT_ChaosClothSimulate);
-	const ClothingSimulationContext* const Context = static_cast<ClothingSimulationContext*>(InContext);
-	if (Context->DeltaSeconds == 0.f)
-	{
-		return;
-	}
-
-	const double StartTime = FPlatformTime::Seconds();
-
-	// Filter delta time to smoothen time variations and prevent unwanted vibrations
-	static const float DeltaTimeDecay = 0.1f;
-	DeltaTime = DeltaTime + (Context->DeltaSeconds - DeltaTime) * DeltaTimeDecay;
-
-	// Check teleport modes
-	const bool bTeleport = (Context->TeleportMode > EClothingTeleportMode::None);
-	const bool bTeleportAndReset = (Context->TeleportMode == EClothingTeleportMode::TeleportAndReset);
-
-	// Get New Animation Positions and Normals + deal with local space & teleportation
-	{
-		SCOPE_CYCLE_COUNTER(STAT_ChaosClothGetAnimationData);
-
-		check(OldAnimationPositions.Num() == AnimationPositions.Num());
-		Swap(OldAnimationPositions, AnimationPositions);
-
-		check(OldCollisionTransforms.Num() == CollisionTransforms.Num());
-		Swap(OldCollisionTransforms, CollisionTransforms);
-
-		// Update the local space transform
-		const FVector PrevLocalSpaceLocation = LocalSpaceLocation;
-		if (bUseLocalSpaceSimulation)
-		{
-			LocalSpaceLocation = Context->ComponentToWorld.GetLocation();
-		}
-		const FVector DeltaLocalSpaceLocation = LocalSpaceLocation - PrevLocalSpaceLocation;
-
-		// Iterate all cloths
-		TPBDParticles<float, 3>& Particles = Evolution->Particles();
-
-		for (int32 Index = 0; Index < Assets.Num(); ++Index)
-		{
-			const UClothingAssetCommon* const Asset = Assets[Index];
-			if (!Asset) { continue; }
-
-			// Set wind velocity, used by the velocity field lambda
-			WindVelocities[Index] = Context->WindVelocity * ChaosClothingSimulationDefault::WorldScale;  // Wind speed is set in m/s and need to be converted to cm/s
-
-			// Set gravity, using the legacy priority: 1) game override, 2) config override, 3) world gravity
-			const TVector<float, 3> Gravity = TVector<float, 3>(
-				bUseGravityOverride ? GravityOverride * GravityScales[Index] :
-				bUseGravityConfigs[Index] ? GravityConfigs[Index] :  // Config gravity is not subject to scale
-				Context->WorldGravity * GravityScales[Index]);
-
-			// Update gravity
-			Evolution->GetGravityForces(Index).SetAcceleration(Gravity);
-
-			const uint32 Offset = IndexToRangeMap[Index][0];
-			const uint32 Range = IndexToRangeMap[Index][1];
-
-			// Update collision transforms using new local space transform
-			UpdateCollisionTransforms(*Context, Index);
-
-			// Update animation transforms via skinning
-			// Optimization note:
-			// This function usually receives the RootBoneTransform in order to transform the result from Component space to RootBone space.
-			// (So the mesh vectors and positions (Mesh is in component space) is multiplied by Inv(RootBoneTransform) at the end of the function)
-			// We actually require world space coordinates so will instead pass Inv(ComponentToWorld)
-			// This saves a lot of Matrix multiplication work later
-			FTransform ComponentToLocalSpace = Context->ComponentToWorld;
-			ComponentToLocalSpace.AddToTranslation(-LocalSpaceLocation);
-
-			ClothingMeshUtils::SkinPhysicsMesh<true, false>(
-				Asset->UsedBoneIndices,
-				Asset->LodData[0].PhysicalMeshData,
-				ComponentToLocalSpace,
-				Context->RefToLocals.GetData(),
-				Context->RefToLocals.Num(),
-				reinterpret_cast<TArray<FVector>&>(AnimationPositions),
-				reinterpret_cast<TArray<FVector>&>(AnimationNormals), Offset);
-
-			// Update root bone reference transforms
-			const FTransform RootBoneTransform = Context->BoneTransforms[Asset->ReferenceBoneIndex];
-			const FTransform PrevRootBoneWorldTransform = RootBoneWorldTransforms[Index];
-			RootBoneWorldTransforms[Index] = RootBoneTransform * Context->ComponentToWorld;
-
-			FTransform PrevRootBoneLocalTransform = PrevRootBoneWorldTransform;
-			PrevRootBoneLocalTransform.AddToTranslation(-PrevLocalSpaceLocation);
-
-			// Teleport & reset
-			if (bTeleportAndReset)
-			{
-				UE_LOG(LogChaosCloth, Verbose, TEXT("Teleport & Reset"));
-				for (uint32 i = Offset; i < Range; ++i)
-				{
-					// Update initial state for particles
-					Particles.P(i) = Particles.X(i) = AnimationPositions[i];
-					Particles.V(i) = Chaos::TVector<float, 3>(0.f);
-
-					// Update anim initial state (target updated by skinning)
-					OldAnimationPositions[i] = AnimationPositions[i];
-				}
-				ForAllCollisions([this](TGeometryClothParticles<float, 3>& CollisionParticles, uint32 i)
-				{
-					// Update initial state for collisions
-					OldCollisionTransforms[i] = CollisionTransforms[i];
-					CollisionParticles.X(i) = CollisionTransforms[i].GetTranslation();
-					CollisionParticles.R(i) = CollisionTransforms[i].GetRotation();
-				}, Index);
-			}
-			// Teleport only
-			else if (bTeleport)
-			{
-				UE_LOG(LogChaosCloth, Verbose, TEXT("Teleport before: %s, after: %s"), *PrevRootBoneWorldTransform.ToString(), *RootBoneWorldTransforms[Index].ToString());
-				const FTransform DeltaTransform = RootBoneWorldTransforms[Index].GetRelativeTransform(PrevRootBoneWorldTransform);
-				const FMatrix Matrix = (PrevRootBoneLocalTransform.Inverse() * DeltaTransform * PrevRootBoneLocalTransform).ToMatrixNoScale();
-
-				for (uint32 i = Offset; i < Range; ++i)
-				{
-					// Update initial state for particles
-					Particles.P(i) = Particles.X(i) = Matrix.TransformPosition(Particles.X(i)) - DeltaLocalSpaceLocation;
-					Particles.V(i) = Matrix.TransformVector(Particles.V(i));
-
-					// Update anim initial state (target updated by skinning)
-					OldAnimationPositions[i] = Matrix.TransformPosition(OldAnimationPositions[i]) - DeltaLocalSpaceLocation;
-				}
-
-				ForAllCollisions([this, &Matrix, &DeltaLocalSpaceLocation](TGeometryClothParticles<float, 3>& CollisionParticles, uint32 i)
-				{
-					// Update initial state for collisions
-					OldCollisionTransforms[i] = Matrix * OldCollisionTransforms[i];
-					OldCollisionTransforms[i].AddToTranslation(-DeltaLocalSpaceLocation);
-					CollisionParticles.X(i) = OldCollisionTransforms[i].GetTranslation();
-					CollisionParticles.R(i) = OldCollisionTransforms[i].GetRotation();
-				}, Index);
-			}
-			// Apply reference space velocity scales
-			else if (AngularDeltaRatios[Index] > KINDA_SMALL_NUMBER ||
-				LinearDeltaRatios[Index].X > KINDA_SMALL_NUMBER ||
-				LinearDeltaRatios[Index].Y > KINDA_SMALL_NUMBER ||
-				LinearDeltaRatios[Index].Z > KINDA_SMALL_NUMBER)
-			{
-				// Calculate deltas
-				const FTransform DeltaTransform = RootBoneWorldTransforms[Index].GetRelativeTransform(PrevRootBoneWorldTransform);
-
-				const FVector DeltaPosition = LinearDeltaRatios[Index] * DeltaTransform.GetTranslation();
-
-				FQuat DeltaRotation = DeltaTransform.GetRotation();
-				FVector Axis;
-				float DeltaAngle;
-				DeltaRotation.ToAxisAndAngle(Axis, DeltaAngle);
-				if (DeltaAngle > PI) { DeltaAngle -= 2.f * PI; }
-				DeltaAngle *= AngularDeltaRatios[Index];
-				DeltaRotation = FQuat(Axis, DeltaAngle);
-				DeltaRotation.Normalize();  // ToMatrixNoScale does not like quaternions built straight from axis angles without being normalized (although they should have been already).
-
-				// Transform points back into the previous frame of reference before applying the adjusted deltas 
-				const FMatrix Matrix = (PrevRootBoneLocalTransform.Inverse() * FTransform(DeltaRotation, DeltaPosition) * PrevRootBoneLocalTransform).ToMatrixNoScale();
-
-				for (uint32 i = Offset; i < Range; ++i)
-				{
-					// Update initial state for particles
-					Particles.P(i) = Particles.X(i) = Matrix.TransformPosition(Particles.X(i)) - DeltaLocalSpaceLocation;
-					Particles.V(i) = Matrix.TransformVector(Particles.V(i));
-
-					// Update anim initial state (target updated by skinning)
-					OldAnimationPositions[i] = Matrix.TransformPosition(OldAnimationPositions[i]) - DeltaLocalSpaceLocation;
-				}
+	const FClothingSimulationContext* const Context = static_cast<FClothingSimulationContext*>(InContext);
+	if (Context->DeltaSeconds == 0.f)
+	{
+		return;
+	}
 
-				ForAllCollisions([this, &Matrix, &DeltaLocalSpaceLocation](TGeometryClothParticles<float, 3>& CollisionParticles, uint32 i)
-				{
-					// Update initial state for collisions
-					OldCollisionTransforms[i] = Matrix * OldCollisionTransforms[i];
-					OldCollisionTransforms[i].AddToTranslation(-DeltaLocalSpaceLocation);
+	const double StartTime = FPlatformTime::Seconds();
 
-					CollisionParticles.X(i) = OldCollisionTransforms[i].GetTranslation();
-					CollisionParticles.R(i) = OldCollisionTransforms[i].GetRotation();
-				}, Index);
-			}
-			else if (bUseLocalSpaceSimulation)
-			{
-				for (uint32 i = Offset; i < Range; ++i)
-				{
-					// Update initial state for particles
-					Particles.P(i) = Particles.X(i) -= DeltaLocalSpaceLocation;
+	// Update Solver animatable parameters
+	Solver->SetLocalSpaceLocation(bUseLocalSpaceSimulation ? Context->ComponentToWorld.GetLocation() : TVector<float, 3>(0.f));
+	Solver->SetWindVelocity(Context->WindVelocity);
+	Solver->SetGravity(bUseGravityOverride ? GravityOverride : Context->WorldGravity);
+	Solver->EnableClothGravityOverride(!bUseGravityOverride);  // Disable all cloth gravity overrides when the interactor takes over
 
-					// Update anim initial state (target updated by skinning)
-					OldAnimationPositions[i] -= DeltaLocalSpaceLocation;
-				}
+	// Check teleport modes
+	const bool bNeedsReset = (Context->TeleportMode == EClothingTeleportMode::TeleportAndReset);
+	const bool bNeedsTeleport = (Context->TeleportMode > EClothingTeleportMode::None);
 
-				ForAllCollisions([this, &DeltaLocalSpaceLocation](TGeometryClothParticles<float, 3>& CollisionParticles, uint32 i)
-				{
-					// Update initial state for collisions
-					OldCollisionTransforms[i].AddToTranslation(-DeltaLocalSpaceLocation);
-					CollisionParticles.X(i) = OldCollisionTransforms[i].GetTranslation();
-				}, Index);
-			}
+	for (const TUniquePtr<FClothingSimulationCloth>& Cloth : Cloths)
+	{
+		// Update Cloth animatable parameters
+		Cloth->SetAnimDriveSpringStiffness(AnimDriveSpringStiffness);
+		Cloth->SetMaxDistancesMultiplier(Context->MaxDistanceScale);
 
-			// Update max distance multiplier
-			MaxDistancesMultipliers[Index] = Context->MaxDistanceScale;
+		if (bNeedsReset)
+		{
+			Cloth->Reset();
+		}
+		if (bNeedsTeleport)
+		{
+			Cloth->Teleport();
 		}
 	}
 
-	// Advance Sim
-	const int32 NumSubstepsRead = NumSubsteps;  // atomic read
-	const float SubstepDeltaTime = DeltaTime / (float)NumSubstepsRead;
-	
-	for (int32 i = 0; i < NumSubstepsRead; ++i)
-	{
-		Evolution->AdvanceOneTimeStep(SubstepDeltaTime);
-	}
-
-	Time = Evolution->GetTime();
-	UE_LOG(LogChaosCloth, VeryVerbose, TEXT("DeltaTime: %.6f, FilteredDeltaTime: %.6f, Time = %.6f,  MaxPhysicsDelta = %.6f"), Context->DeltaSeconds, DeltaTime, Time, FClothingSimulationCommon::MaxPhysicsDelta);
+	// Step the simulation
+	Solver->Update(Context->DeltaSeconds);
 
 	// Update simulation time in ms (and provide an instant average instead of the value in real-time)
 	const float PrevSimulationTime = SimulationTime;  // Copy the atomic to prevent a re-read
@@ -1572,12 +298,12 @@ void ClothingSimulation::Simulate(IClothingSimulationContext* InContext)
 	if (ChaosClothingSimulationConsoleVariables::CVarDebugDrawBounds          .GetValueOnAnyThread()) { DebugDrawBounds              (); }
 	if (ChaosClothingSimulationConsoleVariables::CVarDebugDrawGravity         .GetValueOnAnyThread()) { DebugDrawGravity             (); }
 	if (ChaosClothingSimulationConsoleVariables::CVarDebugDrawPhysMeshWired   .GetValueOnAnyThread()) { DebugDrawPhysMeshWired       (); }
+	if (ChaosClothingSimulationConsoleVariables::CVarDebugDrawAnimMeshWired   .GetValueOnAnyThread()) { DebugDrawAnimMeshWired       (); }
 	if (ChaosClothingSimulationConsoleVariables::CVarDebugPointNormals        .GetValueOnAnyThread()) { DebugDrawPointNormals        (); }
 	if (ChaosClothingSimulationConsoleVariables::CVarDebugInversedPointNormals.GetValueOnAnyThread()) { DebugDrawInversedPointNormals(); }
-	if (ChaosClothingSimulationConsoleVariables::CVarDebugFaceNormals         .GetValueOnAnyThread()) { DebugDrawFaceNormals         (); }
-	if (ChaosClothingSimulationConsoleVariables::CVarDebugInversedFaceNormals .GetValueOnAnyThread()) { DebugDrawInversedFaceNormals (); }
 	if (ChaosClothingSimulationConsoleVariables::CVarDebugCollision           .GetValueOnAnyThread()) { DebugDrawCollision           (); }
 	if (ChaosClothingSimulationConsoleVariables::CVarDebugBackstops           .GetValueOnAnyThread()) { DebugDrawBackstops           (); }
+	if (ChaosClothingSimulationConsoleVariables::CVarDebugBackstopDistances   .GetValueOnAnyThread()) { DebugDrawBackstopDistances   (); }
 	if (ChaosClothingSimulationConsoleVariables::CVarDebugMaxDistances        .GetValueOnAnyThread()) { DebugDrawMaxDistances        (); }
 	if (ChaosClothingSimulationConsoleVariables::CVarDebugAnimDrive           .GetValueOnAnyThread()) { DebugDrawAnimDrive           (); }
 	if (ChaosClothingSimulationConsoleVariables::CVarDebugLongRangeConstraint .GetValueOnAnyThread()) { DebugDrawLongRangeConstraint (); }
@@ -1585,351 +311,252 @@ void ClothingSimulation::Simulate(IClothingSimulationContext* InContext)
 #endif  // #if CHAOS_DEBUG_DRAW
 }
 
-void ClothingSimulation::GetSimulationData(
+void FClothingSimulation::GetSimulationData(
 	TMap<int32, FClothSimulData>& OutData,
 	USkeletalMeshComponent* InOwnerComponent,
 	USkinnedMeshComponent* InOverrideComponent) const
 {
-	const int32 NumSimulatedCloth = IndexToRangeMap.Num();
-	if (OutData.Num() != NumSimulatedCloth)
+	SCOPE_CYCLE_COUNTER(STAT_ChaosClothGetSimulationData);
+
+	// Reset map when new cloths have appeared
+	if (OutData.Num() != Cloths.Num())
 	{
 		OutData.Reset();
 	}
 
+	// Retrieve cloths' particle positions
 	const FTransform& OwnerTransform = InOwnerComponent->GetComponentTransform();
-	for (int32 i = 0; i < NumSimulatedCloth; ++i)
+	const TVector<float, 3>& LocalSpaceLocation = Solver->GetLocalSpaceLocation();
+
+	for (const TUniquePtr<FClothingSimulationCloth>& Cloth : Cloths)
 	{
-		FClothSimulData& Data = OutData.FindOrAdd(i);
+		const int32 AssetIndex = Cloth->GetGroupId();
+		FClothSimulData& Data = OutData.FindOrAdd(AssetIndex);
 
-		const TUniquePtr<Chaos::TTriangleMesh<float>>& Mesh = Meshes[i];
-		if (!Mesh)
-		{
-			Data.Reset();
-			continue;
-		}
-		Mesh->GetFaceNormals(FaceNormals[i], Evolution->Particles().X(), false);  // No need to add a point index offset here since that is baked into the triangles
-		Mesh->GetPointNormals(PointNormals[i], FaceNormals[i], /*bReturnEmptyOnError =*/ false, /*bFillAtStartIndex =*/ false);
+		// Output data in component space
+		Data.ComponentRelativeTransform = FTransform::Identity;
+		Data.Transform = OwnerTransform;
 
-		const UClothingAssetCommon* const Asset = Assets[i];
-		if (!Asset)
-		{
-			Data.Reset();
-			continue;
-		}
+		Data.Positions = Cloth->GetParticlePositions(Solver.Get());
+		Data.Normals = Cloth->GetParticleNormals(Solver.Get());
 
-		const TArray<FTransform>& ComponentSpaceTransforms = InOverrideComponent ?
-			InOverrideComponent->GetComponentSpaceTransforms() :
-			InOwnerComponent->GetComponentSpaceTransforms();
-		if (!ensure(ComponentSpaceTransforms.IsValidIndex(Asset->ReferenceBoneIndex)))
+		for (int32 Index = 0; Index < Data.Positions.Num(); ++Index)
 		{
-			UE_LOG(LogSkeletalMesh, Warning,
-				TEXT("Failed to write back clothing simulation data for component '%s' as bone transforms are invalid."),
-				*InOwnerComponent->GetName());
-			Data.Reset();
-			continue;
-		}
-
-		FTransform RootBoneTransform = ComponentSpaceTransforms[Asset->ReferenceBoneIndex];
-		RootBoneTransform.SetScale3D(FVector(1.0f));
-		RootBoneTransform *= OwnerTransform;
-		Data.Transform = RootBoneTransform;
-		Data.ComponentRelativeTransform = OwnerTransform.Inverse();
-
-		const Chaos::TVector<uint32, 2>& VertexDomain = IndexToRangeMap[i];
-		const uint32 VertexRange = VertexDomain[1] - VertexDomain[0];
-		Data.Positions.SetNum(VertexRange);
-        Data.Normals.SetNum(VertexRange);
-		for (uint32 j = VertexDomain[0]; j < VertexDomain[1]; ++j)
-        {
-			const uint32 LocalIndex = j - VertexDomain[0];
-            Data.Positions[LocalIndex] = Evolution->Particles().X(j) + LocalSpaceLocation;
-            Data.Normals[LocalIndex] = -PointNormals[i][LocalIndex]; // Note the Normals are inverted due to how barycentric coordinates are calculated (see GetPointBaryAndDist in ClothingMeshUtils.cpp)
+			Data.Positions[Index] = OwnerTransform.InverseTransformPositionNoScale(Data.Positions[Index] + LocalSpaceLocation);  // Move into world space first
+			Data.Normals[Index] = OwnerTransform.InverseTransformVectorNoScale(-Data.Normals[Index]);  // Normals are inverted due to how barycentric coordinates are calculated (see GetPointBaryAndDist in ClothingMeshUtils.cpp)
 		}
     }
 }
 
-FBoxSphereBounds ClothingSimulation::GetBounds(const USkeletalMeshComponent* InOwnerComponent) const
+FBoxSphereBounds FClothingSimulation::GetBounds(const USkeletalMeshComponent* InOwnerComponent) const
 {
-	FBoxSphereBounds Bounds(EForceInit::ForceInit);
+	check(Solver);
+	const FBoxSphereBounds Bounds = Solver->CalculateBounds();
 
-	// Calculate simulation bounds (in world space)
-	uint32 NumBoundedCloths = 0;
-	const TPBDParticles<float, 3>& Particles = Evolution->Particles();
-
-	for (int32 Index = 0; Index < Assets.Num(); ++Index)
-	{
-		if (const UClothingAssetCommon* const Asset = Assets[Index])
-		{
-			const TVector<uint32, 2> Range = IndexToRangeMap[Index];
-
-			// Find bounds
-			TAABB<float, 3> BoundingBox = TAABB<float, 3>::EmptyAABB();
-			for (uint32 ParticleIndex = Range[0]; ParticleIndex < Range[1]; ++ParticleIndex)
-			{
-				BoundingBox.GrowToInclude(Particles.X(ParticleIndex));
-			}
-
-			// Find (squared) radius
-			const TVector<float, 3> Center = BoundingBox.Center();
-			float SquaredRadius = 0.f;
-			for (uint32 ParticleIndex = Range[0]; ParticleIndex < Range[1]; ++ParticleIndex)
-			{
-				SquaredRadius = FMath::Max(SquaredRadius, (Particles.X(ParticleIndex) - Center).SizeSquared());
-			}
-
-			// Update bounds with this cloth
-			const FBoxSphereBounds ClothBounds(BoundingBox.Center(), BoundingBox.Extents() * 0.5f, FMath::Sqrt(SquaredRadius));
-			Bounds = (NumBoundedCloths++ == 0) ? ClothBounds : Bounds + ClothBounds;
-		}
-	}
-
-	if (!bUseLocalSpaceSimulation && NumBoundedCloths && InOwnerComponent)
+	if (InOwnerComponent)
 	{
-		// Retrieve the master component (unlike the one passed to the context, this could be a slave component)
-		const bool bIsUsingMaster = InOwnerComponent->MasterPoseComponent.IsValid();
-		const USkinnedMeshComponent* const OwnerComponent = bIsUsingMaster ? InOwnerComponent->MasterPoseComponent.Get() : InOwnerComponent;
-
 		// Return local bounds
-		return Bounds.TransformBy(OwnerComponent->GetComponentTransform().Inverse());
+		return Bounds.TransformBy(InOwnerComponent->GetComponentTransform().Inverse());
 	}
 	return Bounds;
 }
 
-void ClothingSimulation::AddExternalCollisions(const FClothCollisionData& InData)
+void FClothingSimulation::AddExternalCollisions(const FClothCollisionData& InData)
 {
-	// Keep track of the external collisions data
-	ExternalCollisions.Append(InData);
-
-	// Add new map entry
-	const int32 MapIndex = ExternalCollisionsRangeMaps.AddDefaulted();
-	ExternalCollisionsRangeMaps[MapIndex].AddUninitialized(Assets.Num());
-
-	// Setup the new collisions particles for all cloths
-	UE_LOG(LogChaosCloth, VeryVerbose, TEXT("Adding external collisions..."));
-	static const TArray<int32> EmptyUsedBoneIndices;  // There is no bone mapping available for external collisions
-	for (int32 Index = 0; Index < Assets.Num(); ++Index)
-	{
-		ExternalCollisionsRangeMaps[MapIndex][Index][0] = Evolution->CollisionParticles().Size();
-		if (Assets[Index])
-		{
-			AddCollisions(InData, EmptyUsedBoneIndices, Index);
-		}
-		ExternalCollisionsRangeMaps[MapIndex][Index][1] = Evolution->CollisionParticles().Size();
-
-		// Keep collision transforms from previous frame if they exist
-		TGeometryClothParticles<float, 3>& CollisionParticles = Evolution->CollisionParticles();
-		for (uint32 i = ExternalCollisionsRangeMaps[MapIndex][Index][0];
-			i < FMath::Min(ExternalCollisionsRangeMaps[MapIndex][Index][1], (uint32)CollisionTransforms.Num());
-			++i)
-		{
-			CollisionParticles.X(i) = CollisionTransforms[i].GetLocation();
-			CollisionParticles.R(i) = CollisionTransforms[i].GetRotation();
-		}
-	}
+	ExternalCollisionData.Append(InData);
 }
 
-void ClothingSimulation::ClearExternalCollisions()
+void FClothingSimulation::ClearExternalCollisions()
 {
-	// Remove all external collision particles, starting from the external collision offset
-	// But do not resize CollisionTransforms as it is only resized in UpdateCollisionTransforms() to keep old transforms in between frames
-	TGeometryClothParticles<float, 3>& CollisionParticles = Evolution->CollisionParticles();
-	CollisionParticles.Resize(ExternalCollisionsOffset);  // This will also resize GroupIds, BoneIndices and BaseTransforms
-
-	// Reset external collisions
-	ExternalCollisions.Reset();
-
-	// Reset external collision maps
-	ExternalCollisionsRangeMaps.Reset();
-
-	UE_LOG(LogChaosCloth, VeryVerbose, TEXT("Cleared all external collisions."));
+	ExternalCollisionData.Reset();
 }
 
-void ClothingSimulation::GetCollisions(FClothCollisionData& OutCollisions, bool bIncludeExternal) const
+void FClothingSimulation::GetCollisions(FClothCollisionData& OutCollisions, bool bIncludeExternal) const
 {
 	// This code only gathers old apex collisions that don't appear in the physics mesh
-	// It is also never called with bIncludeExternal = true
+	// It is also never called with bIncludeExternal = true 
+	// but the collisions are then added untransformed and added as external
 	// This function is bound to be deprecated at some point
 
 	OutCollisions.Reset();
 
 	// Add internal asset collisions
-	for (const UClothingAssetCommon* Asset : Assets)
+	for (const TUniquePtr<FClothingSimulationCloth>& Cloth : Cloths)
 	{
-		if (Asset)
+		for (const FClothingSimulationCollider* const Collider : Cloth->GetColliders())
 		{
-			const FClothLODDataCommon& ClothLodData = Asset->LodData[0];
-			OutCollisions.Append(ClothLodData.CollisionData);
+			OutCollisions.Append(Collider->GetCollisionData());
 		}
 	}
 
 	// Add external asset collisions
 	if (bIncludeExternal)
 	{
-		OutCollisions.Append(ExternalCollisions);
+		OutCollisions.Append(ExternalCollisionData);
 	}
 
 	UE_LOG(LogChaosCloth, VeryVerbose, TEXT("GetCollisions returned collisions: %d spheres, %d capsules, %d convexes, %d boxes."), OutCollisions.Spheres.Num() - 2 * OutCollisions.SphereConnections.Num(), OutCollisions.SphereConnections.Num(), OutCollisions.Convexes.Num(), OutCollisions.Boxes.Num());
 }
 
-void ClothingSimulation::RefreshClothConfig()
+void FClothingSimulation::RefreshClothConfig(const IClothingSimulationContext* InContext)
 {
 	UpdateSimulationFromSharedSimConfig();
 
-	Evolution->ResetConstraintRules();
-	Evolution->ResetSelfCollision();
+	// Update new space location
+	const FClothingSimulationContext* const Context = static_cast<const FClothingSimulationContext*>(InContext);
+	Solver->SetLocalSpaceLocation(bUseLocalSpaceSimulation ? Context->ComponentToWorld.GetLocation() : TVector<float, 3>(0.f));
 
 	// Reset stats
 	ResetStats();
 
-	for (int32 SimDataIndex = 0; SimDataIndex < Assets.Num(); ++SimDataIndex)
-	{
-		if (const UClothingAssetCommon* const Asset = Assets[SimDataIndex])
-		{
-			if (const UChaosClothConfig* const ChaosClothConfig = Asset->GetClothConfig<UChaosClothConfig>())
-			{
-				check(Asset->GetNumLods() > 0);
-				const FClothPhysicalMeshData& PhysMesh = Asset->LodData[0].PhysicalMeshData;
-
-				ResetParticles(SimDataIndex);
-
-				SetParticleMasses(ChaosClothConfig, PhysMesh, SimDataIndex);
-
-				AddConstraints(ChaosClothConfig, PhysMesh, SimDataIndex);
-
-				LinearDeltaRatios[SimDataIndex] = FVector::OneVector - ChaosClothConfig->LinearVelocityScale.BoundToBox(FVector::ZeroVector, FVector::OneVector);
-				AngularDeltaRatios[SimDataIndex] = 1.f - FMath::Clamp(ChaosClothConfig->AngularVelocityScale, 0.f, 1.f);
-
-				// Set per cloth damping, collision thickness, and friction
-				Evolution->SetDamping(ChaosClothConfig->DampingCoefficient, SimDataIndex);
-				Evolution->SetCollisionThickness(ChaosClothConfig->CollisionThickness, SimDataIndex);
-				Evolution->SetCoefficientOfFriction(ChaosClothConfig->FrictionCoefficient, SimDataIndex);
-
-				// Add Velocity field
-				auto GetVelocity = [this, SimDataIndex](const TVector<float, 3>&)->TVector<float, 3>
-				{
-					return WindVelocities[SimDataIndex];
-				};
-				Evolution->GetVelocityField(SimDataIndex) = MakeUnique<TVelocityField<float, 3>>(
-					*Meshes[SimDataIndex],
-					GetVelocity,
-					/*bInIsUniform =*/ true,
-					ChaosClothConfig->DragCoefficient);
-
-				// Update gravity related config values
-				GravityScales[SimDataIndex] = ChaosClothConfig->GravityScale;
-				bUseGravityConfigs[SimDataIndex] = ChaosClothConfig->bUseGravityOverride;
-				GravityConfigs[SimDataIndex] = ChaosClothConfig->Gravity;
-
-				// Add Self Collisions
-				if (ChaosClothConfig->bUseSelfCollisions)
-				{
-					AddSelfCollisions(SimDataIndex);
-					Evolution->SetSelfCollisionThickness(ChaosClothConfig->SelfCollisionThickness, SimDataIndex);
-				}
-
-				// Update stats
-				UpdateStats(SimDataIndex);
-			}
-		}
-	}
-	UE_LOG(LogChaosCloth, VeryVerbose, TEXT("RefreshClothConfig, all constraints and self-collisions have been updated for all clothing assets"));
+	// Clear all cloths from the solver
+	Solver->RemoveCloths();
+
+	// Recreate all cloths
+	for (TUniquePtr<FClothingSimulationCloth>& Cloth : Cloths)
+	{
+		FClothingSimulationMesh* const Mesh = Cloth->GetMesh();
+		TArray<FClothingSimulationCollider*> ClothColliders = Cloth->GetColliders();
+		const uint32 GroupId = Cloth->GetGroupId();
+		const UChaosClothConfig* const ClothConfig = Mesh->GetAsset()->GetClothConfig<UChaosClothConfig>();
+
+		Cloth = MakeUnique<FClothingSimulationCloth>(
+			Mesh,
+			MoveTemp(ClothColliders),
+			GroupId,
+			(FClothingSimulationCloth::EMassMode)ClothConfig->MassMode,
+			ClothConfig->GetMassValue(),
+			ClothConfig->MinPerParticleMass,
+			ClothConfig->EdgeStiffness,
+			ClothConfig->BendingStiffness,
+			ClothConfig->bUseBendingElements,
+			ClothConfig->AreaStiffness,
+			ClothConfig->VolumeStiffness,
+			ClothConfig->bUseThinShellVolumeConstraints,
+			ClothConfig->StrainLimitingStiffness,
+			ClothConfig->LimitScale,
+			ClothConfig->bUseGeodesicDistance,
+			/*MaxDistancesMultiplier =*/ 1.f,  // Animatable
+			AnimDriveSpringStiffness,  // Animatable
+			ClothConfig->ShapeTargetStiffness,
+			/*bUseXPBDConstraints =*/ false,  // Experimental
+			ClothConfig->GravityScale,
+			ClothConfig->bUseGravityOverride,
+			ClothConfig->Gravity,
+			ClothConfig->LinearVelocityScale,
+			ClothConfig->AngularVelocityScale,
+			ClothConfig->DragCoefficient,
+			ClothConfig->DampingCoefficient,
+			ClothConfig->CollisionThickness,
+			ClothConfig->FrictionCoefficient,
+			ClothConfig->bUseSelfCollisions,
+			ClothConfig->SelfCollisionThickness,
+			/*bUseLODIndexOverride =*/ false,
+			/*LODIndexOverride =*/ INDEX_NONE);
+
+		// Re-add cloth to the solver
+		Solver->AddCloth(Cloth.Get());
+
+		// Update stats
+		UpdateStats(Cloth.Get());
+	}
+	UE_LOG(LogChaosCloth, VeryVerbose, TEXT("RefreshClothConfig, all constraints and self-collisions have been updated for all clothing assets and LODs."));
 }
 
-void ClothingSimulation::RefreshPhysicsAsset()
+void FClothingSimulation::RefreshPhysicsAsset()
 {
-	// Clear all collisions
-	TGeometryClothParticles<float, 3>& CollisionParticles = Evolution->CollisionParticles();
-	CollisionParticles.Resize(0);  // This will also resize BoneIndices and BaseTransforms
+	// Reset stats
+	ResetStats();
 
-	ExternalCollisions.Reset();
-	ExternalCollisionsOffset = 0;
+	// A collider update cannot be re-triggered for now, remove all cloths from the solver instead
+	Solver->RemoveCloths();
 
-	// Re-extract all collisions from every cloth asset
-	for (int32 Index = 0; Index < Assets.Num(); ++Index)
+	for (const TUniquePtr<FClothingSimulationCloth>& Cloth : Cloths)
 	{
-		if (const UClothingAssetCommon* const Asset = Assets[Index])
-		{
-			ExtractCollisions(Asset, Index);
-		}
+		// Backup collider(s)
+		TArray<FClothingSimulationCollider*> ClothColliders = Cloth->GetColliders();
+
+		// Refresh collider
+		Cloth->SetColliders(MoveTemp(ClothColliders));
+
+		// Re-add cloth to the solver
+		Solver->AddCloth(Cloth.Get());
+
+		// Update stats
+		UpdateStats(Cloth.Get());
 	}
 	UE_LOG(LogChaosCloth, VeryVerbose, TEXT("RefreshPhysicsAsset, all collisions have been re-added for all clothing assets"));
 }
 
-void ClothingSimulation::SetAnimDriveSpringStiffness(float InStiffness)
+void FClothingSimulation::SetAnimDriveSpringStiffness(float InAnimDriveSpringStiffness)
 {
-	for (float& stiffness : AnimDriveSpringStiffness)
-	{
-		stiffness = InStiffness;
-	}
+	AnimDriveSpringStiffness = InAnimDriveSpringStiffness;
 }
 
-void ClothingSimulation::SetGravityOverride(const FVector& InGravityOverride)
+void FClothingSimulation::SetGravityOverride(const FVector& InGravityOverride)
 {
 	bUseGravityOverride = true;
 	GravityOverride = InGravityOverride;
 }
 
-void ClothingSimulation::DisableGravityOverride()
+void FClothingSimulation::DisableGravityOverride()
 {
 	bUseGravityOverride = false;
 }
 
 #if WITH_EDITOR
-void ClothingSimulation::AddReferencedObjects(FReferenceCollector& Collector)
+void FClothingSimulation::AddReferencedObjects(FReferenceCollector& Collector)
 {
 	Collector.AddReferencedObject(DebugClothMaterial);
 }
 
-void ClothingSimulation::DebugDrawPhysMeshShaded(USkeletalMeshComponent* OwnerComponent, FPrimitiveDrawInterface* PDI) const
+void FClothingSimulation::DebugDrawPhysMeshShaded(FPrimitiveDrawInterface* PDI) const
 {
-	if (!DebugClothMaterial) { return; }
+	if (!DebugClothMaterial)
+	{
+		return;
+	}
 
 	FDynamicMeshBuilder MeshBuilder(PDI->View->GetFeatureLevel());
-	const TPBDParticles<float, 3>& Particles = Evolution->Particles();
-
 	int32 VertexIndex = 0;
 
-	for (int32 Index = 0; Index < Assets.Num(); ++Index)
+	for (const FClothingSimulationCloth* const Cloth : Solver->GetCloths())
 	{
-		const UClothingAssetCommon* const Asset = Assets[Index];
-		if (!Asset) { continue; }
+		const int32 Offset = Cloth->GetOffset(Solver.Get());
+		if (Offset == INDEX_NONE)
+		{
+			continue;
+		}
 
-		const FClothLODDataCommon& AssetLodData = Asset->LodData[0];
-		const FClothPhysicalMeshData& PhysMesh = AssetLodData.PhysicalMeshData;
-		const FPointWeightMap& MaxDistances = PhysMesh.GetWeightMap(EChaosWeightMapTarget::MaxDistance);
-		if (MaxDistances.Num() == 0) { continue; }
-		const uint32 Offset = IndexToRangeMap[Index][0];
+		const TConstArrayView<TVector<int32, 3>> Elements = Cloth->GetTriangleMesh(Solver.Get()).GetElements();
+		const TConstArrayView<TVector<float, 3>> Positions = Cloth->GetParticlePositions(Solver.Get());
+		const TConstArrayView<float> InvMasses = Cloth->GetParticleInvMasses(Solver.Get());
+		check(InvMasses.Num() == Positions.Num());
 
-		if (const TUniquePtr<Chaos::TTriangleMesh<float>>& Mesh = Meshes[Index])
+		for (int32 ElementIndex = 0; ElementIndex < Elements.Num(); ++ElementIndex, VertexIndex += 3)
 		{
-			const TArray<TVector<int32, 3>>& Elements = Mesh->GetElements();
-
-			for (int32 ElementIndex = 0; ElementIndex < Elements.Num(); ++ElementIndex, VertexIndex += 3)
-			{
-				const auto& Element = Elements[ElementIndex];
+			const TVector<int32, 3>& Element = Elements[ElementIndex];
 
-				const FVector& Pos0 = Particles.X(Element.X);
-				const FVector& Pos1 = Particles.X(Element.Y);
-				const FVector& Pos2 = Particles.X(Element.Z);
+			const FVector Pos0 = Positions[Element.X - Offset]; // TODO: Triangle Mesh shouldn't really be solver dependent (ie not use an offset)
+			const FVector Pos1 = Positions[Element.Y - Offset];
+			const FVector Pos2 = Positions[Element.Z - Offset];
 
-				const FVector& Normal = FVector::CrossProduct(Pos2 - Pos0, Pos1 - Pos0).GetSafeNormal();
-				const FVector Tangent = ((Pos1 + Pos2) * 0.5f - Pos0).GetSafeNormal();
+			const FVector Normal = FVector::CrossProduct(Pos2 - Pos0, Pos1 - Pos0).GetSafeNormal();
+			const FVector Tangent = ((Pos1 + Pos2) * 0.5f - Pos0).GetSafeNormal();
 
-				static const float KinematicDistanceThreshold = 0.1f;  // Must be the same value set in UI, not ideal really
-				const bool bIsKinematic0 = MaxDistances[Element.X - Offset] < KinematicDistanceThreshold;
-				const bool bIsKinematic1 = MaxDistances[Element.Y - Offset] < KinematicDistanceThreshold;
-				const bool bIsKinematic2 = MaxDistances[Element.Z - Offset] < KinematicDistanceThreshold;
+			const bool bIsKinematic0 = (InvMasses[Element.X - Offset] == 0.f);
+			const bool bIsKinematic1 = (InvMasses[Element.Y - Offset] == 0.f);
+			const bool bIsKinematic2 = (InvMasses[Element.Z - Offset] == 0.f);
 
-				MeshBuilder.AddVertex(FDynamicMeshVertex(Pos0, Tangent, Normal, FVector2D(0.f, 0.f), bIsKinematic0 ? FColor::Purple : FColor::White));
-				MeshBuilder.AddVertex(FDynamicMeshVertex(Pos1, Tangent, Normal, FVector2D(0.f, 1.f), bIsKinematic1 ? FColor::Purple : FColor::White));
-				MeshBuilder.AddVertex(FDynamicMeshVertex(Pos2, Tangent, Normal, FVector2D(1.f, 1.f), bIsKinematic2 ? FColor::Purple : FColor::White));
-				MeshBuilder.AddTriangle(VertexIndex, VertexIndex + 1, VertexIndex + 2);
-			}
+			MeshBuilder.AddVertex(FDynamicMeshVertex(Pos0, Tangent, Normal, FVector2D(0.f, 0.f), bIsKinematic0 ? FColor::Purple : FColor::White));
+			MeshBuilder.AddVertex(FDynamicMeshVertex(Pos1, Tangent, Normal, FVector2D(0.f, 1.f), bIsKinematic1 ? FColor::Purple : FColor::White));
+			MeshBuilder.AddVertex(FDynamicMeshVertex(Pos2, Tangent, Normal, FVector2D(1.f, 1.f), bIsKinematic2 ? FColor::Purple : FColor::White));
+			MeshBuilder.AddTriangle(VertexIndex, VertexIndex + 1, VertexIndex + 2);
 		}
 	}
 
 	FMatrix LocalSimSpaceToWorld(FMatrix::Identity);
-	LocalSimSpaceToWorld.SetOrigin(LocalSpaceLocation);
+	LocalSimSpaceToWorld.SetOrigin(Solver->GetLocalSpaceLocation());
 	MeshBuilder.Draw(PDI, LocalSimSpaceToWorld, DebugClothMaterial->GetRenderProxy(), SDPG_World, false, false);
 }
 
@@ -1945,32 +572,37 @@ static void DrawText(FCanvas* Canvas, const FSceneView* SceneView, const FVector
 	}
 }
 
-void ClothingSimulation::DebugDrawParticleIndices(USkeletalMeshComponent* /*OwnerComponent*/, FCanvas* Canvas, const FSceneView* SceneView) const
+void FClothingSimulation::DebugDrawParticleIndices(FCanvas* Canvas, const FSceneView* SceneView) const
 {
-	const TPBDParticles<float, 3>& Particles = Evolution->Particles();
-
 	static const FLinearColor DynamicColor = FColor::White;
 	static const FLinearColor KinematicColor = FColor::Purple;
 
-	for (int32 Index = 0; Index < Assets.Num(); ++Index)
+	const TVector<float, 3>& LocalSpaceLocation = Solver->GetLocalSpaceLocation();
+
+	for (const FClothingSimulationCloth* const Cloth : Solver->GetCloths())
 	{
-		const UClothingAssetCommon* const Asset = Assets[Index];
-		if (!Asset) { continue; }
+		const int32 Offset = Cloth->GetOffset(Solver.Get());
+		if (Offset == INDEX_NONE)
+		{
+			continue;
+		}
+
+		const TConstArrayView<TVector<float, 3>> Positions = Cloth->GetParticlePositions(Solver.Get());
+		const TConstArrayView<float> InvMasses = Cloth->GetParticleInvMasses(Solver.Get());
+		check(InvMasses.Num() == Positions.Num());
 
-		for (uint32 ParticleIndex = IndexToRangeMap[Index][0]; ParticleIndex < IndexToRangeMap[Index][1]; ++ParticleIndex)
+		for (int32 Index = 0; Index < Positions.Num(); ++Index)
 		{
-			const FVector Position = Particles.X(ParticleIndex) + LocalSpaceLocation;
+			const FVector Position = LocalSpaceLocation + Positions[Index];
 
-			const FText Text = FText::AsNumber(ParticleIndex);
-			DrawText(Canvas, SceneView, Position, Text, Particles.InvM(ParticleIndex) == 0.0f ? KinematicColor : DynamicColor);
+			const FText Text = FText::AsNumber(Offset + Index);
+			DrawText(Canvas, SceneView, Position, Text, InvMasses[Index] == 0.f ? KinematicColor : DynamicColor);
 		}
 	}
 }
 
-void ClothingSimulation::DebugDrawMaxDistanceValues(USkeletalMeshComponent* /*OwnerComponent*/, FCanvas* Canvas, const FSceneView* SceneView) const
+void FClothingSimulation::DebugDrawMaxDistanceValues(FCanvas* Canvas, const FSceneView* SceneView) const
 {
-	const TPBDParticles<float, 3>& Particles = Evolution->Particles();
-
 	static const FLinearColor DynamicColor = FColor::White;
 	static const FLinearColor KinematicColor = FColor::Purple;
 
@@ -1983,28 +615,29 @@ void ClothingSimulation::DebugDrawMaxDistanceValues(USkeletalMeshComponent* /*Ow
 	NumberFormattingOptions.MinimumFractionalDigits = 2;
 	NumberFormattingOptions.MaximumFractionalDigits = 2;
 
-	for (int32 Index = 0; Index < Assets.Num(); ++Index)
+	const TVector<float, 3>& LocalSpaceLocation = Solver->GetLocalSpaceLocation();
+
+	for (const FClothingSimulationCloth* const Cloth : Solver->GetCloths())
 	{
-		const UClothingAssetCommon* const Asset = Assets[Index];
-		if (!Asset) { continue; }
-
-		// Get Maximum Distances
-		const FClothLODDataCommon& AssetLodData = Asset->LodData[0];
-		const FClothPhysicalMeshData& PhysMesh = AssetLodData.PhysicalMeshData;
-		const FPointWeightMap& MaxDistances = PhysMesh.GetWeightMap(EChaosWeightMapTarget::MaxDistance);
-		if (MaxDistances.Num() == 0)
+		const int32 Offset = Cloth->GetOffset(Solver.Get());
+		if (Offset == INDEX_NONE)
 		{
 			continue;
 		}
-		
-		for (uint32 ParticleIndex = IndexToRangeMap[Index][0]; ParticleIndex < IndexToRangeMap[Index][1]; ++ParticleIndex)
+
+		const TConstArrayView<float>& MaxDistances = Cloth->GetWeightMaps(Solver.Get())[(int32)EChaosWeightMapTarget::MaxDistance];
+		const TConstArrayView<TVector<float, 3>> Positions = Cloth->GetAnimationPositions(Solver.Get());
+		const TConstArrayView<float> InvMasses = Cloth->GetParticleInvMasses(Solver.Get());
+		check(MaxDistances.Num() == Positions.Num());
+		check(MaxDistances.Num() == InvMasses.Num());
+
+		for (int32 Index = 0; Index < MaxDistances.Num(); ++Index)
 		{
-			const uint32 WeightMapIndex = ParticleIndex - IndexToRangeMap[Index][0];
-			const float Distance = MaxDistances[WeightMapIndex];
-			const FVector Position = Particles.X(ParticleIndex) + LocalSpaceLocation;
+			const float MaxDistance = MaxDistances[Index];
+			const FVector Position = LocalSpaceLocation + Positions[Index];
 
-			const FText Text = FText::AsNumber(Distance, &NumberFormattingOptions);
-			DrawText(Canvas, SceneView, Position, Text, Particles.InvM(ParticleIndex) == 0.0f ? KinematicColor : DynamicColor);
+			const FText Text = FText::AsNumber(MaxDistance, &NumberFormattingOptions);
+			DrawText(Canvas, SceneView, Position, Text, InvMasses[Index] == 0.f ? KinematicColor : DynamicColor);
 		}
 	}
 }
@@ -2190,310 +823,408 @@ static void DrawCoordinateSystem(FPrimitiveDrawInterface* PDI, const FQuat& Rota
 }
 
 #if CHAOS_DEBUG_DRAW
-void ClothingSimulation::DebugDrawBounds() const
+void FClothingSimulation::DebugDrawBounds() const
 {
+	check(Solver);
+
 	// Calculate World space bounds
-	const FBoxSphereBounds Bounds = GetBounds(nullptr);
+	const FBoxSphereBounds Bounds = Solver->CalculateBounds();
 
 	// Draw bounds
-	DrawBox(nullptr, TBox<float, 3>(-Bounds.BoxExtent, Bounds.BoxExtent), FQuat::Identity, LocalSpaceLocation + Bounds.Origin, FLinearColor(FColor::Purple));
-	DrawSphere(nullptr, TSphere<float, 3>(FVector::ZeroVector, Bounds.SphereRadius), FQuat::Identity, LocalSpaceLocation + Bounds.Origin, FLinearColor(FColor::Orange));
+	DrawBox(nullptr, TBox<float, 3>(-Bounds.BoxExtent, Bounds.BoxExtent), FQuat::Identity, Bounds.Origin, FLinearColor(FColor::Purple));
+	DrawSphere(nullptr, TSphere<float, 3>(FVector::ZeroVector, Bounds.SphereRadius), FQuat::Identity, Bounds.Origin, FLinearColor(FColor::Orange));
+
+	// Draw individual cloth bounds
+	static const FLinearColor Color = FLinearColor(FColor::Purple).Desaturate(0.5);
+	for (const TUniquePtr<FClothingSimulationCloth>& Cloth : Cloths)
+	{
+		if (Cloth->GetOffset(Solver.Get()) == INDEX_NONE)
+		{
+			continue;
+		}
+
+		const TAABB<float, 3> BoundingBox = Cloth->CalculateBoundingBox(Solver.Get());
+		DrawBox(nullptr, TBox<float, 3>(BoundingBox), FQuat::Identity, Bounds.Origin, Color);
+	}
 }
 
-void ClothingSimulation::DebugDrawGravity() const
+void FClothingSimulation::DebugDrawGravity() const
 {
-	// Calculate World space bounds
-	const FBoxSphereBounds Bounds = GetBounds(nullptr);
+	check(Solver);
 
 	// Draw gravity
-	const FVector Pos0 = LocalSpaceLocation + Bounds.Origin;
-	const FVector Pos1 = Pos0 + Evolution->GetGravityForces().GetAcceleration();
-	DrawLine(nullptr, Pos0, Pos1, FLinearColor::Red);
+	for (const FClothingSimulationCloth* const Cloth : Solver->GetCloths())
+	{
+		if (Cloth->GetOffset(Solver.Get()) == INDEX_NONE)
+		{
+			continue;
+		}
+
+		const TAABB<float, 3> Bounds = Cloth->CalculateBoundingBox(Solver.Get());
+
+		const FVector Pos0 = Bounds.Center();
+		const FVector Pos1 = Pos0 + Cloth->GetGravity(Solver.Get());
+		DrawLine(nullptr, Pos0, Pos1, FLinearColor::Red);
+	}
 }
 #endif  // #if CHAOS_DEBUG_DRAW
 
-void ClothingSimulation::DebugDrawPhysMeshWired(USkeletalMeshComponent* /*OwnerComponent*/, FPrimitiveDrawInterface* PDI) const
+void FClothingSimulation::DebugDrawPhysMeshWired(FPrimitiveDrawInterface* PDI) const
 {
 	static const FLinearColor DynamicColor = FColor::White;
 	static const FLinearColor KinematicColor = FColor::Purple;
 
-	const TPBDParticles<float, 3>& Particles = Evolution->Particles();
+	const TVector<float, 3>& LocalSpaceLocation = Solver->GetLocalSpaceLocation();
 
-	for (int32 Index = 0; Index < Assets.Num(); ++Index)
+	for (const FClothingSimulationCloth* const Cloth : Solver->GetCloths())
 	{
-		if (Assets[Index] && Meshes[Index])
+		const int32 Offset = Cloth->GetOffset(Solver.Get());
+		if (Offset == INDEX_NONE)
 		{
-			const TArray<TVector<int32, 3>>& Elements = Meshes[Index]->GetElements();
+			continue;
+		}
 
-			for (int32 ElementIndex = 0; ElementIndex < Elements.Num(); ++ElementIndex)
-			{
-				const auto& Element = Elements[ElementIndex];
+		const TConstArrayView<TVector<int32, 3>> Elements = Cloth->GetTriangleMesh(Solver.Get()).GetElements();
+		const TConstArrayView<TVector<float, 3>> Positions = Cloth->GetParticlePositions(Solver.Get());
+		const TConstArrayView<float> InvMasses = Cloth->GetParticleInvMasses(Solver.Get());
+		check(InvMasses.Num() == Positions.Num());
 
-				const FVector Pos0 = LocalSpaceLocation + Particles.X(Element.X);
-				const FVector Pos1 = LocalSpaceLocation + Particles.X(Element.Y);
-				const FVector Pos2 = LocalSpaceLocation + Particles.X(Element.Z);
+		for (int32 ElementIndex = 0; ElementIndex < Elements.Num(); ++ElementIndex)
+		{
+			const auto& Element = Elements[ElementIndex];
 
-				const bool bIsKinematic0 = Particles.InvM(Element.X) == 0.f;
-				const bool bIsKinematic1 = Particles.InvM(Element.Y) == 0.f;
-				const bool bIsKinematic2 = Particles.InvM(Element.Z) == 0.f;
+			const FVector Pos0 = LocalSpaceLocation + Positions[Element.X - Offset]; // TODO: Triangle Mesh shouldn't really be solver dependent (ie not use an offset)
+			const FVector Pos1 = LocalSpaceLocation + Positions[Element.Y - Offset];
+			const FVector Pos2 = LocalSpaceLocation + Positions[Element.Z - Offset];
 
-				DrawLine(PDI, Pos0, Pos1, bIsKinematic0 && bIsKinematic1 ? KinematicColor : DynamicColor);
-				DrawLine(PDI, Pos1, Pos2, bIsKinematic1 && bIsKinematic2 ? KinematicColor : DynamicColor);
-				DrawLine(PDI, Pos2, Pos0, bIsKinematic2 && bIsKinematic0 ? KinematicColor : DynamicColor);
-			}
+			const bool bIsKinematic0 = (InvMasses[Element.X - Offset] == 0.f);
+			const bool bIsKinematic1 = (InvMasses[Element.Y - Offset] == 0.f);
+			const bool bIsKinematic2 = (InvMasses[Element.Z - Offset] == 0.f);
+
+			DrawLine(PDI, Pos0, Pos1, bIsKinematic0 && bIsKinematic1 ? KinematicColor : DynamicColor);
+			DrawLine(PDI, Pos1, Pos2, bIsKinematic1 && bIsKinematic2 ? KinematicColor : DynamicColor);
+			DrawLine(PDI, Pos2, Pos0, bIsKinematic2 && bIsKinematic0 ? KinematicColor : DynamicColor);
 		}
 	}
 }
 
-void ClothingSimulation::DebugDrawPointNormals(USkeletalMeshComponent* OwnerComponent, FPrimitiveDrawInterface* PDI) const
+void FClothingSimulation::DebugDrawAnimMeshWired(FPrimitiveDrawInterface* PDI) const
 {
-	const TPBDParticles<float, 3>& Particles = Evolution->Particles();
+	static const FLinearColor DynamicColor = FColor::White;
+	static const FLinearColor KinematicColor = FColor::Purple;
+
+	const TVector<float, 3>& LocalSpaceLocation = Solver->GetLocalSpaceLocation();
 
-	for (int32 Index = 0; Index < Assets.Num(); ++Index)
+	for (const FClothingSimulationCloth* const Cloth : Solver->GetCloths())
 	{
-		if (Assets[Index])
+		const int32 Offset = Cloth->GetOffset(Solver.Get());
+		if (Offset == INDEX_NONE)
 		{
-			const TVector<uint32, 2> Range = IndexToRangeMap[Index];
-			const TArray<TVector<float, 3>>& MeshPointNormals = PointNormals[Index];
-
-			for (uint32 ParticleIndex = Range[0]; ParticleIndex < Range[1]; ++ParticleIndex)
-			{
-				const FVector Pos0 = LocalSpaceLocation + Particles.X(ParticleIndex);
-				const FVector Pos1 = Pos0 + MeshPointNormals[ParticleIndex - Range[0]] * 20.0f;
-
-				DrawLine(PDI, Pos0, Pos1, FLinearColor::White);
-			}
+			continue;
 		}
-	}
-}
 
-void ClothingSimulation::DebugDrawInversedPointNormals(USkeletalMeshComponent* OwnerComponent, FPrimitiveDrawInterface* PDI) const
-{
-	const TPBDParticles<float, 3>& Particles = Evolution->Particles();
+		const TConstArrayView<TVector<int32, 3>> Elements = Cloth->GetTriangleMesh(Solver.Get()).GetElements();
+		const TConstArrayView<TVector<float, 3>> Positions = Cloth->GetAnimationPositions(Solver.Get());
+		const TConstArrayView<float> InvMasses = Cloth->GetParticleInvMasses(Solver.Get());
+		check(InvMasses.Num() == Positions.Num());
 
-	for (int32 Index = 0; Index < Assets.Num(); ++Index)
-	{
-		if (Assets[Index])
+		for (int32 ElementIndex = 0; ElementIndex < Elements.Num(); ++ElementIndex)
 		{
-			const TVector<uint32, 2> Range = IndexToRangeMap[Index];
-			const TArray<TVector<float, 3>>& MeshPointNormals = PointNormals[Index];
+			const auto& Element = Elements[ElementIndex];
 
-			for (uint32 ParticleIndex = Range[0]; ParticleIndex < Range[1]; ++ParticleIndex)
-			{
-				const FVector Pos0 = LocalSpaceLocation + Particles.X(ParticleIndex);
-				const FVector Pos1 = Pos0 - MeshPointNormals[ParticleIndex - Range[0]] * 20.0f;
+			const FVector Pos0 = LocalSpaceLocation + Positions[Element.X - Offset]; // TODO: Triangle Mesh shouldn't really be solver dependent (ie not use an offset)
+			const FVector Pos1 = LocalSpaceLocation + Positions[Element.Y - Offset];
+			const FVector Pos2 = LocalSpaceLocation + Positions[Element.Z - Offset];
 
-				DrawLine(PDI, Pos0, Pos1, FLinearColor::White);
-			}
+			const bool bIsKinematic0 = (InvMasses[Element.X - Offset] == 0.f);
+			const bool bIsKinematic1 = (InvMasses[Element.Y - Offset] == 0.f);
+			const bool bIsKinematic2 = (InvMasses[Element.Z - Offset] == 0.f);
+
+			DrawLine(PDI, Pos0, Pos1, bIsKinematic0 && bIsKinematic1 ? KinematicColor : DynamicColor);
+			DrawLine(PDI, Pos1, Pos2, bIsKinematic1 && bIsKinematic2 ? KinematicColor : DynamicColor);
+			DrawLine(PDI, Pos2, Pos0, bIsKinematic2 && bIsKinematic0 ? KinematicColor : DynamicColor);
 		}
 	}
 }
 
-void ClothingSimulation::DebugDrawFaceNormals(USkeletalMeshComponent* OwnerComponent, FPrimitiveDrawInterface* PDI) const
+void FClothingSimulation::DebugDrawPointNormals(FPrimitiveDrawInterface* PDI) const
 {
-	const TPBDParticles<float, 3>& Particles = Evolution->Particles();
+	check(Solver);
+	const TVector<float, 3>& LocalSpaceLocation = Solver->GetLocalSpaceLocation();
 
-	for (int32 Index = 0; Index < Assets.Num(); ++Index)
+	for (const FClothingSimulationCloth* const Cloth : Solver->GetCloths())
 	{
-		if (Assets[Index])
+		const int32 Offset = Cloth->GetOffset(Solver.Get());
+		if (Offset == INDEX_NONE)
 		{
-			const TArray<TVector<float, 3>>& MeshFaceNormals = FaceNormals[Index];
+			continue;
+		}
 
-			const TArray<TVector<int32, 3>>& Elements = Meshes[Index]->GetElements();
-			for (int32 ElementIndex = 0; ElementIndex < Elements.Num(); ++ElementIndex)
-			{
-				const TVector<int32, 3>& Element = Elements[ElementIndex];
+		const TConstArrayView<TVector<float, 3>> Positions = Cloth->GetParticlePositions(Solver.Get());
+		const TConstArrayView<TVector<float, 3>> Normals = Cloth->GetParticleNormals(Solver.Get());
+		check(Normals.Num() == Positions.Num());
 
-				const FVector Pos0 = LocalSpaceLocation + (
-					Particles.X(Element.X) +
-					Particles.X(Element.Y) +
-					Particles.X(Element.Z)) / 3.f;
-				const FVector Pos1 = Pos0 + MeshFaceNormals[ElementIndex] * 20.0f;
+		for (int32 Index = 0; Index < Positions.Num(); ++Index)
+		{
+			const FVector Pos0 = LocalSpaceLocation + Positions[Index];
+			const FVector Pos1 = Pos0 + Normals[Index] * 20.f;
 
-				DrawLine(PDI, Pos0, Pos1, FLinearColor::Yellow);
-			}
+			DrawLine(PDI, Pos0, Pos1, FLinearColor::White);
 		}
 	}
 }
 
-void ClothingSimulation::DebugDrawInversedFaceNormals(USkeletalMeshComponent* OwnerComponent, FPrimitiveDrawInterface* PDI) const
+void FClothingSimulation::DebugDrawInversedPointNormals(FPrimitiveDrawInterface* PDI) const
 {
-	const TPBDParticles<float, 3>& Particles = Evolution->Particles();
+	check(Solver);
+	const TVector<float, 3>& LocalSpaceLocation = Solver->GetLocalSpaceLocation();
 
-	for (int32 Index = 0; Index < Assets.Num(); ++Index)
+	for (const FClothingSimulationCloth* const Cloth : Solver->GetCloths())
 	{
-		if (Assets[Index])
+		const int32 Offset = Cloth->GetOffset(Solver.Get());
+		if (Offset == INDEX_NONE)
 		{
-			const TArray<TVector<float, 3>>& MeshFaceNormals = FaceNormals[Index];
+			continue;
+		}
 
-			const TArray<TVector<int32, 3>>& Elements = Meshes[Index]->GetElements();
-			for (int32 ElementIndex = 0; ElementIndex < Elements.Num(); ++ElementIndex)
-			{
-				const TVector<int32, 3>& Element = Elements[ElementIndex];
+		const TConstArrayView<TVector<float, 3>> Positions = Cloth->GetParticlePositions(Solver.Get());
+		const TConstArrayView<TVector<float, 3>> Normals = Cloth->GetParticleNormals(Solver.Get());
 
-				const FVector Pos0 = LocalSpaceLocation + (
-					Particles.X(Element.X) +
-					Particles.X(Element.Y) +
-					Particles.X(Element.Z)) / 3.f;
-				const FVector Pos1 = Pos0 - MeshFaceNormals[ElementIndex] * 20.0f;
+		for (int32 Index = 0; Index < Positions.Num(); ++Index)
+		{
+			const FVector Pos0 = LocalSpaceLocation + Positions[Index];
+			const FVector Pos1 = Pos0 - Normals[Index] * 20.f;
 
-				DrawLine(PDI, Pos0, Pos1, FLinearColor::Yellow);
-			}
+			DrawLine(PDI, Pos0, Pos1, FLinearColor::White);
 		}
 	}
 }
 
-void ClothingSimulation::DebugDrawCollision(USkeletalMeshComponent* /*OwnerComponent*/, FPrimitiveDrawInterface* PDI) const
+void FClothingSimulation::DebugDrawCollision(FPrimitiveDrawInterface* PDI) const
 {
-	auto DrawCollision = [this, PDI](const TVector<uint32, 2>& Ranges)
-	{
-		static const FLinearColor MappedColor(FColor::Cyan);
-		static const FLinearColor UnmappedColor(FColor::Red);
-
-		const TGeometryClothParticles<float, 3>& CollisionParticles = Evolution->CollisionParticles();
-		for (uint32 Index = Ranges[0]; Index < Ranges[1]; ++Index)
+	check(Solver);
+	auto DrawCollision =
+		[this, PDI](const FClothingSimulationCollider* Collider, const FClothingSimulationCloth* Cloth, FClothingSimulationCollider::ECollisionDataType CollisionDataType)
 		{
-			if (const FImplicitObject* const Object = CollisionParticles.DynamicGeometry(Index).Get())
-			{
-				const uint32 BoneIndex = BoneIndices[Index];
-				const FLinearColor Color = (BoneIndex != INDEX_NONE) ? MappedColor : UnmappedColor;
+			static const FLinearColor GlobalColor(FColor::Cyan);
+			static const FLinearColor DynamicColor(FColor::Red);
+			static const FLinearColor LODsColor(FColor::Silver);
+
+			const FLinearColor Color =
+				(CollisionDataType == FClothingSimulationCollider::ECollisionDataType::Global) ? GlobalColor :
+				(CollisionDataType == FClothingSimulationCollider::ECollisionDataType::Dynamic) ? DynamicColor : LODsColor;
 
-				const TVector<float, 3> Position = LocalSpaceLocation + CollisionParticles.X(Index);
-				const TRotation<float, 3>& Rotation = CollisionParticles.R(Index);
+			const TVector<float, 3>& LocalSpaceLocation = Solver->GetLocalSpaceLocation();
 
-				switch (Object->GetType())
+			const TConstArrayView<TUniquePtr<FImplicitObject>> CollisionGeometries = Collider->GetCollisionGeometries(Solver.Get(), Cloth, CollisionDataType);
+			const TConstArrayView<TVector<float, 3>> Translations = Collider->GetCollisionTranslations(Solver.Get(), Cloth, CollisionDataType);
+			const TConstArrayView<TRotation<float, 3>> Rotations = Collider->GetCollisionRotations(Solver.Get(), Cloth, CollisionDataType);
+			check(CollisionGeometries.Num() == Translations.Num());
+			check(CollisionGeometries.Num() == Rotations.Num());
+
+			for (int32 Index = 0; Index < CollisionGeometries.Num(); ++Index)
+			{
+				if (const FImplicitObject* const Object = CollisionGeometries[Index].Get())
 				{
-				case ImplicitObjectType::Sphere:
-					DrawSphere(PDI, Object->GetObjectChecked<TSphere<float, 3>>(), Rotation, Position, Color);
-					break;
+					switch (Object->GetType())
+					{
+					case ImplicitObjectType::Sphere:
+						DrawSphere(PDI, Object->GetObjectChecked<TSphere<float, 3>>(), Rotations[Index], Translations[Index], Color);
+						break;
 
-				case ImplicitObjectType::Box:
-					DrawBox(PDI, Object->GetObjectChecked<TBox<float, 3>>(), Rotation, Position, Color);
-					break;
+					case ImplicitObjectType::Box:
+						DrawBox(PDI, Object->GetObjectChecked<TBox<float, 3>>(), Rotations[Index], Translations[Index], Color);
+						break;
 
-				case ImplicitObjectType::Capsule:
-					DrawCapsule(PDI, Object->GetObjectChecked<TCapsule<float>>(), Rotation, Position, Color);
-					break;
+					case ImplicitObjectType::Capsule:
+						DrawCapsule(PDI, Object->GetObjectChecked<TCapsule<float>>(), Rotations[Index], Translations[Index], Color);
+						break;
 
-				case ImplicitObjectType::Union:  // Union only used as collision tapered capsules
-					for (const TUniquePtr<FImplicitObject>& SubObjectPtr : Object->GetObjectChecked<FImplicitObjectUnion>().GetObjects())
-					{
-						if (const FImplicitObject* const SubObject = SubObjectPtr.Get())
+					case ImplicitObjectType::Union:  // Union only used as collision tapered capsules
+						for (const TUniquePtr<FImplicitObject>& SubObjectPtr : Object->GetObjectChecked<FImplicitObjectUnion>().GetObjects())
 						{
-							switch (SubObject->GetType())
+							if (const FImplicitObject* const SubObject = SubObjectPtr.Get())
 							{
-							case ImplicitObjectType::Sphere:
-								DrawSphere(PDI, SubObject->GetObjectChecked<TSphere<float, 3>>(), Rotation, Position, Color);
-								break;
-
-							case ImplicitObjectType::TaperedCylinder:
-								DrawTaperedCylinder(PDI, SubObject->GetObjectChecked<TTaperedCylinder<float>>(), Rotation, Position, Color);
-								break;
-
-							default:
-								break;
+								switch (SubObject->GetType())
+								{
+								case ImplicitObjectType::Sphere:
+									DrawSphere(PDI, SubObject->GetObjectChecked<TSphere<float, 3>>(), Rotations[Index], Translations[Index], Color);
+									break;
+
+								case ImplicitObjectType::TaperedCylinder:
+									DrawTaperedCylinder(PDI, SubObject->GetObjectChecked<TTaperedCylinder<float>>(), Rotations[Index], Translations[Index], Color);
+									break;
+
+								default:
+									break;
+								}
 							}
 						}
-					}
-					break;
+						break;
 
-				case ImplicitObjectType::Convex:
-					DrawConvex(PDI, Object->GetObjectChecked<FConvex>(), Rotation, Position, Color);
-					break;
+					case ImplicitObjectType::Convex:
+						DrawConvex(PDI, Object->GetObjectChecked<FConvex>(), Rotations[Index], Translations[Index], Color);
+						break;
 
-				default:
-					DrawCoordinateSystem(PDI, Rotation, Position);  // Draw everything else as a coordinate for now
-					break;
+					default:
+						DrawCoordinateSystem(PDI, Rotations[Index], Translations[Index]);  // Draw everything else as a coordinate for now
+						break;
+					}
 				}
 			}
+		};
+
+	// Draw collisions
+	for (const FClothingSimulationCloth* const Cloth : Solver->GetCloths())
+	{
+		for (const FClothingSimulationCollider* const Collider : Cloth->GetColliders())
+		{
+			DrawCollision(Collider, Cloth, FClothingSimulationCollider::ECollisionDataType::Global);
+			DrawCollision(Collider, Cloth, FClothingSimulationCollider::ECollisionDataType::LODs);
+		}
+	}
+	// Draw external collisions
+	for (const FClothingSimulationCollider* const Collider : Solver->GetColliders())
+	{
+		DrawCollision(Collider, nullptr, FClothingSimulationCollider::ECollisionDataType::Dynamic);
+	}
+}
+
+void FClothingSimulation::DebugDrawBackstops(FPrimitiveDrawInterface* PDI) const
+{
+	auto DrawBackstop = [PDI](const FVector& Position, const FVector& Normal, float Radius, const FVector& Axis, const FLinearColor& Color)
+	{
+		static const float ArcAngle = 25.0f; // Arc angle in degrees
+		static const float MaxCosAngle = 0.99f;
+		if (FMath::Abs(FVector::DotProduct(Normal, Axis)) < MaxCosAngle)
+		{
+			DrawArc(PDI, Position, Normal, FVector::CrossProduct(Axis, Normal).GetSafeNormal(), -ArcAngle / 2.0f, ArcAngle / 2.0f, Radius, Color);
 		}
 	};
 
-	for (int32 Index = 0; Index < Assets.Num(); ++Index)
+	const TVector<float, 3>& LocalSpaceLocation = Solver->GetLocalSpaceLocation();
+
+	uint8 ColorSeed = 0;
+
+	for (const FClothingSimulationCloth* const Cloth : Solver->GetCloths())
 	{
-		if (Assets[Index])
+		const int32 Offset = Cloth->GetOffset(Solver.Get());
+		if (Offset == INDEX_NONE)
+		{
+			continue;
+		}
+
+		const FClothConstraints& ClothConstraints = Solver->GetClothConstraints(Offset);
+		if (const TPBDSphericalBackstopConstraint<float, 3>* const BackstopConstraint = ClothConstraints.GetBackstopConstraints().Get())
 		{
-			// Draw collisions
-			DrawCollision(CollisionsRangeMap[Index]);
+			const TConstArrayView<float>& BackstopDistances = Cloth->GetWeightMaps(Solver.Get())[(int32)EChaosWeightMapTarget::BackstopDistance];
+			const TConstArrayView<float>& BackstopRadiuses = Cloth->GetWeightMaps(Solver.Get())[(int32)EChaosWeightMapTarget::BackstopRadius];
+			const TConstArrayView<TVector<float, 3>> AnimationPositions = Cloth->GetAnimationPositions(Solver.Get());
+			const TConstArrayView<TVector<float, 3>> AnimationNormals = Cloth->GetAnimationNormals(Solver.Get());
+			const TConstArrayView<TVector<float, 3>> ParticlePositions = Cloth->GetParticlePositions(Solver.Get());
 
-			// Draw external collisions
-			for (const TArray<TVector<uint32, 2>>& ExternalCollisionsRangeMap : ExternalCollisionsRangeMaps)
+			for (int32 Index = 0; Index < AnimationPositions.Num(); ++Index)
 			{
-				DrawCollision(ExternalCollisionsRangeMap[Index]);
+				ColorSeed += 157;  // Prime number that gives a good spread of colors without getting too similar as a rand might do.
+				const FLinearColor ColorLight = FLinearColor::MakeFromHSV8(ColorSeed, 160, 128);
+				const FLinearColor ColorDark = FLinearColor::MakeFromHSV8(ColorSeed, 160, 64);
+
+				const float BackstopRadius = BackstopRadiuses[Index] * BackstopConstraint->GetSphereRadiiMultiplier();
+				const float BackstopDistance = BackstopDistances[Index];
+
+				const FVector AnimationPosition = LocalSpaceLocation + AnimationPositions[Index];
+				const FVector& AnimationNormal = AnimationNormals[Index];
+
+				// Draw a line to the current position
+				const FVector Pos0 = LocalSpaceLocation + AnimationPositions[Index];
+				const FVector Pos1 = Pos0 - BackstopDistance * AnimationNormal;
+				const FVector Pos2 = LocalSpaceLocation + ParticlePositions[Index];
+				DrawLine(PDI, Pos1, Pos2, ColorLight);
+
+				// Draw the sphere
+				if (BackstopRadius > 0.f)
+				{
+					const FVector Center = Pos0 - (BackstopRadius + BackstopDistance) * AnimationNormal;
+					DrawBackstop(Center, AnimationNormal, BackstopRadius, FVector::ForwardVector, ColorDark);
+					DrawBackstop(Center, AnimationNormal, BackstopRadius, FVector::UpVector, ColorDark);
+					DrawBackstop(Center, AnimationNormal, BackstopRadius, FVector::RightVector, ColorDark);
+				}
 			}
 		}
 	}
 }
 
-void ClothingSimulation::DebugDrawBackstops(USkeletalMeshComponent* /*OwnerComponent*/, FPrimitiveDrawInterface* PDI) const
+void FClothingSimulation::DebugDrawBackstopDistances(FPrimitiveDrawInterface* PDI) const
 {
-	for (int32 Index = 0; Index < Assets.Num(); ++Index)
+	const TVector<float, 3>& LocalSpaceLocation = Solver->GetLocalSpaceLocation();
+
+	uint8 ColorSeed = 0;
+
+	for (const FClothingSimulationCloth* const Cloth : Solver->GetCloths())
 	{
-		const UClothingAssetCommon* const Asset = Assets[Index];
-		if (!Asset) { continue; }
-
-		// Get Backstop Distances
-		const FClothLODDataCommon& AssetLodData = Asset->LodData[0];
-		const FClothPhysicalMeshData& PhysMesh = AssetLodData.PhysicalMeshData;
-		const FPointWeightMap& BackstopDistances = PhysMesh.GetWeightMap(EChaosWeightMapTarget::BackstopDistance);
-		const FPointWeightMap& BackstopRadiuses = PhysMesh.GetWeightMap(EChaosWeightMapTarget::BackstopRadius);
-		if (BackstopDistances.Num() == 0 || BackstopRadiuses.Num() == 0)
+		const int32 Offset = Cloth->GetOffset(Solver.Get());
+		if (Offset == INDEX_NONE)
 		{
 			continue;
 		}
 
-		for (uint32 ParticleIndex = IndexToRangeMap[Index][0]; ParticleIndex < IndexToRangeMap[Index][1]; ++ParticleIndex)
+		const FClothConstraints& ClothConstraints = Solver->GetClothConstraints(Offset);
+		if (const TPBDSphericalBackstopConstraint<float, 3>* const BackstopConstraint = ClothConstraints.GetBackstopConstraints().Get())
 		{
-			const uint32 WeightMapIndex = ParticleIndex - IndexToRangeMap[Index][0];
-			const float Radius = BackstopRadiuses[WeightMapIndex];
-			const float Distance = BackstopDistances[WeightMapIndex];
-			const FVector Position = LocalSpaceLocation + AnimationPositions[ParticleIndex];
-			const FVector& Normal = AnimationNormals[ParticleIndex];
-			DrawLine(PDI, Position, Position - Normal * (Distance - Radius), FLinearColor::White);
-			if (Radius > 0.0f)
+			const TConstArrayView<float>& BackstopDistances = Cloth->GetWeightMaps(Solver.Get())[(int32)EChaosWeightMapTarget::BackstopDistance];
+			const TConstArrayView<float>& BackstopRadiuses = Cloth->GetWeightMaps(Solver.Get())[(int32)EChaosWeightMapTarget::BackstopRadius];
+			const TConstArrayView<TVector<float, 3>> AnimationPositions = Cloth->GetAnimationPositions(Solver.Get());
+			const TConstArrayView<TVector<float, 3>> AnimationNormals = Cloth->GetAnimationNormals(Solver.Get());
+
+			for (int32 Index = 0; Index < AnimationPositions.Num(); ++Index)
 			{
-				auto DrawBackstop = [Radius, Distance, &Normal, &Position, PDI](const FVector& Axis, const FLinearColor& Color)
-				{
-					const float ArcLength = 5.0f; // Arch length in cm
-					const float ArcAngle = ArcLength * 360.0f / (Radius * 2.0f * PI);
-					
-					const float MaxCosAngle = 0.99f;
-					if (FMath::Abs(FVector::DotProduct(Normal, Axis)) < MaxCosAngle)
-					{
-						DrawArc(PDI, Position - Normal * Distance, Normal, FVector::CrossProduct(Axis, Normal).GetSafeNormal(), -ArcAngle / 2.0f, ArcAngle / 2.0f, Radius, Color);
-					}
-				};
-				DrawBackstop(FVector::ForwardVector, FLinearColor::Blue);
-				DrawBackstop(FVector::UpVector, FLinearColor::Blue);
-				DrawBackstop(FVector::RightVector, FLinearColor::Blue);
+				ColorSeed += 157;  // Prime number that gives a good spread of colors without getting too similar as a rand might do.
+				const FLinearColor ColorLight = FLinearColor::MakeFromHSV8(ColorSeed, 160, 128);
+				const FLinearColor ColorDark = FLinearColor::MakeFromHSV8(ColorSeed, 160, 64);
+
+				const float BackstopRadius = BackstopRadiuses[Index] * BackstopConstraint->GetSphereRadiiMultiplier();
+				const float BackstopDistance = BackstopDistances[Index];
+
+				const FVector AnimationPosition = LocalSpaceLocation + AnimationPositions[Index];
+				const FVector& AnimationNormal = AnimationNormals[Index];
+
+				// Draw a line to the sphere boundary
+				const FVector Pos0 = LocalSpaceLocation + AnimationPositions[Index];
+				const FVector Pos1 = Pos0 - BackstopDistance * AnimationNormal;
+				DrawLine(PDI, Pos0, Pos1, ColorDark);
 			}
 		}
 	}
 }
 
-void ClothingSimulation::DebugDrawMaxDistances(USkeletalMeshComponent* /*OwnerComponent*/, FPrimitiveDrawInterface* PDI) const
+void FClothingSimulation::DebugDrawMaxDistances(FPrimitiveDrawInterface* PDI) const
 {
-	const TPBDParticles<float, 3>& Particles = Evolution->Particles();
-	for (int32 Index = 0; Index < Assets.Num(); ++Index)
+	check(Solver);
+	
+	// Draw max distances
+	const TVector<float, 3>& LocalSpaceLocation = Solver->GetLocalSpaceLocation();
+
+	for (const FClothingSimulationCloth* const Cloth : Solver->GetCloths())
 	{
-		const UClothingAssetCommon* const Asset = Assets[Index];
-		if (!Asset) { continue; }
-
-		// Get Maximum Distances
-		const FClothLODDataCommon& AssetLodData = Asset->LodData[0];
-		const FClothPhysicalMeshData& PhysMesh = AssetLodData.PhysicalMeshData;
-		const FPointWeightMap& MaxDistances = PhysMesh.GetWeightMap(EChaosWeightMapTarget::MaxDistance);
-		if (MaxDistances.Num() == 0)
+		const int32 Offset = Cloth->GetOffset(Solver.Get());
+		if (Offset == INDEX_NONE)
 		{
 			continue;
 		}
-		
-		for (uint32 ParticleIndex = IndexToRangeMap[Index][0]; ParticleIndex < IndexToRangeMap[Index][1]; ++ParticleIndex)
+
+		const TConstArrayView<float>& MaxDistances = Cloth->GetWeightMaps(Solver.Get())[(int32)EChaosWeightMapTarget::MaxDistance];
+		const TConstArrayView<float> InvMasses = Cloth->GetParticleInvMasses(Solver.Get());
+		const TConstArrayView<TVector<float, 3>> Positions = Cloth->GetAnimationPositions(Solver.Get());
+		const TConstArrayView<TVector<float, 3>> Normals = Cloth->GetAnimationNormals(Solver.Get());
+		check(Normals.Num() == Positions.Num());
+		check(MaxDistances.Num() == Positions.Num());
+		check(InvMasses.Num() == Positions.Num());
+
+		for (int32 Index = 0; Index < MaxDistances.Num(); ++Index)
 		{
-			const uint32 WeightMapIndex = ParticleIndex - IndexToRangeMap[Index][0];
-			const float Distance = MaxDistances[WeightMapIndex];
-			const FVector Position = LocalSpaceLocation + AnimationPositions[ParticleIndex];
-			if (Particles.InvM(ParticleIndex) == 0.0f)
+			const float MaxDistance = MaxDistances[Index];
+			const FVector Position = LocalSpaceLocation + Positions[Index];
+			if (InvMasses[Index] == 0.f)
 			{
 #if WITH_EDITOR
 				DrawPoint(PDI, Position, FLinearColor::Red, DebugClothMaterialVertex);
@@ -2501,46 +1232,58 @@ void ClothingSimulation::DebugDrawMaxDistances(USkeletalMeshComponent* /*OwnerCo
 			}
 			else
 			{
-				DrawLine(PDI, Position, Position + AnimationNormals[ParticleIndex] * Distance, FLinearColor::White);
+				DrawLine(PDI, Position, Position + Normals[Index] * MaxDistance, FLinearColor::White);
 			}
 		}
 	}
 }
 
-void ClothingSimulation::DebugDrawAnimDrive(USkeletalMeshComponent* /*OwnerComponent*/, FPrimitiveDrawInterface* PDI) const
+void FClothingSimulation::DebugDrawAnimDrive(FPrimitiveDrawInterface* PDI) const
 {
-	const TPBDParticles<float, 3>& Particles = Evolution->Particles();
-	for (int32 Index = 0; Index < Assets.Num(); ++Index)
+	check(Solver);
+
+	const TVector<float, 3>& LocalSpaceLocation = Solver->GetLocalSpaceLocation();
+
+	for (const FClothingSimulationCloth* const Cloth : Solver->GetCloths())
 	{
-		const UClothingAssetCommon* const Asset = Assets[Index];
-		if (!Asset) { continue; }
-
-		// Get anim drive multiplier
-		const FClothLODDataCommon& AssetLodData = Asset->LodData[0];
-		const FClothPhysicalMeshData& PhysMesh = AssetLodData.PhysicalMeshData;
-		const FPointWeightMap& AnimDriveMultipliers = PhysMesh.GetWeightMap(EChaosWeightMapTarget::AnimDriveMultiplier);
-		if (AnimDriveMultipliers.Num() == 0)
+		const int32 Offset = Cloth->GetOffset(Solver.Get());
+		if (Offset == INDEX_NONE)
 		{
 			continue;
 		}
 
-		for (uint32 ParticleIndex = IndexToRangeMap[Index][0]; ParticleIndex < IndexToRangeMap[Index][1]; ++ParticleIndex)
+		const FClothConstraints& ClothConstraints = Solver->GetClothConstraints(Offset);
+		if (const TPBDAnimDriveConstraint<float, 3>* const AnimDriveConstraint = ClothConstraints.GetAnimDriveConstraints().Get())
 		{
-			const uint32 WeightMapIndex = ParticleIndex - IndexToRangeMap[Index][0];
-			const float Multiplier = AnimDriveMultipliers[WeightMapIndex];
-			DrawLine(PDI, AnimationPositions[ParticleIndex] + LocalSpaceLocation, Particles.X(ParticleIndex) + LocalSpaceLocation, FLinearColor(FColor::Cyan) * Multiplier * AnimDriveSpringStiffness[Index]);
+			const float SpringStiffness = AnimDriveConstraint->GetSpringStiffness();
+
+			const TConstArrayView<float>& AnimDriveMultipliers = Cloth->GetWeightMaps(Solver.Get())[(int32)EChaosWeightMapTarget::AnimDriveMultiplier];
+			const TConstArrayView<TVector<float, 3>> AnimationPositions = Cloth->GetAnimationPositions(Solver.Get());
+			const TConstArrayView<TVector<float, 3>> ParticlePositions = Cloth->GetParticlePositions(Solver.Get());
+			check(AnimDriveMultipliers.Num() == AnimationPositions.Num());
+			check(AnimDriveMultipliers.Num() == ParticlePositions.Num());
+
+			for (int32 Index = 0; Index < AnimDriveMultipliers.Num(); ++Index)
+			{
+				const float AnimDriveMultiplier = AnimDriveMultipliers[Index];
+				const FVector AnimationPosition = LocalSpaceLocation + AnimationPositions[Index];
+				const FVector ParticlePosition = LocalSpaceLocation + ParticlePositions[Index];
+				DrawLine(PDI, AnimationPosition, ParticlePosition, FLinearColor(FColor::Cyan) * AnimDriveMultiplier * SpringStiffness);
+			}
 		}
 	}
 }
 
-void ClothingSimulation::DebugDrawLongRangeConstraint(USkeletalMeshComponent* /*OwnerComponent*/, FPrimitiveDrawInterface* PDI) const
+void FClothingSimulation::DebugDrawLongRangeConstraint(FPrimitiveDrawInterface* PDI) const
 {
-	auto PseudoRandomColor = 
-		[](int32 NumRotations) -> FLinearColor
+	const TVector<float, 3>& LocalSpaceLocation = Solver->GetLocalSpaceLocation();
+
+	auto PseudoRandomColor =
+		[](int32 NumColorRotations) -> FLinearColor
 		{
 			static const uint8 Spread = 157;  // Prime number that gives a good spread of colors without getting too similar as a rand might do.
 			uint8 Seed = Spread;
-			for (int32 i = 0; i < NumRotations; ++i)
+			for (int32 i = 0; i < NumColorRotations; ++i)
 			{
 				Seed += Spread;
 			}
@@ -2549,64 +1292,75 @@ void ClothingSimulation::DebugDrawLongRangeConstraint(USkeletalMeshComponent* /*
 
 	int32 ColorOffset = 0;
 
-	const TPBDParticles<float, 3>& Particles = Evolution->Particles();
-
-	for (int32 Index = 0; Index < Assets.Num(); ++Index)
+	for (const FClothingSimulationCloth* const Cloth : Solver->GetCloths())
 	{
-		if (!Assets[Index] || !Meshes[Index] || !LongRangeConstraints[Index]) { continue; }
+		const int32 Offset = Cloth->GetOffset(Solver.Get());
+		if (Offset == INDEX_NONE)
+		{
+			continue;
+		}
 
 		// Recompute islands
-		const TMap<int32, TSet<uint32>>& PointToNeighborsMap = Meshes[Index]->GetPointToNeighborsMap();
+		const TTriangleMesh<float>& TriangleMesh = Cloth->GetTriangleMesh(Solver.Get());
+		const TConstArrayView<float> InvMasses = Cloth->GetParticleInvMasses(Solver.Get());
+
+		const TMap<int32, TSet<uint32>>& PointToNeighborsMap = TriangleMesh.GetPointToNeighborsMap();
 
 		static TArray<uint32> KinematicIndices;  // Make static to prevent constant allocations
 		KinematicIndices.Reset();
 		for (const TPair<int32, TSet<uint32>>& PointNeighbors : PointToNeighborsMap)
 		{
-			const int32 PointIndex = PointNeighbors.Key;
-			if (Particles.InvM(PointIndex) == 0.f)
+			const int32 Index = PointNeighbors.Key;
+			if (InvMasses[Index - Offset] == 0.f)  // TODO: Triangle indices should ideally be starting at 0 to avoid these mix-ups
 			{
-				KinematicIndices.Add(PointIndex);
+				KinematicIndices.Add(Index);
 			}
 		}
 
 		const TArray<TArray<uint32>> IslandElements = TPBDLongRangeConstraints<float, 3>::ComputeIslands(PointToNeighborsMap, KinematicIndices);
 
 		// Draw constraints
-		const TArray<TArray<uint32>>& Constraints = LongRangeConstraints[Index]->GetConstraints();
+		const FClothConstraints& ClothConstraints = Solver->GetClothConstraints(Offset);
+		
+		const TConstArrayView<TVector<float, 3>> Positions = Cloth->GetParticlePositions(Solver.Get());
 
-		for (int32 ConstraintIndex = 0; ConstraintIndex < Constraints.Num(); ++ConstraintIndex)
+		if (const TPBDLongRangeConstraints<float, 3>* const LongRangeConstraints = ClothConstraints.GetLongRangeConstraints().Get())
 		{
-			const TArray<uint32>& Path = Constraints[ConstraintIndex];
-			const uint32 KinematicIndex = Path[0];
-			const uint32 DynamicIndex = Path[Path.Num() - 1];
+			const TArray<TArray<uint32>>& Constraints = LongRangeConstraints->GetConstraints();
 
-			// Find Island
-			int32 ColorIndex = 0;
-			for (int32 IslandIndex = 0; IslandIndex < IslandElements.Num(); ++IslandIndex)
+			for (int32 ConstraintIndex = 0; ConstraintIndex < Constraints.Num(); ++ConstraintIndex)
 			{
-				if (IslandElements[IslandIndex].Find(KinematicIndex) != INDEX_NONE)  // TODO: This is O(n^2), it'll be nice to make this faster, even if it is only debug viz. Maybe binary search if the kinematic indices are ordered?
+				const TArray<uint32>& Path = Constraints[ConstraintIndex];
+				const uint32 KinematicIndex = Path[0];
+				const uint32 DynamicIndex = Path[Path.Num() - 1];
+
+				// Find Island
+				int32 ColorIndex = 0;
+				for (int32 IslandIndex = 0; IslandIndex < IslandElements.Num(); ++IslandIndex)
 				{
-					ColorIndex = ColorOffset + IslandIndex;
-					break;
+					if (IslandElements[IslandIndex].Find(KinematicIndex) != INDEX_NONE)  // TODO: This is O(n^2), it'll be nice to make this faster, even if it is only debug viz. Maybe binary search if the kinematic indices are ordered?
+					{
+						ColorIndex = ColorOffset + IslandIndex;
+						break;
+					}
 				}
-			}
 
-			const TVector<float, 3> Pos0 = Particles.X(KinematicIndex) + LocalSpaceLocation;
-			const TVector<float, 3> Pos1 = Particles.X(DynamicIndex) + LocalSpaceLocation;
-
-			DrawLine(PDI, Pos0, Pos1, PseudoRandomColor(ColorIndex));
+				const TVector<float, 3> Pos0 = Positions[KinematicIndex - Offset] + LocalSpaceLocation;
+				const TVector<float, 3> Pos1 = Positions[DynamicIndex - Offset] + LocalSpaceLocation;
+				DrawLine(PDI, Pos0, Pos1, PseudoRandomColor(ColorIndex));
+			}
 		}
 
 		// Draw islands
-		const TArray<TVector<int32, 3>>& Elements = Meshes[Index]->GetElements();
+		const TConstArrayView<TVector<int32, 3>> Elements = Cloth->GetTriangleMesh(Solver.Get()).GetElements();
 
 		for (int32 ElementIndex = 0; ElementIndex < Elements.Num(); ++ElementIndex)
 		{
 			const auto& Element = Elements[ElementIndex];
 
-			const bool bIsKinematic0 = (Particles.InvM(Element.X) == 0.f);
-			const bool bIsKinematic1 = (Particles.InvM(Element.Y) == 0.f);
-			const bool bIsKinematic2 = (Particles.InvM(Element.Z) == 0.f);
+			const bool bIsKinematic0 = (InvMasses[Element.X - Offset] == 0.f);
+			const bool bIsKinematic1 = (InvMasses[Element.Y - Offset] == 0.f);
+			const bool bIsKinematic2 = (InvMasses[Element.Z - Offset] == 0.f);
 
 			// Lookup for any kinematic point on the triangle element to use for finding the island (it doesn't matter which one, if two kinematic points are on the same triangle they have to be on the same island)
 			const int32 KinematicIndex = bIsKinematic0 ? Element.X : bIsKinematic1 ? Element.Y : bIsKinematic2 ? Element.Z : INDEX_NONE;
@@ -2615,7 +1369,7 @@ void ClothingSimulation::DebugDrawLongRangeConstraint(USkeletalMeshComponent* /*
 				continue;
 			}
 
-			// Find Island Color
+			// Find island Color
 			int32 ColorIndex = 0;
 			for (int32 IslandIndex = 0; IslandIndex < IslandElements.Num(); ++IslandIndex)
 			{
@@ -2627,9 +1381,9 @@ void ClothingSimulation::DebugDrawLongRangeConstraint(USkeletalMeshComponent* /*
 			}
 			const FLinearColor Color = PseudoRandomColor(ColorIndex);
 
-			const FVector Pos0 = LocalSpaceLocation + Particles.X(Element.X);
-			const FVector Pos1 = LocalSpaceLocation + Particles.X(Element.Y);
-			const FVector Pos2 = LocalSpaceLocation + Particles.X(Element.Z);
+			const FVector Pos0 = LocalSpaceLocation + Positions[Element.X - Offset];
+			const FVector Pos1 = LocalSpaceLocation + Positions[Element.Y - Offset];
+			const FVector Pos2 = LocalSpaceLocation + Positions[Element.Z - Offset];
 
 			if (bIsKinematic0 && bIsKinematic1)
 			{
@@ -2645,50 +1399,55 @@ void ClothingSimulation::DebugDrawLongRangeConstraint(USkeletalMeshComponent* /*
 			}
 		}
 
-		// Rotate the colors for each cloths
+		// Rotate the colors for each cloth
 		ColorOffset += IslandElements.Num();
 	}
 }
 
-void ClothingSimulation::DebugDrawWindDragForces(USkeletalMeshComponent* /*OwnerComponent*/, FPrimitiveDrawInterface* PDI) const
+void FClothingSimulation::DebugDrawWindDragForces(FPrimitiveDrawInterface* PDI) const
 {
-	const TPBDParticles<float, 3>& Particles = Evolution->Particles();
-
-	for (int32 Index = 0; Index < Assets.Num(); ++Index)
-	{
-		const UClothingAssetCommon* const Asset = Assets[Index];
-		if (!Asset) { continue; }
-
-		const TUniquePtr<TVelocityField<float, 3>>& VelocityField = Evolution->GetVelocityField(Index);
-
-		const TArray<TVector<int32, 3>>& Elements = VelocityField->GetElements();
-		const TArray<TVector<float, 3>>& Forces = VelocityField->GetForces();
-
-		for (int32 ElementIndex = 0; ElementIndex < Elements.Num(); ++ElementIndex)
-		{
-			const TVector<int32, 3>& Element = Elements[ElementIndex];
-			const TVector<float, 3> Position = LocalSpaceLocation + (
-				Particles.X(Element[0]) +
-				Particles.X(Element[1]) +
-				Particles.X(Element[2])) / 3.f;
-			const TVector<float, 3>& Force = Forces[ElementIndex];
-			DrawLine(PDI, Position, Position + Force, FColor::Green);
-		}
-	}
+	// TODO: Add lift and re-enable debug draw code
+	//const TPBDParticles<float, 3>& Particles = Evolution->Particles();
+
+	//for (int32 Index = 0; Index < Assets.Num(); ++Index)
+	//{
+	//	const UClothingAssetCommon* const Asset = Assets[Index];
+	//	if (!Asset) { continue; }
+
+	//	const FVelocityField& VelocityField = Evolution->GetVelocityField(Index);
+
+	//	const TArrayView<TVector<int32, 3>>& Elements = VelocityField.GetElements();
+	//	const TArray<TVector<float, 3>>& Forces = VelocityField.GetForces();
+
+	//	for (int32 ElementIndex = 0; ElementIndex < Elements.Num(); ++ElementIndex)
+	//	{
+	//		const TVector<int32, 3>& Element = Elements[ElementIndex];
+	//		const TVector<float, 3> Position = LocalSpaceLocation + (
+	//			Particles.X(Element[0]) +
+	//			Particles.X(Element[1]) +
+	//			Particles.X(Element[2])) / 3.f;
+	//		const TVector<float, 3>& Force = Forces[ElementIndex];
+	//		DrawLine(PDI, Position, Position + Force, FColor::Green);
+	//	}
+	//}
 }
 
-void ClothingSimulation::DebugDrawLocalSpace(USkeletalMeshComponent* /*OwnerComponent*/, FPrimitiveDrawInterface* PDI) const
+void FClothingSimulation::DebugDrawLocalSpace(FPrimitiveDrawInterface* PDI) const
 {
+	check(Solver);
+
 	// Draw local space
-	DrawCoordinateSystem(PDI, FQuat::Identity, LocalSpaceLocation);
+	DrawCoordinateSystem(PDI, FQuat::Identity, Solver->GetLocalSpaceLocation());
 
 	// Draw reference spaces
-	for (int32 Index = 0; Index < Assets.Num(); ++Index)
+	for (const FClothingSimulationCloth* const Cloth : Solver->GetCloths())
 	{
-		if (Assets[Index])
+		if (Cloth->GetOffset(Solver.Get()) == INDEX_NONE)
 		{
-			DrawCoordinateSystem(PDI, RootBoneWorldTransforms[Index].GetRotation(), RootBoneWorldTransforms[Index].GetLocation());
+			continue;
 		}
+		const TRigidTransform<float, 3>& ReferenceSpaceTransform = Cloth->GetReferenceSpaceTransform();
+		DrawCoordinateSystem(PDI, ReferenceSpaceTransform.GetRotation(), ReferenceSpaceTransform.GetLocation());
 	}
 }
 #endif  // #if WITH_EDITOR || CHAOS_DEBUG_DRAW
