Chaos Cloth - LOD/graph node refactor.

diff --git a/Engine/Plugins/Experimental/ChaosCloth/Source/Chaos/Private/ChaosCloth/ChaosClothingSimulationCloth.cpp b/Engine/Plugins/Experimental/ChaosCloth/Source/Chaos/Private/ChaosCloth/ChaosClothingSimulationCloth.cpp
new file mode 100644
--- /dev/null
+++ b/Engine/Plugins/Experimental/ChaosCloth/Source/Chaos/Private/ChaosCloth/ChaosClothingSimulationCloth.cpp
@@ -0,0 +1,719 @@
+// Copyright Epic Games, Inc. All Rights Reserved.
+#include "ChaosCloth/ChaosClothingSimulationCloth.h"
+#include "ChaosCloth/ChaosClothingSimulationSolver.h"
+#include "ChaosCloth/ChaosClothingSimulationMesh.h"
+#include "ChaosCloth/ChaosClothingSimulationCollider.h"
+#include "ChaosCloth/ChaosClothConstraints.h"
+#include "ChaosCloth/ChaosWeightMapTarget.h"
+#include "ChaosCloth/ChaosClothPrivate.h"
+#include "Containers/ArrayView.h"
+
+using namespace Chaos;
+
+FClothingSimulationCloth::FLODData::FLODData(int32 InNumParticles, const TConstArrayView<uint32>& InIndices, const TArray<TConstArrayView<float>>& InWeightMaps)
+	: NumParticles(InNumParticles)
+	, Indices(InIndices)
+	, WeightMaps(InWeightMaps)
+{
+}
+
+void FClothingSimulationCloth::FLODData::Add(FClothingSimulationSolver* Solver, FClothingSimulationCloth* Cloth, int32 InLODIndex)
+{
+	check(Solver);
+	check(Cloth);
+
+	if (!NumParticles)
+	{
+		return;
+	}
+
+	// Add a new solver data chunk
+	check(!SolverData.Find(Solver));
+	FSolverData& SolverDatum = SolverData.Add(Solver);
+	int32& Offset = SolverDatum.Offset;
+	TTriangleMesh<float>& TriangleMesh = SolverDatum.TriangleMesh;
+
+	// Add particles
+	Offset = Solver->AddParticles(NumParticles, Cloth->GroupId);  // TODO: Have a per solver map of offset
+
+	// Update source mesh for this LOD, this is required prior to reset the start pose
+	Cloth->Mesh->Update(Solver, INDEX_NONE, InLODIndex, 0, Offset);
+
+	// Reset the particles start pose before setting up mass and constraints
+	ResetStartPose(Solver);
+
+	// Build a sim friendly triangle mesh including the solver particle's Offset
+	const int32 NumElements = Indices.Num() / 3;
+	TArray<TVector<int32, 3>> Elements;
+	Elements.Reserve(NumElements);
+
+	for (int32 i = 0; i < NumElements; ++i)
+	{
+		const int32 Index = 3 * i;
+		Elements.Add(
+			{ static_cast<int32>(Offset + Indices[Index]),
+			 static_cast<int32>(Offset + Indices[Index + 1]),
+			 static_cast<int32>(Offset + Indices[Index + 2]) });
+	}
+
+	TriangleMesh.Init(MoveTemp(Elements));
+	TriangleMesh.GetPointToTriangleMap(); // Builds map for later use by GetPointNormals(), and the velocity fields
+
+	// Set the particle masses
+	const TConstArrayView<float>& MaxDistances = WeightMaps[(int32)EChaosWeightMapTarget::MaxDistance];
+	static const float KinematicDistanceThreshold = 0.1f;  // TODO: This is not the same value as set in the painting UI but we might want to expose this value as parameter
+	auto KinematicPredicate =
+		[MaxDistances](int32 Index)
+		{
+			return MaxDistances.IsValidIndex(Index) && MaxDistances[Index] < KinematicDistanceThreshold;
+		};
+
+	switch (Cloth->MassMode)
+	{
+	default:
+		check(false);
+	case EMassMode::UniformMass:
+		Solver->SetParticleMassUniform(Offset, Cloth->MassValue, Cloth->MinPerParticleMass, TriangleMesh, KinematicPredicate);
+		break;
+	case EMassMode::TotalMass:
+		Solver->SetParticleMassFromTotalMass(Offset, Cloth->MassValue, Cloth->MinPerParticleMass, TriangleMesh, KinematicPredicate);
+		break;
+	case EMassMode::Density:
+		Solver->SetParticleMassFromDensity(Offset, Cloth->MassValue, Cloth->MinPerParticleMass, TriangleMesh, KinematicPredicate);
+		break;
+	}
+	const TConstArrayView<float> InvMasses(Solver->GetParticleInvMasses(Offset), NumParticles);
+
+	// Self collisions
+	Solver->SetSelfCollisions(Cloth->GroupId, Cloth->SelfCollisionThickness, Cloth->bUseSelfCollisions ? &TriangleMesh : nullptr);
+
+	// Setup solver constraints
+	FClothConstraints& ClothConstraints = Solver->GetClothConstraints(Offset);
+	const TArray<TVector<int32, 3>>& SurfaceElements = TriangleMesh.GetSurfaceElements();
+
+	// Edge constraints
+	if (Cloth->EdgeStiffness)
+	{
+		TArray<TVector<int32, 3>> DynamicSurfaceElements;
+		TArray<TVector<int32, 2>> Attachments;
+		for (const TVector<int32, 3>& SurfaceElement : SurfaceElements)
+		{
+			const bool bIsKinematic0 = (InvMasses[SurfaceElement[0] - Offset] == 0.f);
+			const bool bIsKinematic1 = (InvMasses[SurfaceElement[1] - Offset] == 0.f);
+			const bool bIsKinematic2 = (InvMasses[SurfaceElement[2] - Offset] == 0.f);
+			bool bIsAttachment = false;
+			if (bIsKinematic0 != bIsKinematic1)
+			{
+				Attachments.Emplace(SurfaceElement[0], SurfaceElement[1]);
+				bIsAttachment = true;
+			}
+			if (bIsKinematic1 != bIsKinematic2)
+			{
+				Attachments.Emplace(SurfaceElement[1], SurfaceElement[2]);
+				bIsAttachment = true;
+			}
+			if (bIsKinematic2 != bIsKinematic0)
+			{
+				Attachments.Emplace(SurfaceElement[2], SurfaceElement[0]);
+				bIsAttachment = true;
+			}
+			if (!bIsAttachment)
+			{
+				DynamicSurfaceElements.Add(SurfaceElement);
+			}
+		}
+		if (Attachments.Num())
+		{
+			ClothConstraints.SetEdgeConstraints(MoveTemp(Attachments), Cloth->EdgeStiffness, Cloth->bUseXPBDConstraints);
+		}
+		if (DynamicSurfaceElements.Num())
+		{
+			ClothConstraints.SetEdgeConstraints(DynamicSurfaceElements, Cloth->EdgeStiffness, Cloth->bUseXPBDConstraints);
+		}
+	}
+
+	// Bending constraints
+	if (Cloth->BendingStiffness > 0.f)
+	{
+		if (Cloth->bUseBendingElements)
+		{
+			TArray<Chaos::TVector<int32, 4>> BendingElements = TriangleMesh.GetUniqueAdjacentElements();
+			ClothConstraints.SetBendingConstraints(MoveTemp(BendingElements), Cloth->BendingStiffness);
+		}
+		else
+		{
+			TArray<Chaos::TVector<int32, 2>> Edges = TriangleMesh.GetUniqueAdjacentPoints();
+			ClothConstraints.SetBendingConstraints(MoveTemp(Edges), Cloth->BendingStiffness, Cloth->bUseXPBDConstraints);
+		}
+	}
+
+	// Area constraints
+	if (Cloth->AreaStiffness)
+	{
+		TArray<Chaos::TVector<int32, 3>> SurfaceConstraints = SurfaceElements;
+		ClothConstraints.SetAreaConstraints(MoveTemp(SurfaceConstraints), Cloth->AreaStiffness, Cloth->bUseXPBDConstraints);
+	}
+
+	// Volume constraints
+	if (Cloth->VolumeStiffness)
+	{
+		if (Cloth->bUseThinShellVolumeConstraints)
+		{
+			TArray<Chaos::TVector<int32, 2>> BendingConstraints = TriangleMesh.GetUniqueAdjacentPoints();
+			TArray<Chaos::TVector<int32, 2>> DoubleBendingConstraints;
+			{
+				TMap<int32, TArray<int32>> BendingHash;
+				for (int32 i = 0; i < BendingConstraints.Num(); ++i)
+				{
+					BendingHash.FindOrAdd(BendingConstraints[i][0]).Add(BendingConstraints[i][1]);
+					BendingHash.FindOrAdd(BendingConstraints[i][1]).Add(BendingConstraints[i][0]);
+				}
+				TSet<Chaos::TVector<int32, 2>> Visited;
+				for (auto Elem : BendingHash)
+				{
+					for (int32 i = 0; i < Elem.Value.Num(); ++i)
+					{
+						for (int32 j = i + 1; j < Elem.Value.Num(); ++j)
+						{
+							if (Elem.Value[i] == Elem.Value[j])
+								continue;
+							auto NewElem = Chaos::TVector<int32, 2>(Elem.Value[i], Elem.Value[j]);
+							if (!Visited.Contains(NewElem))
+							{
+								DoubleBendingConstraints.Add(NewElem);
+								Visited.Add(NewElem);
+								Visited.Add(Chaos::TVector<int32, 2>(Elem.Value[j], Elem.Value[i]));
+							}
+						}
+					}
+				}
+			}
+
+			ClothConstraints.SetVolumeConstraints(MoveTemp(DoubleBendingConstraints), Cloth->VolumeStiffness);
+		}
+		else
+		{
+			TArray<Chaos::TVector<int32, 3>> SurfaceConstraints = SurfaceElements;
+			ClothConstraints.SetVolumeConstraints(MoveTemp(SurfaceConstraints), Cloth->VolumeStiffness);
+		}
+	}
+
+	// Long range constraints
+	if (Cloth->StrainLimitingStiffness)
+	{
+		check(TriangleMesh.GetNumElements() > 0);
+		// PerFormance note: The Per constraint version of this function is quite a bit faster for smaller assets
+		// There might be a cross-over point where the PerParticle version is faster: To be determined
+		ClothConstraints.SetLongRangeConstraints(
+			TriangleMesh.GetPointToNeighborsMap(),
+			Cloth->StrainLimitingStiffness,
+			Cloth->LimitScale,
+			Cloth->bUseGeodesicDistance,
+			Cloth->bUseXPBDConstraints);
+	}
+
+	// Max distances
+	if (MaxDistances.Num())
+	{
+		ClothConstraints.SetMaximumDistanceConstraints(MaxDistances);
+	}
+
+	// Backstop Constraints
+	const TConstArrayView<float>& BackstopDistances = WeightMaps[(int32)EChaosWeightMapTarget::BackstopDistance];
+	const TConstArrayView<float>& BackstopRadiuses = WeightMaps[(int32)EChaosWeightMapTarget::BackstopRadius];
+	if (BackstopRadiuses.Num() && BackstopDistances.Num())
+	{
+		ClothConstraints.SetBackstopConstraints(BackstopDistances, BackstopRadiuses);
+	}
+
+	// Animation Drive Constraints
+	const TConstArrayView<float>& AnimDriveMultipliers = WeightMaps[(int32)EChaosWeightMapTarget::AnimDriveMultiplier];
+	if (AnimDriveMultipliers.Num())
+	{
+		ClothConstraints.SetAnimDriveConstraints(AnimDriveMultipliers);
+	}
+
+	// Shape target constraint
+	if (Cloth->ShapeTargetStiffness)
+	{
+		ClothConstraints.SetShapeTargetConstraints(Cloth->ShapeTargetStiffness);
+	}
+
+	// Commit rules to solver
+	ClothConstraints.CreateRules();
+
+	// Disable constraints by default
+	ClothConstraints.Enable(false);
+
+	// Update LOD stats
+	NumKinenamicParticles = 0;
+	NumDynammicParticles = 0;
+	for (int32 Index = 0; Index < NumParticles; ++Index)
+	{
+		if (InvMasses[Index] == 0.f)
+		{
+			++NumKinenamicParticles;
+		}
+		else
+		{
+			++NumDynammicParticles;
+		}
+	}
+}
+
+void FClothingSimulationCloth::FLODData::Enable(FClothingSimulationSolver* Solver, bool bEnable) const
+{
+	check(Solver);
+	const int32 Offset = SolverData.FindChecked(Solver).Offset;
+	check(Offset != INDEX_NONE);
+	
+	// Enable particles (and related constraints)
+	Solver->EnableParticles(Offset, bEnable);
+}
+
+void FClothingSimulationCloth::FLODData::ResetStartPose(FClothingSimulationSolver* Solver) const
+{
+	check(Solver);
+	const int32 Offset = SolverData.FindChecked(Solver).Offset;
+	check(Offset != INDEX_NONE);
+
+	TVector<float, 3>* const Ps = Solver->GetParticlePs(Offset);
+	TVector<float, 3>* const Xs = Solver->GetParticleXs(Offset);
+	TVector<float, 3>* const Vs = Solver->GetParticleVs(Offset);
+	const TVector<float, 3>* const AnimationPositions = Solver->GetAnimationPositions(Offset);
+	TVector<float, 3>* const OldAnimationPositions = Solver->GetOldAnimationPositions(Offset);
+
+	for (int32 Index = 0; Index < NumParticles; ++Index)
+	{
+		Ps[Index] = Xs[Index] = OldAnimationPositions[Index] = AnimationPositions[Index];
+		Vs[Index] = TVector<float, 3>(0.f);
+	}
+}
+
+void FClothingSimulationCloth::FLODData::UpdateNormals(FClothingSimulationSolver* Solver) const
+{
+	check(Solver);
+
+	const FSolverData& SolverDatum = SolverData.FindChecked(Solver);
+	const int32 Offset = SolverDatum.Offset;
+	const TTriangleMesh<float>& TriangleMesh = SolverDatum.TriangleMesh;
+
+	check(Offset != INDEX_NONE);
+
+	TConstArrayView<TVector<float, 3>> Points(Solver->GetParticleXs(Offset) - Offset, Offset + NumParticles);  // TODO: TriangleMesh still uses global array
+	TArray<TVector<float, 3>> FaceNormals;
+	TriangleMesh.GetFaceNormals(FaceNormals, Points, /*ReturnEmptyOnError =*/ false);
+
+	TArrayView<TVector<float, 3>> Normals(Solver->GetNormals(Offset), NumParticles);
+	TriangleMesh.GetPointNormals(Normals, FaceNormals, /*bUseGlobalArray =*/ false);
+}
+
+FClothingSimulationCloth::FClothingSimulationCloth(
+	FClothingSimulationMesh* InMesh,
+	TArray<FClothingSimulationCollider*>&& InColliders,
+	uint32 InGroupId,
+	EMassMode InMassMode,
+	float InMassValue,
+	float InMinPerParticleMass,
+	float InEdgeStiffness,
+	float InBendingStiffness,
+	bool bInUseBendingElements,
+	float InAreaStiffness,
+	float InVolumeStiffness,
+	bool bInUseThinShellVolumeConstraints,
+	float InStrainLimitingStiffness,
+	float InLimitScale,
+	bool bInUseGeodesicDistance,
+	float InMaxDistancesMultiplier,
+	float InAnimDriveSpringStiffness,
+	float InShapeTargetStiffness,
+	bool bInUseXPBDConstraints,
+	float InGravityScale,
+	bool bInIsGravityOverridden,
+	const TVector<float, 3>& InGravityOverride,
+	const TVector<float, 3>& InLinearVelocityScale,
+	float InAngularVelocityScale,
+	float InDragCoefficient,
+	float InDampingCoefficient,
+	float InCollisionThickness,
+	float InFrictionCoefficient,
+	bool bInUseSelfCollisions,
+	float InSelfCollisionThickness,
+	bool bInUseLODIndexOverride, 
+	int32 InLODIndexOverride)
+	: Mesh(nullptr)
+	, Colliders()
+	, GroupId(InGroupId)
+	, MassMode(InMassMode)
+	, MassValue(InMassValue)
+	, MinPerParticleMass(InMinPerParticleMass)
+	, EdgeStiffness(InEdgeStiffness)
+	, BendingStiffness(InBendingStiffness)
+	, bUseBendingElements(bInUseBendingElements)
+	, AreaStiffness(InAreaStiffness)
+	, VolumeStiffness(InVolumeStiffness)
+	, bUseThinShellVolumeConstraints(bInUseThinShellVolumeConstraints)
+	, StrainLimitingStiffness(InStrainLimitingStiffness)
+	, LimitScale(InLimitScale)
+	, bUseGeodesicDistance(bInUseGeodesicDistance)
+	, MaxDistancesMultiplier(InMaxDistancesMultiplier)
+	, AnimDriveSpringStiffness(InAnimDriveSpringStiffness)
+	, ShapeTargetStiffness(InShapeTargetStiffness)
+	, bUseXPBDConstraints(bInUseXPBDConstraints)
+	, GravityScale(InGravityScale)
+	, bIsGravityOverridden(bInIsGravityOverridden)
+	, GravityOverride(InGravityOverride)
+	, LinearVelocityScale(InLinearVelocityScale)
+	, AngularVelocityScale(InAngularVelocityScale)
+	, DragCoefficient(InDragCoefficient)
+	, DampingCoefficient(InDampingCoefficient)
+	, CollisionThickness(InCollisionThickness)
+	, FrictionCoefficient(InFrictionCoefficient)
+	, bUseSelfCollisions(bInUseSelfCollisions)
+	, SelfCollisionThickness(InSelfCollisionThickness)
+	, bUseLODIndexOverride(bInUseLODIndexOverride)
+	, LODIndexOverride(InLODIndexOverride)
+	, bNeedsReset(false)
+	, bNeedsTeleport(false)
+	, NumActiveKinematicParticles(0)
+	, NumActiveDynamicParticles(0)
+{
+	SetMesh(InMesh);
+	SetColliders(MoveTemp(InColliders));
+}
+
+FClothingSimulationCloth::~FClothingSimulationCloth()
+{
+}
+
+void FClothingSimulationCloth::SetMesh(FClothingSimulationMesh* InMesh)
+{
+	Mesh = InMesh;
+
+	// Backup all currently attached solvers
+	TArray<FClothingSimulationSolver*> Solvers;
+	LODIndices.GetKeys(Solvers);
+
+	// Remove this cloth from all solvers
+	for (FClothingSimulationSolver* const Solver : Solvers)
+	{
+		Solver->RemoveCloth(this);
+	}
+
+	// Reset LODs
+	const int32 NumLODs = Mesh ? Mesh->GetNumLODs() : 0;
+
+	LODData.Reset(NumLODs);
+	for (int32 Index = 0; Index < NumLODs; ++Index)
+	{
+		LODData.Emplace(Mesh->GetNumPoints(Index), Mesh->GetIndices(Index), Mesh->GetWeightMaps(Index));
+	}
+
+	// Re-add this cloth to existing solvers
+	for (FClothingSimulationSolver* const Solver : Solvers)
+	{
+		Solver->AddCloth(this);
+	}
+}
+
+void FClothingSimulationCloth::SetColliders(TArray<FClothingSimulationCollider*>&& InColliders)
+{
+	Colliders = InColliders;
+
+	// Backup all currently attached solvers
+	TArray<FClothingSimulationSolver*> Solvers;
+	LODIndices.GetKeys(Solvers);
+
+	// Remove this cloth from all solvers
+	for (FClothingSimulationSolver* const Solver : Solvers)
+	{
+		Solver->RemoveCloth(this);
+	}
+
+	// Re-add this cloth to existing solvers
+	for (FClothingSimulationSolver* const Solver : Solvers)
+	{
+		Solver->AddCloth(this);
+	}
+}
+
+void FClothingSimulationCloth::RemoveColliders()
+{
+	Colliders.Reset();
+
+	// Backup all currently attached solvers
+	TArray<FClothingSimulationSolver*> Solvers;
+	LODIndices.GetKeys(Solvers);
+
+	// Remove this cloth from all solvers
+	for (FClothingSimulationSolver* const Solver : Solvers)
+	{
+		Solver->RemoveCloth(this);
+	}
+
+	// Re-add this cloth to existing solvers
+	for (FClothingSimulationSolver* const Solver : Solvers)
+	{
+		Solver->AddCloth(this);
+	}
+}
+
+void FClothingSimulationCloth::Add(FClothingSimulationSolver* Solver)
+{
+	check(Solver);
+	check(!LODIndices.Find(Solver));
+
+	// Reset LODIndex
+	int32& LODIndex = LODIndices.Add(Solver);
+	LODIndex = INDEX_NONE;
+
+	// Add LODs
+	for (int32 Index = 0; Index < LODData.Num(); ++Index)
+	{
+		LODData[Index].Add(Solver, this, Index);
+	}
+
+	// Add Colliders
+	for (FClothingSimulationCollider* Collider : Colliders)
+	{
+		Collider->Add(Solver, this);
+	}
+}
+
+void FClothingSimulationCloth::Remove(FClothingSimulationSolver* Solver)
+{
+	// Remove LODIndex from map
+	LODIndices.Remove(Solver);
+
+	// Remove Colliders
+	for (FClothingSimulationCollider* Collider : Colliders)
+	{
+		Collider->Remove(Solver, this);
+	}
+}
+
+int32 FClothingSimulationCloth::GetNumParticles(int32 InLODIndex) const
+{
+	return (LODData.IsValidIndex(InLODIndex)) ? LODData[InLODIndex].NumParticles : 0;
+}
+
+int32 FClothingSimulationCloth::GetOffset(const FClothingSimulationSolver* Solver, int32 InLODIndex) const
+{
+	return (LODData.IsValidIndex(InLODIndex)) ? LODData[InLODIndex].SolverData.FindChecked(Solver).Offset : 0;
+}
+
+TVector<float, 3> FClothingSimulationCloth::GetGravity(const FClothingSimulationSolver* Solver) const
+{
+	check(Solver);
+	return Solver->IsClothGravityOverrideEnabled() && bIsGravityOverridden ? GravityOverride : Solver->GetGravity() * GravityScale;
+}
+
+TAABB<float, 3> FClothingSimulationCloth::CalculateBoundingBox(const FClothingSimulationSolver* Solver) const
+{
+	check(Solver);
+
+	// Calculate local space bounding box
+	TAABB<float, 3> BoundingBox = TAABB<float, 3>::EmptyAABB();
+
+	const TConstArrayView<TVector<float, 3>> ParticlePositions = GetParticlePositions(Solver);
+	for (const TVector<float, 3>& ParticlePosition : ParticlePositions)
+	{
+		BoundingBox.GrowToInclude(ParticlePosition);
+	}
+
+	// Return world space bounding box
+	return BoundingBox.TransformedAABB(TRigidTransform<float, 3>(Solver->GetLocalSpaceLocation(), TRotation<float, 3>::Identity));
+}
+
+int32 FClothingSimulationCloth::GetOffset(const FClothingSimulationSolver* Solver) const
+{
+	const int32 LODIndex = LODIndices.FindChecked(Solver);
+	return LODData.IsValidIndex(LODIndex) ? GetOffset(Solver, LODIndex) : INDEX_NONE;
+}
+
+const TTriangleMesh<float>& FClothingSimulationCloth::GetTriangleMesh(const FClothingSimulationSolver* Solver) const
+{
+	const int32 LODIndex = LODIndices.FindChecked(Solver);
+	static const TTriangleMesh<float> EmptyTriangleMesh;
+	return LODData.IsValidIndex(LODIndex) ? LODData[LODIndex].SolverData.FindChecked(Solver).TriangleMesh : EmptyTriangleMesh;
+}
+
+const TArray<TConstArrayView<float>>& FClothingSimulationCloth::GetWeightMaps(const FClothingSimulationSolver* Solver) const
+{
+	const int32 LODIndex = LODIndices.FindChecked(Solver);
+	static const TArray<TConstArrayView<float>> EmptyWeightMaps;
+	return LODData.IsValidIndex(LODIndex) ? LODData[LODIndex].WeightMaps : EmptyWeightMaps;
+}
+
+void FClothingSimulationCloth::Update(FClothingSimulationSolver* Solver)
+{
+	check(Solver);
+
+	// Exit if the input mesh is missing
+	if (!Mesh)
+	{
+		return;
+	}
+
+	// Retrieve LOD Index, either from the override, or from the mesh input
+	int32& LODIndex = LODIndices.FindChecked(Solver);  // Must be added to solver first
+
+	const int32 PrevLODIndex = LODIndex;
+	LODIndex = bUseLODIndexOverride && LODData.IsValidIndex(LODIndexOverride) ? LODIndexOverride : Mesh->GetLODIndex();
+
+	// Update reference space transform from the mesh's reference bone transform  TODO: Add override in the style of LODIndexOverride
+	const TRigidTransform<float, 3> OldReferenceSpaceTransform = ReferenceSpaceTransform;
+	ReferenceSpaceTransform = Mesh->GetReferenceBoneTransform();
+
+	// Update Cloth Colliders
+	for (FClothingSimulationCollider* Collider : Colliders)
+	{
+		Collider->Update(Solver, this);
+	}
+
+	// Update the source mesh skinned positions
+	const int32 PrevOffset = GetOffset(Solver, PrevLODIndex);
+	const int32 Offset = GetOffset(Solver, LODIndex);
+	check(PrevOffset != INDEX_NONE && Offset != INDEX_NONE);
+
+	Mesh->Update(Solver, PrevLODIndex, LODIndex, PrevOffset, Offset);
+
+	// LOD Switching
+	if (LODIndex != PrevLODIndex)
+	{
+		if (PrevLODIndex != INDEX_NONE)
+		{
+			// Disable previous LOD's particles
+			LODData[PrevLODIndex].Enable(Solver, false);
+		}
+		if (LODIndex != INDEX_NONE)
+		{
+			// Enable new LOD's particles
+			LODData[LODIndex].Enable(Solver, true);
+			NumActiveKinematicParticles = LODData[LODIndex].NumKinenamicParticles;
+			NumActiveDynamicParticles = LODData[LODIndex].NumDynammicParticles;
+
+			// Wrap new LOD based on previous LOD if possible (can only do 1 level LOD at a time, and if previous LOD exists)
+			bNeedsReset = bNeedsReset || !Mesh->WrapDeformLOD(
+				PrevLODIndex,
+				LODIndex,
+				Solver->GetNormals(PrevOffset),
+				Solver->GetParticlePs(PrevOffset),
+				Solver->GetParticleVs(PrevOffset),
+				Solver->GetParticlePs(Offset),
+				Solver->GetParticleXs(Offset),
+				Solver->GetParticleVs(Offset));
+		}
+		else
+		{
+			NumActiveKinematicParticles = 0;
+			NumActiveDynamicParticles = 0;
+		}
+	}
+
+	if (LODIndex != INDEX_NONE)
+	{
+		// Update Cloth group parameters  TODO: Cloth groups should exist as their own node object so that they can be used by several cloth objects
+
+		// Update gravity
+		// This code relies on the solver gravity property being already set.
+		// In order to use a cloth gravity override, it must first be enabled by the solver so that an override at solver level can still take precedence if needed.
+		// In all cases apart from when the cloth override is used, the gravity scale must be combined to the solver gravity value.
+		Solver->SetGravity(GroupId, GetGravity(Solver));
+
+		// Update wind
+		Solver->SetWindVelocityField(GroupId, DragCoefficient, &GetTriangleMesh(Solver));
+
+		// Update general solver properties
+		Solver->SetProperties(GroupId, DampingCoefficient, CollisionThickness, FrictionCoefficient);
+
+		// Set the reference input velocity and deal with teleport & reset
+		TVector<float, 3> OutLinearVelocityScale;
+		float OutAngularVelocityScale;
+
+		if (bNeedsReset)
+		{
+			// Make sure not to do any pre-sim transform just after a reset
+			OutLinearVelocityScale = TVector<float, 3>(1.f);
+			OutAngularVelocityScale = 1.f;
+
+			// Reset to start pose
+			LODData[LODIndex].ResetStartPose(Solver);
+		}
+		else if (bNeedsTeleport)
+		{
+			// Remove all impulse velocity from the last frame
+			OutLinearVelocityScale = TVector<float, 3>(0.f);
+			OutAngularVelocityScale = 0.f;
+		}
+		else
+		{
+			// Use the cloth group's parameters
+			OutLinearVelocityScale = LinearVelocityScale;
+			OutAngularVelocityScale = AngularVelocityScale;
+		}
+
+		Solver->SetReferenceVelocityScale(
+			GroupId,
+			OldReferenceSpaceTransform,
+			ReferenceSpaceTransform,
+			OutLinearVelocityScale,
+			OutAngularVelocityScale);
+	}
+
+	// Reset trigger flags
+	bNeedsTeleport = false;
+	bNeedsReset = false;
+}
+
+void FClothingSimulationCloth::PostUpdate(FClothingSimulationSolver* Solver)
+{
+	check(Solver);
+
+	const int32 LODIndex = LODIndices.FindChecked(Solver);
+	if (LODIndex != INDEX_NONE)
+	{
+		// Update normals
+		LODData[LODIndex].UpdateNormals(Solver);
+	}
+}
+
+TConstArrayView<TVector<float, 3>> FClothingSimulationCloth::GetAnimationPositions(const FClothingSimulationSolver* Solver) const
+{
+	check(Solver);
+	const int32 LODIndex = LODIndices.FindChecked(Solver);
+	check(GetOffset(Solver, LODIndex) != INDEX_NONE);
+	return TConstArrayView<TVector<float, 3>>(Solver->GetAnimationPositions(GetOffset(Solver, LODIndex)), GetNumParticles(LODIndex));
+}
+
+TConstArrayView<TVector<float, 3>> FClothingSimulationCloth::GetAnimationNormals(const FClothingSimulationSolver* Solver) const
+{
+	check(Solver);
+	const int32 LODIndex = LODIndices.FindChecked(Solver);
+	check(GetOffset(Solver, LODIndex) != INDEX_NONE);
+	return TConstArrayView<TVector<float, 3>>(Solver->GetAnimationNormals(GetOffset(Solver, LODIndex)), GetNumParticles(LODIndex));
+}
+
+TConstArrayView<TVector<float, 3>> FClothingSimulationCloth::GetParticlePositions(const FClothingSimulationSolver* Solver) const
+{
+	check(Solver);
+	const int32 LODIndex = LODIndices.FindChecked(Solver);
+	check(GetOffset(Solver, LODIndex) != INDEX_NONE);
+	return TConstArrayView<TVector<float, 3>>(Solver->GetParticleXs(GetOffset(Solver, LODIndex)), GetNumParticles(LODIndex));
+}
+
+TConstArrayView<TVector<float, 3>> FClothingSimulationCloth::GetParticleNormals(const FClothingSimulationSolver* Solver) const
+{
+	check(Solver);
+	const int32 LODIndex = LODIndices.FindChecked(Solver);
+	check(GetOffset(Solver, LODIndex) != INDEX_NONE);
+	return TConstArrayView<TVector<float, 3>>(Solver->GetNormals(GetOffset(Solver, LODIndex)), GetNumParticles(LODIndex));
+}
+
+TConstArrayView<float> FClothingSimulationCloth::GetParticleInvMasses(const FClothingSimulationSolver* Solver) const
+{
+	check(Solver);
+	const int32 LODIndex = LODIndices.FindChecked(Solver);
+	check(GetOffset(Solver, LODIndex) != INDEX_NONE);
+	return TConstArrayView<float>(Solver->GetParticleInvMasses(GetOffset(Solver, LODIndex)), GetNumParticles(LODIndex));
+}
diff --git a/Engine/Plugins/Experimental/ChaosCloth/Source/Chaos/Private/ChaosCloth/ChaosClothingSimulationCollider.cpp b/Engine/Plugins/Experimental/ChaosCloth/Source/Chaos/Private/ChaosCloth/ChaosClothingSimulationCollider.cpp
new file mode 100644
--- /dev/null
+++ b/Engine/Plugins/Experimental/ChaosCloth/Source/Chaos/Private/ChaosCloth/ChaosClothingSimulationCollider.cpp
@@ -0,0 +1,677 @@
+// Copyright Epic Games, Inc. All Rights Reserved.
+#include "ChaosCloth/ChaosClothingSimulationCollider.h"
+#include "ChaosCloth/ChaosClothingSimulationSolver.h"
+#include "ChaosCloth/ChaosClothingSimulationCloth.h"
+#include "ChaosCloth/ChaosClothPrivate.h"
+#include "Chaos/Plane.h"
+#include "Chaos/Sphere.h"
+#include "Chaos/TaperedCylinder.h"
+#include "Chaos/Convex.h"
+#include "ClothingSimulation.h" // For context
+#include "ClothingAsset.h"
+#include "HAL/PlatformMath.h"
+#include "Containers/ArrayView.h"
+#include "Containers/BitArray.h"
+
+DECLARE_CYCLE_STAT(TEXT("Chaos Cloth Update Collider"), STAT_ChaosClothingSimulationColliderUpdate, STATGROUP_ChaosCloth);
+
+using namespace Chaos;
+
+void FClothingSimulationCollider::FLODData::Add(
+	FClothingSimulationSolver* Solver,
+	FClothingSimulationCloth* Cloth,
+	const FClothCollisionData& InClothCollisionData,
+	const TArray<int32>& UsedBoneIndices)
+{
+	check(Solver);
+
+	// Keep a list of all collisions
+	ClothCollisionData = InClothCollisionData;
+
+	// Calculate the number of geometries
+	int32 NumSpheres = ClothCollisionData.Spheres.Num();
+	TBitArray<TInlineAllocator<8>> CapsuleEnds(false, NumSpheres);  // Init on stack for 256 spheres, any number over this will be allocated
+
+	const int32 NumCapsules = ClothCollisionData.SphereConnections.Num();
+	for (int32 Index = 0; Index < NumCapsules; ++Index)
+	{
+		const FClothCollisionPrim_SphereConnection& Connection = ClothCollisionData.SphereConnections[Index];
+		for (int32 SphereIndex : Connection.SphereIndices)
+		{
+			FBitReference IsCapsuleEnd = CapsuleEnds[SphereIndex];
+			if (!IsCapsuleEnd)
+			{
+				IsCapsuleEnd = true;
+				--NumSpheres;
+			}
+		}
+	}
+
+	const uint32 NumConvexes = ClothCollisionData.Convexes.Num();
+	const int32 NumBoxes = ClothCollisionData.Boxes.Num();
+	NumGeometries = NumSpheres + NumCapsules + NumConvexes + NumBoxes;
+
+	// Retrieve cloth group Id, or use INDEX_NONE if this collider applies to all cloths (when Cloth == nullptr)
+	const uint32 GroupId = Cloth ? Cloth->GetGroupId() : INDEX_NONE;
+
+	// The offset will be set to the first collision particle's index
+	// Try to reuse existing offsets when Add is called during the solver's collider update
+	const TPair<const FClothingSimulationSolver*, const FClothingSimulationCloth*> SolverClothPair(Solver, Cloth);
+	int32* const OffsetPtr = Offsets.Find(SolverClothPair);
+	const bool bIsNewCollider = !OffsetPtr;
+	check(!Cloth || bIsNewCollider);  // TODO: We need to be able to update cloth colliders too at some point, but this won't work as it is
+	int32& Offset = bIsNewCollider ? Offsets.Add(SolverClothPair) : *OffsetPtr;
+	Offset = Solver->AddCollisionParticles(NumGeometries, GroupId, bIsNewCollider ? INDEX_NONE : Offset);
+
+	// Capsules
+	const int32 CapsuleOffset = Offset;
+	if (NumCapsules)
+	{
+		int32* const BoneIndices = Solver->GetCollisionBoneIndices(CapsuleOffset);
+		TRigidTransform<float, 3>* const BaseTransforms = Solver->GetCollisionBaseTransforms(CapsuleOffset);
+
+		for (int32 Index = 0; Index < NumCapsules; ++Index)
+		{
+			const FClothCollisionPrim_SphereConnection& Connection = ClothCollisionData.SphereConnections[Index];
+
+			const int32 SphereIndex0 = Connection.SphereIndices[0];
+			const int32 SphereIndex1 = Connection.SphereIndices[1];
+			checkSlow(SphereIndex0 != SphereIndex1);
+			const FClothCollisionPrim_Sphere& Sphere0 = ClothCollisionData.Spheres[SphereIndex0];
+			const FClothCollisionPrim_Sphere& Sphere1 = ClothCollisionData.Spheres[SphereIndex1];
+
+			BoneIndices[Index] = GetMappedBoneIndex(UsedBoneIndices, Sphere0.BoneIndex);
+			checkSlow(Sphere0.BoneIndex == Sphere1.BoneIndex);
+			UE_CLOG(Sphere0.BoneIndex != Sphere1.BoneIndex,
+				LogChaosCloth, Warning, TEXT("Found a legacy cloth asset with a collision capsule spanning across two bones. This is not supported with the current system."));
+			UE_LOG(LogChaosCloth, VeryVerbose, TEXT("Found collision capsule on bone index %d."), BoneIndices[Index]);
+
+			const TVector<float, 3> X0 = Sphere0.LocalPosition;
+			const TVector<float, 3> X1 = Sphere1.LocalPosition;
+			const TVector<float, 3> Axis = X1 - X0;
+			const float AxisSize = Axis.Size();
+
+			const float Radius0 = Sphere0.Radius;
+			const float Radius1 = Sphere1.Radius;
+			float MinRadius, MaxRadius;
+			if (Radius0 <= Radius1) { MinRadius = Radius0; MaxRadius = Radius1; }
+			else { MinRadius = Radius1; MaxRadius = Radius0; }
+
+			if (AxisSize < KINDA_SMALL_NUMBER)
+			{
+				// Sphere
+				BaseTransforms[Index] = TRigidTransform<float, 3>::Identity;
+
+				Solver->SetCollisionGeometry(CapsuleOffset, Index,
+					MakeUnique<TSphere<float, 3>>(
+						X0,
+						MaxRadius));
+			}
+			else if (MaxRadius - MinRadius < KINDA_SMALL_NUMBER)
+			{
+				// Capsule
+				const TVector<float, 3> Center = (X0 + X1) * 0.5f;  // Construct a capsule centered at the origin along the Z axis
+				const TRotation<float, 3> Rotation = TRotation<float, 3>::FromRotatedVector(
+					TVector<float, 3>::AxisVector(2),
+					Axis.GetSafeNormal());
+
+				BaseTransforms[Index] = TRigidTransform<float, 3>(Center, Rotation);
+
+				const float HalfHeight = AxisSize * 0.5f;
+				Solver->SetCollisionGeometry(CapsuleOffset, Index,
+					MakeUnique<TCapsule<float>>(
+						TVector<float, 3>(0.f, 0.f, -HalfHeight), // Min
+						TVector<float, 3>(0.f, 0.f, HalfHeight), // Max
+						MaxRadius));
+			}
+			else
+			{
+				// Tapered capsule
+				BaseTransforms[Index] = TRigidTransform<float, 3>::Identity;
+
+				TArray<TUniquePtr<FImplicitObject>> Objects;
+				Objects.Reserve(3);
+				Objects.Add(TUniquePtr<FImplicitObject>(
+					new TTaperedCylinder<float>(X0, X1, Radius0, Radius1)));
+				Objects.Add(TUniquePtr<FImplicitObject>(
+					new TSphere<float, 3>(X0, Radius0)));
+				Objects.Add(TUniquePtr<FImplicitObject>(
+					new TSphere<float, 3>(X1, Radius1)));
+				Solver->SetCollisionGeometry(CapsuleOffset, Index,
+					MakeUnique<FImplicitObjectUnion>(MoveTemp(Objects)));  // TODO(Kriss.Gossart): Replace this once a TTaperedCapsule implicit type is implemented (note: this tapered cylinder with spheres is an approximation of a real tapered capsule)
+			}
+		}
+	}
+
+	// Spheres
+	const int32 SphereOffset = CapsuleOffset + NumCapsules;
+	if (NumSpheres != 0)
+	{
+		int32* const BoneIndices = Solver->GetCollisionBoneIndices(SphereOffset);
+		TRigidTransform<float, 3>* const BaseTransforms = Solver->GetCollisionBaseTransforms(SphereOffset);
+
+		for (int32 Index = 0, SphereIndex = 0; SphereIndex < ClothCollisionData.Spheres.Num(); ++SphereIndex)
+		{
+			// Skip spheres that are the end caps of capsules.
+			if (CapsuleEnds[SphereIndex])
+			//if (CapsuleEnds.Contains(SphereIndex))
+			{
+				continue;
+			}
+
+			const FClothCollisionPrim_Sphere& Sphere = ClothCollisionData.Spheres[SphereIndex];
+
+			BoneIndices[Index] = GetMappedBoneIndex(UsedBoneIndices, Sphere.BoneIndex);
+			UE_LOG(LogChaosCloth, VeryVerbose, TEXT("Found collision sphere on bone index %d."), BoneIndices[Index]);
+
+			BaseTransforms[Index] = TRigidTransform<float, 3>::Identity;
+
+			Solver->SetCollisionGeometry(SphereOffset, Index,
+				MakeUnique<TSphere<float, 3>>(
+					Sphere.LocalPosition,
+					Sphere.Radius));
+
+			++Index;
+		}
+	}
+
+	// Convexes
+	const int32 ConvexOffset = SphereOffset + NumSpheres;
+	if (NumConvexes != 0)
+	{
+		int32* const BoneIndices = Solver->GetCollisionBoneIndices(ConvexOffset);
+		TRigidTransform<float, 3>* const BaseTransforms = Solver->GetCollisionBaseTransforms(ConvexOffset);
+
+		for (uint32 Index = 0; Index < NumConvexes; ++Index)
+		{
+			const FClothCollisionPrim_Convex& Convex = ClothCollisionData.Convexes[Index];
+
+			// Always initialize the collision particle transforms before setting any geometry as otherwise NaNs gets detected during the bounding box updates
+			BaseTransforms[Index] = TRigidTransform<float, 3>::Identity;
+
+			BoneIndices[Index] = GetMappedBoneIndex(UsedBoneIndices, Convex.BoneIndex);
+			UE_LOG(LogChaosCloth, VeryVerbose, TEXT("Found collision convex on bone index %d."), BoneIndices[Index]);
+
+			TArray<TPlaneConcrete<float, 3>> Planes;
+
+			const int32 NumSurfacePoints = Convex.SurfacePoints.Num();
+			const int32 NumPlanes = Convex.Planes.Num();
+
+			if (NumSurfacePoints < 4)
+			{
+				UE_LOG(LogChaosCloth, Warning, TEXT("Invalid convex collision: not enough surface points."));
+			}
+			else if (NumPlanes < 4)
+			{
+				UE_LOG(LogChaosCloth, Warning, TEXT("Invalid convex collision: not enough planes."));
+			}
+			else
+			{
+				// Retrieve convex planes
+				Planes.Reserve(NumPlanes);
+				for (const FPlane& Plane : Convex.Planes)
+				{
+					FPlane NormalizedPlane(Plane);
+					if (NormalizedPlane.Normalize())
+					{
+						const TVector<float, 3> Normal(static_cast<FVector>(NormalizedPlane));
+						const TVector<float, 3> Base = Normal * NormalizedPlane.W;
+
+						Planes.Add(TPlaneConcrete<float, 3>(Base, Normal));
+					}
+					else
+					{
+						UE_LOG(LogChaosCloth, Warning, TEXT("Invalid convex collision: bad plane normal."));
+						break;
+					}
+				}
+			}
+
+			if (Planes.Num() == NumPlanes)
+			{
+				// Retrieve particles
+				TParticles<float, 3> SurfaceParticles;
+				SurfaceParticles.Resize(NumSurfacePoints);
+				for (int32 ParticleIndex = 0; ParticleIndex < NumSurfacePoints; ++ParticleIndex)
+				{
+					SurfaceParticles.X(ParticleIndex) = Convex.SurfacePoints[ParticleIndex];
+				}
+
+				// Setup the collision particle geometry
+				Solver->SetCollisionGeometry(ConvexOffset, Index, MakeUnique<FConvex>(MoveTemp(Planes), MoveTemp(SurfaceParticles)));
+			}
+			else
+			{
+				UE_LOG(LogChaosCloth, Warning, TEXT("Replacing invalid convex collision by a default unit sphere."));
+				Solver->SetCollisionGeometry(ConvexOffset, Index, MakeUnique<TSphere<float, 3>>(TVector<float, 3>(0.0f), 1.0f));  // Default to a unit sphere to replace the faulty convex
+			}
+		}
+	}
+
+	// Boxes
+	const int32 BoxOffset = ConvexOffset + NumConvexes;
+	if (NumBoxes != 0)
+	{
+		int32* const BoneIndices = Solver->GetCollisionBoneIndices(BoxOffset);
+		TRigidTransform<float, 3>* const BaseTransforms = Solver->GetCollisionBaseTransforms(BoxOffset);
+		
+		for (int32 Index = 0; Index < NumBoxes; ++Index)
+		{
+			const FClothCollisionPrim_Box& Box = ClothCollisionData.Boxes[Index];
+			
+			BaseTransforms[Index] = TRigidTransform<float, 3>(Box.LocalPosition, Box.LocalRotation);
+			
+			BoneIndices[Index] = GetMappedBoneIndex(UsedBoneIndices, Box.BoneIndex);
+			UE_LOG(LogChaosCloth, VeryVerbose, TEXT("Found collision box on bone index %d."), BoneIndices[Index]);
+
+			Solver->SetCollisionGeometry(BoxOffset, Index, MakeUnique<TBox<float, 3>>(-Box.HalfExtents, Box.HalfExtents));
+		}
+	}
+
+	UE_LOG(LogChaosCloth, VeryVerbose, TEXT("Added collisions: %d spheres, %d capsules, %d convexes, %d boxes."), NumSpheres, NumCapsules, NumConvexes, NumBoxes);
+}
+
+void FClothingSimulationCollider::FLODData::Update(FClothingSimulationSolver* Solver, FClothingSimulationCloth* Cloth, const FClothingSimulationContextCommon* Context)
+{
+	check(Solver);
+	if (NumGeometries)
+	{
+		const TPair<const FClothingSimulationSolver*, const FClothingSimulationCloth*> SolverClothPair(Solver, Cloth);
+		const int32 Offset = Offsets.FindChecked(SolverClothPair);
+		const int32* const BoneIndices = Solver->GetCollisionBoneIndices(Offset);
+		const TRigidTransform<float, 3>* BaseTransforms = Solver->GetCollisionBaseTransforms(Offset);
+		TRigidTransform<float, 3>* const CollisionTransforms = Solver->GetCollisionTransforms(Offset);
+
+		const TArrayView<const FTransform> BoneTransforms = Context ? Context->BoneTransforms : TArrayView<const FTransform>();
+
+		FTransform ComponentToLocalSpace = Context ? Context->ComponentToWorld : FTransform::Identity;
+		ComponentToLocalSpace.AddToTranslation(-Solver->GetLocalSpaceLocation());
+	
+		// Update the collision transforms
+		for (int32 Index = 0; Index < NumGeometries; ++Index)
+		{
+			const int32 BoneIndex = BoneIndices[Index];
+			CollisionTransforms[Index] = BoneTransforms.IsValidIndex(BoneIndex) ?
+				BaseTransforms[Index] * BoneTransforms[BoneIndex] * ComponentToLocalSpace :
+				BaseTransforms[Index] * ComponentToLocalSpace;  // External collisions don't map to a bone
+		}
+	}
+}
+
+void FClothingSimulationCollider::FLODData::Enable(FClothingSimulationSolver* Solver, FClothingSimulationCloth* Cloth, bool bEnable)
+{
+	check(Solver);
+	if (NumGeometries)
+	{
+		const TPair<const FClothingSimulationSolver*, const FClothingSimulationCloth*> SolverClothPair(Solver, Cloth);
+		const int32 Offset = Offsets.FindChecked(SolverClothPair);
+		Solver->EnableCollisionParticles(Offset, bEnable);
+	}
+}
+
+void FClothingSimulationCollider::FLODData::ResetStartPose(FClothingSimulationSolver* Solver, FClothingSimulationCloth* Cloth)
+{
+	check(Solver);
+	if (NumGeometries)
+	{
+		const TPair<const FClothingSimulationSolver*, const FClothingSimulationCloth*> SolverClothPair(Solver, Cloth);
+		const int32 Offset = Offsets.FindChecked(SolverClothPair);
+		const TRigidTransform<float, 3>* const CollisionTransforms = Solver->GetCollisionTransforms(Offset);
+		TRigidTransform<float, 3>* const OldCollisionTransforms = Solver->GetOldCollisionTransforms(Offset);
+		TRotation<float, 3>* const Rs = Solver->GetCollisionParticleRs(Offset);
+		TVector<float, 3>* const Xs = Solver->GetCollisionParticleXs(Offset);
+
+		for (int32 Index = 0; Index < NumGeometries; ++Index)
+		{
+			OldCollisionTransforms[Index] = CollisionTransforms[Index];
+			Xs[Index] = CollisionTransforms[Index].GetTranslation();
+			Rs[Index] = CollisionTransforms[Index].GetRotation();
+		}
+	}
+}
+
+FClothingSimulationCollider::FClothingSimulationCollider(
+	const UClothingAssetCommon* InAsset,
+	const USkeletalMeshComponent* InSkeletalMeshComponent,
+	bool bInUseLODIndexOverride,
+	int32 InLODIndexOverride)
+	: Asset(InAsset)
+	, SkeletalMeshComponent(InSkeletalMeshComponent)
+	, CollisionData(nullptr)
+	, bUseLODIndexOverride(bInUseLODIndexOverride)
+	, LODIndexOverride(InLODIndexOverride)
+	, LODIndex(INDEX_NONE)
+{
+	// Prepare LOD array
+	const int32 NumLODs = Asset ? Asset->LodData.Num() : 0;
+	LODData.SetNum((int32)ECollisionDataType::LODs + NumLODs);
+}
+
+FClothingSimulationCollider::~FClothingSimulationCollider()
+{
+}
+
+FClothCollisionData FClothingSimulationCollider::GetCollisionData() const
+{
+	FClothCollisionData ClothCollisionData;
+	ClothCollisionData.Append(LODData[(int32)ECollisionDataType::Global].ClothCollisionData);
+	ClothCollisionData.Append(LODData[(int32)ECollisionDataType::Dynamic].ClothCollisionData);
+	if (LODIndex >= (int32)ECollisionDataType::LODs)
+	{
+		ClothCollisionData.Append(LODData[LODIndex].ClothCollisionData);
+	}
+	return ClothCollisionData;
+}
+
+void FClothingSimulationCollider::ExtractPhysicsAssetCollision(FClothCollisionData& ClothCollisionData, TArray<int32>& UsedBoneIndices)
+{
+	ClothCollisionData.Reset();
+	UsedBoneIndices.Reset();
+
+	if(!Asset)
+	{
+		return;
+	}
+
+	if (const UPhysicsAsset* const PhysAsset = Asset->PhysicsAsset)
+	{
+		const USkeletalMesh* const TargetMesh = CastChecked<USkeletalMesh>(Asset->GetOuter());
+
+		UsedBoneIndices.Reserve(PhysAsset->SkeletalBodySetups.Num());
+
+		for (const USkeletalBodySetup* BodySetup : PhysAsset->SkeletalBodySetups)
+		{
+			if (!BodySetup)
+				continue;
+
+			const int32 MeshBoneIndex = TargetMesh->RefSkeleton.FindBoneIndex(BodySetup->BoneName);
+			const int32 MappedBoneIndex = UsedBoneIndices.Add(MeshBoneIndex);
+			
+			// Add capsules
+			const FKAggregateGeom& AggGeom = BodySetup->AggGeom;
+			if (AggGeom.SphylElems.Num())
+			{
+				for (const FKSphylElem& SphylElem : AggGeom.SphylElems)
+				{
+					if (SphylElem.Length == 0.0f)
+					{
+						// Add extracted sphere collision data
+						FClothCollisionPrim_Sphere Sphere;
+						Sphere.LocalPosition = SphylElem.Center;
+						Sphere.Radius = SphylElem.Radius;
+						Sphere.BoneIndex = MappedBoneIndex;
+						ClothCollisionData.Spheres.Add(Sphere);
+					}
+					else
+					{
+						// Add extracted spheres collision data
+						FClothCollisionPrim_Sphere Sphere0;
+						FClothCollisionPrim_Sphere Sphere1;
+						const FVector OrientedDirection = SphylElem.Rotation.RotateVector(FVector::UpVector);
+						const FVector HalfDim = OrientedDirection * (SphylElem.Length / 2.f);
+						Sphere0.LocalPosition = SphylElem.Center - HalfDim;
+						Sphere1.LocalPosition = SphylElem.Center + HalfDim;
+						Sphere0.Radius = SphylElem.Radius;
+						Sphere1.Radius = SphylElem.Radius;
+						Sphere0.BoneIndex = MappedBoneIndex;
+						Sphere1.BoneIndex = MappedBoneIndex;
+
+						// Add extracted sphere connection collision data
+						FClothCollisionPrim_SphereConnection SphereConnection;
+						SphereConnection.SphereIndices[0] = ClothCollisionData.Spheres.Add(Sphere0);
+						SphereConnection.SphereIndices[1] = ClothCollisionData.Spheres.Add(Sphere1);
+						ClothCollisionData.SphereConnections.Add(SphereConnection);
+					}
+				}
+			}
+
+			// Add spheres
+			for (const FKSphereElem& SphereElem : AggGeom.SphereElems)
+			{
+				// Add extracted sphere collision data
+				FClothCollisionPrim_Sphere Sphere;
+				Sphere.LocalPosition = SphereElem.Center;
+				Sphere.Radius = SphereElem.Radius;
+				Sphere.BoneIndex = MappedBoneIndex;
+				ClothCollisionData.Spheres.Add(Sphere);
+			}
+
+			// Add boxes
+			for (const FKBoxElem& BoxElem : AggGeom.BoxElems)
+			{
+				// Add extracted box collision data
+				FClothCollisionPrim_Box Box;
+				Box.LocalPosition = BoxElem.Center;
+				Box.LocalRotation = BoxElem.Rotation.Quaternion();
+				Box.HalfExtents = FVector(BoxElem.X, BoxElem.Y, BoxElem.Z) * 0.5f;
+				Box.BoneIndex = MappedBoneIndex;
+				ClothCollisionData.Boxes.Add(Box);
+			}
+
+			// Add tapered capsules
+			for (const FKTaperedCapsuleElem& TaperedCapsuleElem : AggGeom.TaperedCapsuleElems)
+			{
+				if (TaperedCapsuleElem.Length == 0)
+				{
+					// Add extracted sphere collision data
+					FClothCollisionPrim_Sphere Sphere;
+					Sphere.LocalPosition = TaperedCapsuleElem.Center;
+					Sphere.Radius = FMath::Max(TaperedCapsuleElem.Radius0, TaperedCapsuleElem.Radius1);
+					Sphere.BoneIndex = MappedBoneIndex;
+					ClothCollisionData.Spheres.Add(Sphere);
+				}
+				else
+				{
+					// Add extracted spheres collision data
+					FClothCollisionPrim_Sphere Sphere0;
+					FClothCollisionPrim_Sphere Sphere1;
+					const FVector OrientedDirection = TaperedCapsuleElem.Rotation.RotateVector(FVector::UpVector);
+					const FVector HalfDim = OrientedDirection * (TaperedCapsuleElem.Length / 2.f);
+					Sphere0.LocalPosition = TaperedCapsuleElem.Center + HalfDim;
+					Sphere1.LocalPosition = TaperedCapsuleElem.Center - HalfDim;
+					Sphere0.Radius = TaperedCapsuleElem.Radius0;
+					Sphere1.Radius = TaperedCapsuleElem.Radius1;
+					Sphere0.BoneIndex = MappedBoneIndex;
+					Sphere1.BoneIndex = MappedBoneIndex;
+
+					// Add extracted sphere connection collision data
+					FClothCollisionPrim_SphereConnection SphereConnection;
+					SphereConnection.SphereIndices[0] = ClothCollisionData.Spheres.Add(Sphere0);
+					SphereConnection.SphereIndices[1] = ClothCollisionData.Spheres.Add(Sphere1);
+					ClothCollisionData.SphereConnections.Add(SphereConnection);
+				}
+			}
+
+#if !PLATFORM_LUMIN && !PLATFORM_ANDROID  // TODO(Kriss.Gossart): Compile on Android and fix whatever errors the following code is causing
+			// Add convexes
+			for (const FKConvexElem& ConvexElem : AggGeom.ConvexElems)
+			{
+				// Add stub for extracted collision data
+				FClothCollisionPrim_Convex Convex;
+				Convex.BoneIndex = MappedBoneIndex;
+#if PHYSICS_INTERFACE_PHYSX
+				// Collision bodies are stored in PhysX specific data structures so they can only be imported if we enable PhysX.
+				const physx::PxConvexMesh* const PhysXMesh = ConvexElem.GetConvexMesh();  // TODO(Kriss.Gossart): Deal with this legacy structure in a different place, so that there's only TConvex
+				const int32 NumPolygons = int32(PhysXMesh->getNbPolygons());
+				Convex.Planes.SetNumUninitialized(NumPolygons);
+				for (int32 i = 0; i < NumPolygons; ++i)
+				{
+					physx::PxHullPolygon Poly;
+					PhysXMesh->getPolygonData(i, Poly);
+					check(Poly.mNbVerts == 3);
+					const auto Indices = PhysXMesh->getIndexBuffer() + Poly.mIndexBase;
+
+					Convex.Planes[i] = FPlane(
+						ConvexElem.VertexData[Indices[0]],
+						ConvexElem.VertexData[Indices[1]],
+						ConvexElem.VertexData[Indices[2]]);
+				}
+
+				// Rebuild surface points
+				Convex.RebuildSurfacePoints();
+
+#elif WITH_CHAOS  // #if PHYSICS_INTERFACE_PHYSX
+				const FImplicitObject& ChaosConvexMesh = *ConvexElem.GetChaosConvexMesh();
+				const FConvex& ChaosConvex = ChaosConvexMesh.GetObjectChecked<FConvex>();
+
+				// Copy planes
+				const TArray<TPlaneConcrete<float, 3>>& Planes = ChaosConvex.GetFaces();
+				Convex.Planes.Reserve(Planes.Num());
+				for (const TPlaneConcrete<float, 3>& Plane : Planes)
+				{
+					Convex.Planes.Add(FPlane(Plane.X(), Plane.Normal()));
+				}
+
+				// Copy surface points
+				const uint32 NumSurfacePoints = ChaosConvex.GetSurfaceParticles().Size();
+				Convex.SurfacePoints.Reserve(NumSurfacePoints);
+				for (uint32 ParticleIndex = 0; ParticleIndex < NumSurfacePoints; ++ParticleIndex)
+				{
+					Convex.SurfacePoints.Add(ChaosConvex.GetSurfaceParticles().X(ParticleIndex));
+				}
+#endif  // #if PHYSICS_INTERFACE_PHYSX #elif WITH_CHAOS
+
+				// Add extracted collision data
+				ClothCollisionData.Convexes.Add(Convex);
+			}
+#endif  // #if !PLATFORM_LUMIN && !PLATFORM_ANDROID
+
+		}  // End for PhysAsset->SkeletalBodySetups
+	}  // End if Asset->PhysicsAsset
+}
+
+void FClothingSimulationCollider::Add(FClothingSimulationSolver* Solver, FClothingSimulationCloth* Cloth)
+{
+	check(Solver);
+
+	// Reset LOD
+	LODIndex = INDEX_NONE;  // TODO: Have map of LODIndex per solver
+
+	// Create physics asset collisions, this will affect all LODs, so store it at index 0
+	FClothCollisionData PhysicsAssetCollisionData;
+	TArray<int32> UsedBoneIndices;
+	ExtractPhysicsAssetCollision(PhysicsAssetCollisionData, UsedBoneIndices);
+
+	LODData[(int32)ECollisionDataType::Global].Add(Solver, Cloth, PhysicsAssetCollisionData, UsedBoneIndices);
+
+	// Create user collision data
+	LODData[(int32)ECollisionDataType::Dynamic].Add(Solver, Cloth, CollisionData ? *CollisionData : FClothCollisionData());
+
+	// Create legacy asset LOD collisions
+	const int32 NumLODs = Asset ? Asset->LodData.Num() : 0;
+	for (int32 Index = 0; Index < NumLODs; ++Index)
+	{
+		// Warn about legacy apex collisions
+		const FClothCollisionData& AssetCollisionData = Asset->LodData[Index].CollisionData;
+		UE_CLOG(AssetCollisionData.Spheres.Num() > 0 || AssetCollisionData.SphereConnections.Num() > 0 || AssetCollisionData.Convexes.Num() > 0,
+			LogChaosCloth, Warning, TEXT(
+				"Actor '%s' component '%s' has %d sphere, %d capsule, and %d convex collision objects for"
+				"physics authored as part of a LOD construct, probably by the Apex cloth authoring system."
+				"This is deprecated. Please update your asset!"),
+			SkeletalMeshComponent->GetOwner() ? *SkeletalMeshComponent->GetOwner()->GetName() : TEXT("None"),
+			*SkeletalMeshComponent->GetName(),
+			AssetCollisionData.Spheres.Num(),
+			AssetCollisionData.SphereConnections.Num(),
+			AssetCollisionData.Convexes.Num());
+
+		// Add legacy collision
+		LODData[(int32)ECollisionDataType::LODs + Index].Add(Solver, Cloth, AssetCollisionData, Asset->UsedBoneIndices);
+	}
+}
+
+void FClothingSimulationCollider::Update(FClothingSimulationSolver* Solver, FClothingSimulationCloth* Cloth)
+{
+	check(Solver);
+
+	SCOPE_CYCLE_COUNTER(STAT_ChaosClothingSimulationColliderUpdate);
+
+	// Retrieve context (can be null when no skeletal mesh component is specified)
+	const FClothingSimulationContextCommon* const Context = SkeletalMeshComponent ? static_cast<const FClothingSimulationContextCommon*>(SkeletalMeshComponent->GetClothingSimulationContext()) : nullptr;
+
+	// Update the collision transforms
+	// Note: All the LODs are updated at once, so that the previous transforms are always correct during LOD switching
+	for (FLODData& LODDatum : LODData)
+	{
+		LODDatum.Update(Solver, Cloth, Context);
+	}
+
+	// Update current LOD index
+	const int32 PrevLODIndex = LODIndex;
+	LODIndex = (int32)ECollisionDataType::LODs + (bUseLODIndexOverride ? LODIndexOverride : Cloth ? Cloth->GetLODIndex(Solver) : INDEX_NONE);
+	if (!LODData.IsValidIndex(LODIndex))
+	{
+		LODIndex = (int32)ECollisionDataType::LODs + INDEX_NONE;
+	}
+
+	// Enable/disable particle LOD
+	if (LODIndex != PrevLODIndex)
+	{
+		if (PrevLODIndex == INDEX_NONE)
+		{
+			// Enable non LOD collisions at first initialization
+			LODData[(int32)ECollisionDataType::Global].Enable(Solver, Cloth, true);
+
+			// Update initial state for collisions
+			LODData[(int32)ECollisionDataType::Global].ResetStartPose(Solver, Cloth);
+		}
+		else if (PrevLODIndex > 0)
+		{
+			// Disable previous LOD (can't be index 0, since it is not a LOD and therefore is always left enabled)
+			LODData[PrevLODIndex].Enable(Solver, Cloth, false);
+		}
+		if (LODIndex >= (int32)ECollisionDataType::LODs)
+		{
+			// Enable new LOD
+			LODData[LODIndex].Enable(Solver, Cloth, true);
+
+			if (PrevLODIndex == INDEX_NONE)
+			{
+				// Update initial state for collisions
+				LODData[LODIndex].ResetStartPose(Solver, Cloth);
+			}
+		}
+	}
+}
+
+TConstArrayView<TVector<float, 3>> FClothingSimulationCollider::GetCollisionTranslations(const FClothingSimulationSolver* Solver, const FClothingSimulationCloth* Cloth, ECollisionDataType CollisionDataType) const
+{
+	check(Solver);
+
+	const int32 Index = 
+		(CollisionDataType < ECollisionDataType::LODs) ? (int32)CollisionDataType :
+		(LODIndex >= (int32)ECollisionDataType::LODs) ? LODIndex : INDEX_NONE;
+
+	const TPair<const FClothingSimulationSolver*, const FClothingSimulationCloth*> SolverClothPair(Solver, Cloth);
+
+	return (Index != INDEX_NONE && LODData[Index].NumGeometries) ?
+		TConstArrayView<TVector<float, 3>>(Solver->GetCollisionParticleXs(LODData[Index].Offsets.FindChecked(SolverClothPair)), LODData[Index].NumGeometries) :
+		TConstArrayView<TVector<float, 3>>();
+}
+
+TConstArrayView<TRotation<float, 3>> FClothingSimulationCollider::GetCollisionRotations(const FClothingSimulationSolver* Solver, const FClothingSimulationCloth* Cloth, ECollisionDataType CollisionDataType) const
+{
+	check(Solver);
+	const int32 Index = 
+		(CollisionDataType < ECollisionDataType::LODs) ? (int32)CollisionDataType :
+		(LODIndex >= (int32)ECollisionDataType::LODs) ? LODIndex : INDEX_NONE;
+
+	const TPair<const FClothingSimulationSolver*, const FClothingSimulationCloth*> SolverClothPair(Solver, Cloth);
+
+	return (Index != INDEX_NONE && LODData[Index].NumGeometries) ?
+		TConstArrayView<TRotation<float, 3>>(Solver->GetCollisionParticleRs(LODData[Index].Offsets.FindChecked(SolverClothPair)), LODData[Index].NumGeometries) :
+		TConstArrayView<TRotation<float, 3>>();
+}
+
+TConstArrayView<TUniquePtr<FImplicitObject>> FClothingSimulationCollider::GetCollisionGeometries(const FClothingSimulationSolver* Solver, const FClothingSimulationCloth* Cloth, ECollisionDataType CollisionDataType) const
+{
+	check(Solver);
+	const int32 Index = 
+		(CollisionDataType < ECollisionDataType::LODs) ? (int32)CollisionDataType :
+		(LODIndex >= (int32)ECollisionDataType::LODs) ? LODIndex : INDEX_NONE;
+
+	const TPair<const FClothingSimulationSolver*, const FClothingSimulationCloth*> SolverClothPair(Solver, Cloth);
+
+	return (Index != INDEX_NONE && LODData[Index].NumGeometries) ?
+		TConstArrayView<TUniquePtr<FImplicitObject>>(Solver->GetCollisionGeometries(LODData[Index].Offsets.FindChecked(SolverClothPair)), LODData[Index].NumGeometries) :
+		TConstArrayView<TUniquePtr<FImplicitObject>>();
+}
diff --git a/Engine/Plugins/Experimental/ChaosCloth/Source/Chaos/Private/ChaosCloth/ChaosClothingSimulationFactory.cpp b/Engine/Plugins/Experimental/ChaosCloth/Source/Chaos/Private/ChaosCloth/ChaosClothingSimulationFactory.cpp
--- a/Engine/Plugins/Experimental/ChaosCloth/Source/Chaos/Private/ChaosCloth/ChaosClothingSimulationFactory.cpp
+++ b/Engine/Plugins/Experimental/ChaosCloth/Source/Chaos/Private/ChaosCloth/ChaosClothingSimulationFactory.cpp
@@ -8,7 +8,7 @@
 
 IClothingSimulation* UChaosClothingSimulationFactory::CreateSimulation()
 {
-	IClothingSimulation* Simulation = new Chaos::ClothingSimulation();
+	IClothingSimulation* Simulation = new Chaos::FClothingSimulation();
 	return Simulation;
 }
 
diff --git a/Engine/Plugins/Experimental/ChaosCloth/Source/Chaos/Private/ChaosCloth/ChaosClothingSimulationInteractor.cpp b/Engine/Plugins/Experimental/ChaosCloth/Source/Chaos/Private/ChaosCloth/ChaosClothingSimulationInteractor.cpp
--- a/Engine/Plugins/Experimental/ChaosCloth/Source/Chaos/Private/ChaosCloth/ChaosClothingSimulationInteractor.cpp
+++ b/Engine/Plugins/Experimental/ChaosCloth/Source/Chaos/Private/ChaosCloth/ChaosClothingSimulationInteractor.cpp
@@ -4,7 +4,7 @@
 
 void UChaosClothingSimulationInteractor::PhysicsAssetUpdated()
 {
-	Commands.Add(ChaosClothInteractorCommand::CreateLambda([](Chaos::ClothingSimulation* InSimulation, Chaos::ClothingSimulationContext* InContext)
+	Commands.Add(ChaosClothInteractorCommand::CreateLambda([](Chaos::FClothingSimulation* InSimulation, Chaos::FClothingSimulationContext* InContext)
 	{
 		InSimulation->RefreshPhysicsAsset();
 	}));
@@ -14,9 +14,9 @@ void UChaosClothingSimulationInteractor::PhysicsAssetUpdated()
 
 void UChaosClothingSimulationInteractor::ClothConfigUpdated()
 {
-	Commands.Add(ChaosClothInteractorCommand::CreateLambda([](Chaos::ClothingSimulation* InSimulation, Chaos::ClothingSimulationContext* InContext)
+	Commands.Add(ChaosClothInteractorCommand::CreateLambda([](Chaos::FClothingSimulation* InSimulation, Chaos::FClothingSimulationContext* InContext)
 	{
-		InSimulation->RefreshClothConfig();
+		InSimulation->RefreshClothConfig(InContext);
 	}));
 
 	MarkDirty();
@@ -27,8 +27,8 @@ void UChaosClothingSimulationInteractor::Sync(IClothingSimulation* InSimulation,
 	check(InSimulation);
 	check(InContext);
 
-	Chaos::ClothingSimulation* ChaosSim = static_cast<Chaos::ClothingSimulation*>(InSimulation);
-	Chaos::ClothingSimulationContext* ChaosContext = static_cast<Chaos::ClothingSimulationContext*>(InContext);
+	Chaos::FClothingSimulation* ChaosSim = static_cast<Chaos::FClothingSimulation*>(InSimulation);
+	Chaos::FClothingSimulationContext* ChaosContext = static_cast<Chaos::FClothingSimulationContext*>(InContext);
 
 	for(ChaosClothInteractorCommand& Command : Commands)
 	{
@@ -48,7 +48,7 @@ void UChaosClothingSimulationInteractor::Sync(IClothingSimulation* InSimulation,
 
 void UChaosClothingSimulationInteractor::SetAnimDriveSpringStiffness(float InStiffness)
 {
-	Commands.Add(ChaosClothInteractorCommand::CreateLambda([InStiffness](Chaos::ClothingSimulation* InSimulation, Chaos::ClothingSimulationContext* InContext)
+	Commands.Add(ChaosClothInteractorCommand::CreateLambda([InStiffness](Chaos::FClothingSimulation* InSimulation, Chaos::FClothingSimulationContext* InContext)
 	{
 		InSimulation->SetAnimDriveSpringStiffness(InStiffness);
 	}));
@@ -58,7 +58,7 @@ void UChaosClothingSimulationInteractor::SetAnimDriveSpringStiffness(float InSti
 
 void UChaosClothingSimulationInteractor::EnableGravityOverride(const FVector& InVector)
 {
-	Commands.Add(ChaosClothInteractorCommand::CreateLambda([InVector](Chaos::ClothingSimulation* InSimulation, Chaos::ClothingSimulationContext* InContext)
+	Commands.Add(ChaosClothInteractorCommand::CreateLambda([InVector](Chaos::FClothingSimulation* InSimulation, Chaos::FClothingSimulationContext* InContext)
 	{
 		InSimulation->SetGravityOverride(InVector);
 	}));
@@ -68,7 +68,7 @@ void UChaosClothingSimulationInteractor::EnableGravityOverride(const FVector& In
 
 void UChaosClothingSimulationInteractor::DisableGravityOverride()
 {
-	Commands.Add(ChaosClothInteractorCommand::CreateLambda([](Chaos::ClothingSimulation* InSimulation, Chaos::ClothingSimulationContext* InContext)
+	Commands.Add(ChaosClothInteractorCommand::CreateLambda([](Chaos::FClothingSimulation* InSimulation, Chaos::FClothingSimulationContext* InContext)
 	{
 		InSimulation->DisableGravityOverride();
 	}));
diff --git a/Engine/Plugins/Experimental/ChaosCloth/Source/Chaos/Private/ChaosCloth/ChaosClothingSimulationMesh.cpp b/Engine/Plugins/Experimental/ChaosCloth/Source/Chaos/Private/ChaosCloth/ChaosClothingSimulationMesh.cpp
new file mode 100644
--- /dev/null
+++ b/Engine/Plugins/Experimental/ChaosCloth/Source/Chaos/Private/ChaosCloth/ChaosClothingSimulationMesh.cpp
@@ -0,0 +1,316 @@
+// Copyright Epic Games, Inc. All Rights Reserved.
+#include "ChaosCloth/ChaosClothingSimulationMesh.h"
+#include "ChaosCloth/ChaosClothingSimulationSolver.h"
+#include "ChaosCloth/ChaosWeightMapTarget.h"
+#include "ChaosCloth/ChaosClothPrivate.h"
+#include "ClothingSimulation.h"
+#include "ClothingAsset.h"
+#include "Containers/ArrayView.h"
+
+DECLARE_CYCLE_STAT(TEXT("Chaos Cloth Skin Physics Mesh"), STAT_ChaosClothSkinPhysicsMesh, STATGROUP_Physics);
+DECLARE_CYCLE_STAT(TEXT("Chaos Cloth Wrap Deform Mesh"), STAT_ChaosClothWrapDeformMesh, STATGROUP_Physics);
+DECLARE_CYCLE_STAT(TEXT("Chaos Cloth Wrap Deform Cloth LOD"), STAT_ChaosClothWrapDeformClothLOD, STATGROUP_Physics);
+
+using namespace Chaos;
+
+FClothingSimulationMesh::FClothingSimulationMesh(const UClothingAssetCommon* InAsset, const USkeletalMeshComponent* InSkeletalMeshComponent)
+	: Asset(InAsset)
+	, SkeletalMeshComponent(InSkeletalMeshComponent)
+{
+}
+
+FClothingSimulationMesh::~FClothingSimulationMesh()
+{
+}
+
+int32 FClothingSimulationMesh::GetNumLODs() const
+{
+	return Asset ? Asset->LodData.Num() : 0;
+}
+
+int32 FClothingSimulationMesh::GetLODIndex() const
+{
+	int32 LODIndex = INDEX_NONE;
+
+	if (Asset && SkeletalMeshComponent)
+	{
+		if (const FClothingSimulationContextCommon* const Context = 
+			static_cast<const FClothingSimulationContextCommon*>(SkeletalMeshComponent->GetClothingSimulationContext()))
+		{
+			const int32 PredictedLODIndex = Context->PredictedLod;
+
+			// If PredictedLODIndex doesn't map to a valid LOD, we try higher LOD levels for a valid LOD.
+			// Asset might only have lod on LOD 1 and not 0, however if mesh doesn't force LOD to 1, 
+			// asset will not be assigned valid LOD index and will not generate sim data, breaking things.
+			for (int32 Index = PredictedLODIndex; Index < Asset->LodMap.Num(); ++Index)
+			{
+				const int32 MappedLODIndex = Asset->LodMap[Index];
+				if (Asset->LodData.IsValidIndex(MappedLODIndex))
+				{
+					LODIndex = MappedLODIndex;
+					break;
+				}
+			}
+		}
+	}
+	return LODIndex;
+}
+
+int32 FClothingSimulationMesh::GetNumPoints(int32 LODIndex) const
+{
+	return (Asset && Asset->LodData.IsValidIndex(LODIndex)) ?
+		Asset->LodData[LODIndex].PhysicalMeshData.Vertices.Num() :
+		0;
+}
+
+TConstArrayView<uint32> FClothingSimulationMesh::GetIndices(int32 LODIndex) const
+{
+	return (Asset && Asset->LodData.IsValidIndex(LODIndex)) ?
+		TConstArrayView<uint32>(Asset->LodData[LODIndex].PhysicalMeshData.Indices) :
+		TConstArrayView<uint32>();
+}
+
+TArray<TConstArrayView<float>> FClothingSimulationMesh::GetWeightMaps(int32 LODIndex) const
+{
+	TArray<TConstArrayView<float>> WeightMaps;
+	if (Asset && Asset->LodData.IsValidIndex(LODIndex))
+	{
+		const FClothLODDataCommon& ClothLODData = Asset->LodData[LODIndex];
+		const FClothPhysicalMeshData& ClothPhysicalMeshData = ClothLODData.PhysicalMeshData;
+
+		const UEnum* const ChaosWeightMapTargetEnum = StaticEnum<EChaosWeightMapTarget>();
+		const int32 NumWeightMaps = (int32)ChaosWeightMapTargetEnum->GetMaxEnumValue() + 1;
+
+		WeightMaps.SetNum(NumWeightMaps);
+
+		for (int32 EnumIndex = 0; EnumIndex < ChaosWeightMapTargetEnum->NumEnums(); ++EnumIndex)
+		{
+			const int32 TargetIndex = (int32)ChaosWeightMapTargetEnum->GetValueByIndex(EnumIndex);
+			if (const FPointWeightMap* const WeightMap = ClothPhysicalMeshData.FindWeightMap(TargetIndex))
+			{
+				WeightMaps[TargetIndex] = WeightMap->Values;
+			}
+		}
+	}
+	return WeightMaps;
+}
+
+int32 FClothingSimulationMesh::GetReferenceBoneIndex() const
+{
+	return Asset ? Asset->ReferenceBoneIndex : INDEX_NONE;
+}
+
+TRigidTransform<float, 3> Chaos::FClothingSimulationMesh::GetReferenceBoneTransform() const
+{
+	if (SkeletalMeshComponent)
+	{
+		if (const FClothingSimulationContextCommon* const Context =
+			static_cast<const FClothingSimulationContextCommon*>(SkeletalMeshComponent->GetClothingSimulationContext()))
+		{
+			const int32 ReferenceBoneIndex = GetReferenceBoneIndex();
+			const TArray<FTransform>& BoneTransforms = Context->BoneTransforms;
+
+			return BoneTransforms.IsValidIndex(ReferenceBoneIndex) ?
+				BoneTransforms[ReferenceBoneIndex] * Context->ComponentToWorld :
+				Context->ComponentToWorld;
+		}
+	}
+	return TRigidTransform<float, 3>::Identity;
+}
+
+bool FClothingSimulationMesh::WrapDeformLOD(
+	int32 PrevLODIndex,
+	int32 LODIndex,
+	const TVector<float, 3>* Normals,
+	const TVector<float, 3>* Positions,
+	TVector<float, 3>* OutPositions) const
+{
+	SCOPE_CYCLE_COUNTER(STAT_ChaosClothWrapDeformMesh);
+
+	const int32 NumLODsPassed = FMath::Abs(LODIndex - PrevLODIndex);
+	if (NumLODsPassed != 1 || !Asset || !Asset->LodData.IsValidIndex(PrevLODIndex) || !Asset->LodData.IsValidIndex(LODIndex))
+	{
+		return false;
+	}
+
+	const FClothLODDataCommon& LODData = Asset->LodData[LODIndex];
+	const int32 NumPoints = LODData.PhysicalMeshData.Vertices.Num();
+	const TArray<FMeshToMeshVertData>& SkinData = (PrevLODIndex < LODIndex) ?
+		LODData.TransitionUpSkinData :
+		LODData.TransitionDownSkinData;
+
+	for (int32 Index = 0; Index < NumPoints; ++Index)  // TODO: Profile for parallel for
+	{
+		const FMeshToMeshVertData& VertData = SkinData[Index];
+
+		const int32 VertIndex0 = (int32)VertData.SourceMeshVertIndices[0];  // Note: The source is uint16. Watch out for large mesh sections!
+		const int32 VertIndex1 = (int32)VertData.SourceMeshVertIndices[1];
+		const int32 VertIndex2 = (int32)VertData.SourceMeshVertIndices[2];
+
+		OutPositions[Index] = 
+			Positions[VertIndex0] * VertData.PositionBaryCoordsAndDist.X + Normals[VertIndex0] * VertData.PositionBaryCoordsAndDist.W +
+			Positions[VertIndex1] * VertData.PositionBaryCoordsAndDist.Y + Normals[VertIndex1] * VertData.PositionBaryCoordsAndDist.W +
+			Positions[VertIndex2] * VertData.PositionBaryCoordsAndDist.Z + Normals[VertIndex2] * VertData.PositionBaryCoordsAndDist.W;
+	}
+
+	return true;
+}
+
+bool FClothingSimulationMesh::WrapDeformLOD(
+	int32 PrevLODIndex,
+	int32 LODIndex,
+	const TVector<float, 3>* Normals,
+	const TVector<float, 3>* Positions,
+	const TVector<float, 3>* Velocities,
+	TVector<float, 3>* OutPositions0,
+	TVector<float, 3>* OutPositions1,
+	TVector<float, 3>* OutVelocities) const
+{
+	SCOPE_CYCLE_COUNTER(STAT_ChaosClothWrapDeformClothLOD);
+
+	const int32 NumLODsPassed = FMath::Abs(LODIndex - PrevLODIndex);
+	if (NumLODsPassed != 1 || !Asset || !Asset->LodData.IsValidIndex(PrevLODIndex) || !Asset->LodData.IsValidIndex(LODIndex))
+	{
+		return false;
+	}
+
+	const FClothLODDataCommon& LODData = Asset->LodData[LODIndex];
+	const int32 NumPoints = LODData.PhysicalMeshData.Vertices.Num();
+	const TArray<FMeshToMeshVertData>& SkinData = (PrevLODIndex < LODIndex) ?
+		LODData.TransitionUpSkinData :
+		LODData.TransitionDownSkinData;
+
+	for (int32 Index = 0; Index < NumPoints; ++Index)  // TODO: Profile for parallel for
+	{
+		const FMeshToMeshVertData& VertData = SkinData[Index];
+
+		const int32 VertIndex0 = (int32)VertData.SourceMeshVertIndices[0];  // Note: The source is uint16. Watch out for large mesh sections!
+		const int32 VertIndex1 = (int32)VertData.SourceMeshVertIndices[1];
+		const int32 VertIndex2 = (int32)VertData.SourceMeshVertIndices[2];
+
+		OutPositions0[Index] = OutPositions1[Index] =
+			Positions[VertIndex0] * VertData.PositionBaryCoordsAndDist.X + Normals[VertIndex0] * VertData.PositionBaryCoordsAndDist.W +
+			Positions[VertIndex1] * VertData.PositionBaryCoordsAndDist.Y + Normals[VertIndex1] * VertData.PositionBaryCoordsAndDist.W +
+			Positions[VertIndex2] * VertData.PositionBaryCoordsAndDist.Z + Normals[VertIndex2] * VertData.PositionBaryCoordsAndDist.W;
+
+		OutVelocities[Index] = 
+			Velocities[VertIndex0] * VertData.PositionBaryCoordsAndDist.X +
+			Velocities[VertIndex1] * VertData.PositionBaryCoordsAndDist.Y +
+			Velocities[VertIndex2] * VertData.PositionBaryCoordsAndDist.Z;
+	}
+
+	return true;
+}
+
+// Inline function used to force the unrolling of the skinning loop
+FORCEINLINE static void AddInfluence(FVector& OutPosition, FVector& OutNormal, const FVector& RefParticle, const FVector& RefNormal, const FMatrix& BoneMatrix, const float Weight)
+{
+	OutPosition += BoneMatrix.TransformPosition(RefParticle) * Weight;
+	OutNormal += BoneMatrix.TransformVector(RefNormal) * Weight;
+}
+
+void FClothingSimulationMesh::SkinPhysicsMesh(int32 LODIndex, const TVector<float, 3>& LocalSpaceLocation, TVector<float, 3>* OutPositions, TVector<float, 3>* OutNormals) const
+{
+	SCOPE_CYCLE_COUNTER(STAT_ChaosClothSkinPhysicsMesh);
+
+	check(Asset && Asset->LodData.IsValidIndex(LODIndex));
+	const FClothPhysicalMeshData& PhysicalMeshData = Asset->LodData[LODIndex].PhysicalMeshData;
+	UE_CLOG(PhysicalMeshData.MaxBoneWeights > 12, LogChaosCloth, Warning, TEXT("The cloth physics mesh skinning code can't cope with more than 12 bone influences."));
+
+	const uint32 NumPoints = PhysicalMeshData.Vertices.Num();
+
+	check(SkeletalMeshComponent && SkeletalMeshComponent->GetClothingSimulationContext());
+	const FClothingSimulationContextCommon* const Context = static_cast<const FClothingSimulationContextCommon*>(SkeletalMeshComponent->GetClothingSimulationContext());
+	FTransform ComponentToLocalSpace = Context->ComponentToWorld;
+	ComponentToLocalSpace.AddToTranslation(-LocalSpaceLocation);
+
+	// Zero out positions & normals
+	FMemory::Memzero((uint8*)OutPositions, NumPoints * sizeof(TVector<float, 3>));  // PS4 performance note: It is faster to zero the memory first
+	FMemory::Memzero((uint8*)OutNormals, NumPoints * sizeof(TVector<float, 3>));    // instead of changing this function to work with uninitialized memory
+
+	const int32* const RESTRICT BoneMap = Asset->UsedBoneIndices.GetData();
+	const FMatrix* const RESTRICT BoneMatrices = Context->RefToLocals.GetData();
+		
+	static const uint32 MinParallelVertices = 500;  // 500 seems to be the lowest threshold still giving gains even on profiled assets that are only using a small number of influences
+
+	ParallelFor(NumPoints, [&PhysicalMeshData, &ComponentToLocalSpace, BoneMap, BoneMatrices, &OutPositions, &OutNormals](uint32 VertIndex)
+	{
+		const uint16* const RESTRICT BoneIndices = PhysicalMeshData.BoneData[VertIndex].BoneIndices;
+		const float* const RESTRICT BoneWeights = PhysicalMeshData.BoneData[VertIndex].BoneWeights;
+
+		// WARNING - HORRIBLE UNROLLED LOOP + JUMP TABLE BELOW
+		// done this way because this is a pretty tight and performance critical loop. essentially
+		// rather than checking each influence we can just jump into this switch and fall through
+		// everything to compose the final skinned data
+		const TVector<float, 3>& RefParticle = PhysicalMeshData.Vertices[VertIndex];
+		const TVector<float, 3>& RefNormal = PhysicalMeshData.Normals[VertIndex];
+		TVector<float, 3>& OutPosition = OutPositions[VertIndex];
+		TVector<float, 3>& OutNormal = OutNormals[VertIndex];
+		switch (PhysicalMeshData.BoneData[VertIndex].NumInfluences)
+		{
+		case 12: AddInfluence(OutPosition, OutNormal, RefParticle, RefNormal, BoneMatrices[BoneMap[BoneIndices[11]]], BoneWeights[11]);  // Intentional fall through
+		case 11: AddInfluence(OutPosition, OutNormal, RefParticle, RefNormal, BoneMatrices[BoneMap[BoneIndices[10]]], BoneWeights[10]);  // Intentional fall through
+		case 10: AddInfluence(OutPosition, OutNormal, RefParticle, RefNormal, BoneMatrices[BoneMap[BoneIndices[ 9]]], BoneWeights[ 9]);  // Intentional fall through
+		case  9: AddInfluence(OutPosition, OutNormal, RefParticle, RefNormal, BoneMatrices[BoneMap[BoneIndices[ 8]]], BoneWeights[ 8]);  // Intentional fall through
+		case  8: AddInfluence(OutPosition, OutNormal, RefParticle, RefNormal, BoneMatrices[BoneMap[BoneIndices[ 7]]], BoneWeights[ 7]);  // Intentional fall through
+		case  7: AddInfluence(OutPosition, OutNormal, RefParticle, RefNormal, BoneMatrices[BoneMap[BoneIndices[ 6]]], BoneWeights[ 6]);  // Intentional fall through
+		case  6: AddInfluence(OutPosition, OutNormal, RefParticle, RefNormal, BoneMatrices[BoneMap[BoneIndices[ 5]]], BoneWeights[ 5]);  // Intentional fall through
+		case  5: AddInfluence(OutPosition, OutNormal, RefParticle, RefNormal, BoneMatrices[BoneMap[BoneIndices[ 4]]], BoneWeights[ 4]);  // Intentional fall through
+		case  4: AddInfluence(OutPosition, OutNormal, RefParticle, RefNormal, BoneMatrices[BoneMap[BoneIndices[ 3]]], BoneWeights[ 3]);  // Intentional fall through
+		case  3: AddInfluence(OutPosition, OutNormal, RefParticle, RefNormal, BoneMatrices[BoneMap[BoneIndices[ 2]]], BoneWeights[ 2]);  // Intentional fall through
+		case  2: AddInfluence(OutPosition, OutNormal, RefParticle, RefNormal, BoneMatrices[BoneMap[BoneIndices[ 1]]], BoneWeights[ 1]);  // Intentional fall through
+		case  1: AddInfluence(OutPosition, OutNormal, RefParticle, RefNormal, BoneMatrices[BoneMap[BoneIndices[ 0]]], BoneWeights[ 0]);  // Intentional fall through
+		default: break;
+		}
+
+		OutPosition = ComponentToLocalSpace.TransformPosition(OutPosition);
+		OutNormal = ComponentToLocalSpace.TransformVector(OutNormal);
+		OutNormal.Normalize();
+	}, NumPoints > MinParallelVertices ? EParallelForFlags::None : EParallelForFlags::ForceSingleThread);
+}
+
+void FClothingSimulationMesh::Update(
+	FClothingSimulationSolver* Solver,
+	int32 PrevLODIndex,
+	int32 LODIndex,
+	int32 PrevOffset,
+	int32 Offset)
+{
+	check(Solver);
+
+	// Exit if any inputs are missing or not ready, and if the LOD is invalid
+	if (!Asset || !Asset->LodData.IsValidIndex(LODIndex) || !SkeletalMeshComponent || !SkeletalMeshComponent->GetClothingSimulationContext())
+	{
+		return;
+	}
+
+	// Skin current LOD positions
+	const TVector<float, 3>& LocalSpaceLocation = Solver->GetLocalSpaceLocation();
+	TVector<float, 3>* const OutPositions = Solver->GetAnimationPositions(Offset);
+	TVector<float, 3>* const OutNormals = Solver->GetAnimationNormals(Offset);
+	
+	SkinPhysicsMesh(LODIndex, LocalSpaceLocation, OutPositions, OutNormals);
+
+	// Update old positions after LOD Switching
+	if (LODIndex != PrevLODIndex)
+	{
+		// TODO: Using the more accurate skinning method here would require double buffering the context at the skeletal mesh level
+		const TVector<float, 3>* const SrcWrapNormals = Solver->GetAnimationNormals(PrevOffset);  // No need to keep an old normals array around, since the LOD has just changed
+		const TVector<float, 3>* const SrcWrapPositions = Solver->GetOldAnimationPositions(PrevOffset);
+		TVector<float, 3>* const OutOldPositions = Solver->GetOldAnimationPositions(Offset);
+
+		const bool bValidWrap = WrapDeformLOD(PrevLODIndex, LODIndex, SrcWrapNormals, SrcWrapPositions, OutOldPositions);
+	
+		if (!bValidWrap)
+		{
+			// The previous LOD is invalid, reset old positions with the new LOD
+			const FClothPhysicalMeshData& PhysicalMeshData = Asset->LodData[LODIndex].PhysicalMeshData;
+			const int32 NumPoints = PhysicalMeshData.Vertices.Num();
+
+			for (int32 Index = 0; Index < NumPoints; ++Index)
+			{
+				OutOldPositions[Index] = OutPositions[Index];
+			}
+		}
+	}
+}
