Chaos Cloth - LOD/graph node refactor.

diff --git a/Engine/Plugins/Experimental/ChaosCloth/Source/Chaos/Private/ChaosCloth/ChaosClothingSimulationSolver.cpp b/Engine/Plugins/Experimental/ChaosCloth/Source/Chaos/Private/ChaosCloth/ChaosClothingSimulationSolver.cpp
new file mode 100644
--- /dev/null
+++ b/Engine/Plugins/Experimental/ChaosCloth/Source/Chaos/Private/ChaosCloth/ChaosClothingSimulationSolver.cpp
@@ -0,0 +1,775 @@
+// Copyright Epic Games, Inc. All Rights Reserved.
+#include "ChaosCloth/ChaosClothingSimulationSolver.h"
+#include "ChaosCloth/ChaosClothingSimulationCloth.h"
+#include "ChaosCloth/ChaosClothingSimulationCollider.h"
+#include "ChaosCloth/ChaosClothingSimulationMesh.h"
+#include "ChaosCloth/ChaosClothingSimulation.h"
+#include "ChaosCloth/ChaosClothPrivate.h"
+
+#include "Chaos/PBDEvolution.h"
+
+using namespace Chaos;
+
+DECLARE_CYCLE_STAT(TEXT("Chaos Cloth Solver Update"), STAT_ChaosClothSolverUpdate, STATGROUP_ChaosCloth);
+DECLARE_CYCLE_STAT(TEXT("Chaos Cloth Solver Update Colliders"), STAT_ChaosClothSolverUpdateColliders, STATGROUP_ChaosCloth);
+DECLARE_CYCLE_STAT(TEXT("Chaos Cloth Solver Update Cloths"), STAT_ChaosClothSolverUpdateCloths, STATGROUP_ChaosCloth);
+DECLARE_CYCLE_STAT(TEXT("Chaos Cloth Solver Update Pre Solver Step"), STAT_ChaosClothSolverUpdatePreSolverStep, STATGROUP_ChaosCloth);
+DECLARE_CYCLE_STAT(TEXT("Chaos Cloth Solver Update Solver Step"), STAT_ChaosClothSolverUpdateSolverStep, STATGROUP_ChaosCloth);
+DECLARE_CYCLE_STAT(TEXT("Chaos Cloth Solver Update Post Solver Step"), STAT_ChaosClothSolverUpdatePostSolverStep, STATGROUP_ChaosCloth);
+DECLARE_CYCLE_STAT(TEXT("Chaos Cloth Solver Calculate Bounds"), STAT_ChaosClothSolverCalculateBounds, STATGROUP_ChaosCloth);
+
+namespace ChaosClothingSimulationSolverDefault
+{
+	static const TVector<float, 3> Gravity(0.f, 0.f, -980.665f);  // cm/s^2
+	static const TVector<float, 3> WindVelocity(0.f);
+	static const float WindFluidDensity = 1.225e-6f;  // kg/cm^3
+	static const int32 NumIterations = 1;
+	static const int32 NumSubsteps = 1;
+	static const float SelfCollisionThickness = 2.f;
+	static const float CollisionThickness = 1.2f;
+	static const float FrictionCoefficient = 0.2f;
+	static const float DampingCoefficient = 0.01f;
+}
+
+namespace ChaosClothingSimulationSolverConstant
+{
+	static const float WorldScale = 100.f;  // World is in cm, but values like wind speed and density are in SI unit and relates to m.
+	static const float StartDeltaTime = 1.f / 30.f;  // Initialize filtered timestep at 30fps
+}
+
+FClothingSimulationSolver::FClothingSimulationSolver()
+	: OldLocalSpaceLocation(0.f)
+	, LocalSpaceLocation(0.f)
+	, Time(0)
+	, DeltaTime(ChaosClothingSimulationSolverConstant::StartDeltaTime)
+	, NumIterations(ChaosClothingSimulationSolverDefault::NumIterations)
+	, NumSubsteps(ChaosClothingSimulationSolverDefault::NumSubsteps)
+	, CollisionParticlesOffset(0)
+	, CollisionParticlesSize(0)
+	, bAreCollisionParticlesDirty(false)
+	, Gravity(ChaosClothingSimulationSolverDefault::Gravity)
+	, WindVelocity(ChaosClothingSimulationSolverDefault::WindVelocity)
+	, WindFluidDensity(ChaosClothingSimulationSolverDefault::WindFluidDensity)
+	, bIsClothGravityOverrideEnabled(false)
+{
+	TPBDParticles<float, 3> LocalParticles;
+	TKinematicGeometryClothParticles<float, 3> TRigidParticles;
+	Evolution.Reset(
+		new TPBDEvolution<float, 3>(
+			MoveTemp(LocalParticles),
+			MoveTemp(TRigidParticles),
+			{}, // CollisionTriangles
+			ChaosClothingSimulationSolverDefault::NumIterations,
+			ChaosClothingSimulationSolverDefault::CollisionThickness,
+			ChaosClothingSimulationSolverDefault::SelfCollisionThickness,
+			ChaosClothingSimulationSolverDefault::FrictionCoefficient,
+			ChaosClothingSimulationSolverDefault::DampingCoefficient));
+
+	// Add simulation groups arrays
+	Evolution->AddArray(&PreSimulationTransforms);
+
+	Evolution->Particles().AddArray(&Normals);
+	Evolution->Particles().AddArray(&OldAnimationPositions);
+	Evolution->Particles().AddArray(&AnimationPositions);
+	Evolution->Particles().AddArray(&AnimationNormals);
+
+	Evolution->CollisionParticles().AddArray(&CollisionBoneIndices);
+	Evolution->CollisionParticles().AddArray(&CollisionBaseTransforms);
+	Evolution->CollisionParticles().AddArray(&OldCollisionTransforms);
+	Evolution->CollisionParticles().AddArray(&CollisionTransforms);
+
+	Evolution->SetKinematicUpdateFunction(
+		[this](TPBDParticles<float, 3>& ParticlesInput, const float Dt, const float LocalTime, const int32 Index)
+		{
+			if (!OldAnimationPositions.IsValidIndex(Index) || ParticlesInput.InvM(Index) > 0)
+			{
+				return;
+			}
+			const float Alpha = (LocalTime - Time) / DeltaTime;
+			ParticlesInput.X(Index) = Alpha * AnimationPositions[Index] + (1.f - Alpha) * OldAnimationPositions[Index];
+		});
+
+	Evolution->SetCollisionKinematicUpdateFunction(
+		[this](TKinematicGeometryClothParticles<float, 3>& ParticlesInput, const float Dt, const float LocalTime, const int32 Index)
+		{
+			checkSlow(Dt > SMALL_NUMBER && DeltaTime > SMALL_NUMBER);
+			const float Alpha = (LocalTime - Time) / DeltaTime;
+			const TVector<float, 3> NewX =
+				Alpha * CollisionTransforms[Index].GetTranslation() + (1.f - Alpha) * OldCollisionTransforms[Index].GetTranslation();
+			ParticlesInput.V(Index) = (NewX - ParticlesInput.X(Index)) / Dt;
+			ParticlesInput.X(Index) = NewX;
+			TRotation<float, 3> NewR = FQuat::Slerp(OldCollisionTransforms[Index].GetRotation(), CollisionTransforms[Index].GetRotation(), Alpha);
+			TRotation<float, 3> Delta = NewR * ParticlesInput.R(Index).Inverse();
+			TVector<float, 3> Axis;
+			float Angle;
+			Delta.ToAxisAndAngle(Axis, Angle);
+			ParticlesInput.W(Index) = Axis * Angle / Dt;
+			ParticlesInput.R(Index) = NewR;
+		});
+}
+
+FClothingSimulationSolver::~FClothingSimulationSolver()
+{
+}
+
+void FClothingSimulationSolver::SetCloths(TArray<FClothingSimulationCloth*>&& InCloths)
+{
+	// Reset solver collisions so that there is never any dynamic collision particles below cloth's own ones
+	Evolution->ResetCollisionParticles();
+
+	// Reset cloth particles and associated elements
+	ResetParticles();
+
+	// Update array
+	Cloths = MoveTemp(InCloths);
+
+	// Add the new cloths' particles
+	for (FClothingSimulationCloth* const Cloth : Cloths)
+	{
+		check(Cloth);
+
+		// Add the cloth's particles
+		Cloth->Add(this);
+
+		// Set initial state
+		Cloth->Update(this);
+	}
+
+	// Update solver collider's offset
+	CollisionParticlesOffset = Evolution->CollisionParticles().Size();
+	CollisionParticlesSize = 0;
+}
+
+void FClothingSimulationSolver::AddCloth(FClothingSimulationCloth* InCloth)
+{
+	check(InCloth);
+
+	if (Cloths.Find(InCloth) != INDEX_NONE)
+	{
+		return;
+	}
+
+	// Reset solver collisions so that there is never any dynamic collision particles below cloth's own ones
+	Evolution->ResetCollisionParticles(CollisionParticlesOffset);
+
+	// Add the cloth to the solver update array
+	Cloths.Emplace(InCloth);
+
+	// Add the cloth's particles
+	InCloth->Add(this);
+
+	// Set initial state
+	InCloth->Update(this);
+
+	// Update solver collider's offset
+	CollisionParticlesOffset = Evolution->CollisionParticles().Size();
+	CollisionParticlesSize = 0;
+}
+
+void FClothingSimulationSolver::RemoveCloth(FClothingSimulationCloth* InCloth)
+{
+	if (Cloths.Find(InCloth) == INDEX_NONE)
+	{
+		return;
+	}
+
+	// Reset solver collisions so that there is never any dynamic collision particles below cloth's own ones
+	Evolution->ResetCollisionParticles();
+
+	// Reset cloth particles and associated elements
+	ResetParticles();
+
+	// Remove reference to this solver
+	InCloth->Remove(this);
+
+	// Remove collider from array
+	Cloths.RemoveSwap(InCloth);
+
+	// Re-add the remaining cloths' particles
+	for (FClothingSimulationCloth* const Cloth : Cloths)
+	{
+		// Add the cloth's particles
+		Cloth->Add(this);
+
+		// Set initial state
+		Cloth->Update(this);
+	}
+
+	// Update solver collider's offset
+	CollisionParticlesOffset = Evolution->CollisionParticles().Size();
+	CollisionParticlesSize = 0;
+}
+
+void FClothingSimulationSolver::RemoveCloths()
+{
+	// Remove all cloths from array
+	for (FClothingSimulationCloth* const Cloth : Cloths)
+	{
+		Cloth->Remove(this);
+	}
+	Cloths.Reset();
+
+	// Reset solver collisions
+	Evolution->ResetCollisionParticles();
+
+	// Reset cloth particles and associated elements
+	ResetParticles();
+
+	// Update solver collider's offset
+	CollisionParticlesOffset = 0;
+	CollisionParticlesSize = 0;
+}
+
+void FClothingSimulationSolver::ResetParticles()
+{
+	Evolution->ResetParticles();
+	Evolution->ResetConstraintRules();
+	Evolution->ResetSelfCollision();
+	ClothsConstraints.Reset();
+}
+
+int32 FClothingSimulationSolver::AddParticles(int32 NumParticles, uint32 GroupId)
+{
+	if (!NumParticles)
+	{
+		return INDEX_NONE;
+	}
+	const int32 Offset = Evolution->AddParticleRange(NumParticles, GroupId, /*bActivate =*/ false);
+
+	// Add an empty constraints container for this range
+	ClothsConstraints.Add(Offset).Initialize(Evolution.Get(), AnimationPositions, AnimationNormals, Offset, NumParticles);
+
+	// Always starts with particles disabled
+	EnableParticles(Offset, false);
+
+	return Offset;
+}
+
+void FClothingSimulationSolver::EnableParticles(int32 Offset, bool bEnable)
+{
+	Evolution->ActivateParticleRange(Offset, bEnable);
+	GetClothConstraints(Offset).Enable(bEnable);
+}
+
+const TVector<float, 3>* FClothingSimulationSolver::GetParticlePs(int32 Offset) const
+{
+	return &Evolution->Particles().P(Offset);
+}
+
+TVector<float, 3>* FClothingSimulationSolver::GetParticlePs(int32 Offset)
+{
+	return &Evolution->Particles().P(Offset);
+}
+
+const TVector<float, 3>* FClothingSimulationSolver::GetParticleXs(int32 Offset) const
+{
+	return &Evolution->Particles().X(Offset);
+}
+
+TVector<float, 3>* FClothingSimulationSolver::GetParticleXs(int32 Offset)
+{
+	return &Evolution->Particles().X(Offset);
+}
+
+const TVector<float, 3>* FClothingSimulationSolver::GetParticleVs(int32 Offset) const
+{
+	return &Evolution->Particles().V(Offset);
+}
+
+TVector<float, 3>* FClothingSimulationSolver::GetParticleVs(int32 Offset)
+{
+	return &Evolution->Particles().V(Offset);
+}
+
+const float* FClothingSimulationSolver::GetParticleInvMasses(int32 Offset) const
+{
+	return &Evolution->Particles().InvM(Offset);
+}
+
+void FClothingSimulationSolver::SetColliders(TArray<FClothingSimulationCollider*>&& InColliders)
+{
+	// Update array
+	Colliders = MoveTemp(InColliders);
+
+	// Colliders have changed, must prevent from trying reusing the ranges
+	Evolution->ResetCollisionParticles(CollisionParticlesOffset);
+	CollisionParticlesSize = 0;
+}
+
+void FClothingSimulationSolver::AddCollider(FClothingSimulationCollider* InCollider)
+{
+	check(InCollider);
+
+	if (Colliders.Find(InCollider) != INDEX_NONE)
+	{
+		return;
+	}
+
+	// Add the collider to the solver update array
+	Colliders.Emplace(InCollider);
+}
+
+void FClothingSimulationSolver::RemoveCollider(FClothingSimulationCollider* InCollider)
+{
+	if (Colliders.Find(InCollider) == INDEX_NONE)
+	{
+		return;
+	}
+
+	// Remove collider from array
+	Colliders.RemoveSwap(InCollider);
+
+	// Colliders have changed, must prevent from trying reusing the ranges
+	Evolution->ResetCollisionParticles(CollisionParticlesOffset);
+	CollisionParticlesSize = 0;
+}
+
+void FClothingSimulationSolver::RemoveColliders()
+{
+	// Remove all colliders from array
+	Colliders.Reset();
+
+	// Colliders have changed, must prevent from trying reusing the ranges
+	Evolution->ResetCollisionParticles(CollisionParticlesOffset);
+	CollisionParticlesSize = 0;
+}
+
+int32 FClothingSimulationSolver::AddCollisionParticles(int32 NumCollisionParticles, uint32 GroupId, int32 RecycledOffset)
+{
+	// Try reusing the particle range so that solver collisions can be added/removed dynamically
+	// If it doesn't match then remove all ranges above the given offset to start again.
+	// This rely on the assumption that these ranges are added again in the same order
+	// by the solver at every update
+	if (RecycledOffset == CollisionParticlesOffset + CollisionParticlesSize)
+	{
+		CollisionParticlesSize += NumCollisionParticles;
+
+		// Check that the range still exists
+		if (CollisionParticlesOffset + CollisionParticlesSize <= (int32)Evolution->CollisionParticles().Size() &&  // Check first that the range hasn't been reset
+			NumCollisionParticles == Evolution->GetCollisionParticleRangeSize(RecycledOffset))  // This will assert if range has been reset
+		{
+			// Always starts with particles disabled
+			EnableCollisionParticles(RecycledOffset, false);
+			return RecycledOffset;
+		}
+		// Size has changed. must reset this collision range (and all of those following up) and reallocate some new particles
+		Evolution->ResetCollisionParticles(RecycledOffset);
+	}
+
+	if (!NumCollisionParticles)
+	{
+		return INDEX_NONE;
+	}
+
+	const int32 Offset = Evolution->AddCollisionParticleRange(NumCollisionParticles, GroupId, /*bActivate =*/ false);
+
+	// Always initialize the collision particle's transforms as otherwise setting the geometry will get NaNs detected during the bounding box updates
+	TRotation<float, 3>* const Rs = GetCollisionParticleRs(Offset);
+	TVector<float, 3>* const Xs = GetCollisionParticleXs(Offset);
+
+	for (int32 Index = 0; Index < NumCollisionParticles; ++Index)
+	{
+		Xs[Index] = TVector<float, 3>(0.f);
+		Rs[Index] = TRotation<float, 3>::FromIdentity();
+	}
+
+	// Always starts with particles disabled
+	EnableCollisionParticles(Offset, false);
+
+	// Any change in particle allocation should trigger an initial state update
+	bAreCollisionParticlesDirty = true;  // TODO: Dirtying collision initial state update only based on the number of particles is not bug proof, but hashing the collision geometries is a bit heavy handed too
+
+	return Offset;
+}
+
+void FClothingSimulationSolver::EnableCollisionParticles(int32 Offset, bool bEnable)
+{
+	Evolution->ActivateCollisionParticleRange(Offset, bEnable);
+}
+
+const TVector<float, 3>* FClothingSimulationSolver::GetCollisionParticleXs(int32 Offset) const
+{
+	return &Evolution->CollisionParticles().X(Offset);
+}
+
+TVector<float, 3>* FClothingSimulationSolver::GetCollisionParticleXs(int32 Offset)
+{
+	return &Evolution->CollisionParticles().X(Offset);
+}
+
+const TRotation<float, 3>* FClothingSimulationSolver::GetCollisionParticleRs(int32 Offset) const
+{
+	return &Evolution->CollisionParticles().R(Offset);
+}
+
+TRotation<float, 3>* FClothingSimulationSolver::GetCollisionParticleRs(int32 Offset)
+{
+	return &Evolution->CollisionParticles().R(Offset);
+}
+
+void FClothingSimulationSolver::SetCollisionGeometry(int32 Offset, int32 Index, TUniquePtr<FImplicitObject>&& Geometry)
+{
+	Evolution->CollisionParticles().SetDynamicGeometry(Offset + Index, MoveTemp(Geometry));
+}
+
+const TUniquePtr<FImplicitObject>* FClothingSimulationSolver::GetCollisionGeometries(int32 Offset) const
+{
+	return &Evolution->CollisionParticles().DynamicGeometry(Offset);
+}
+
+void FClothingSimulationSolver::SetParticleMassUniform(int32 Offset, float UniformMass, float MinPerParticleMass, const TTriangleMesh<float>& Mesh, const TFunctionRef<bool(int32)>& KinematicPredicate)
+{
+	// Retrieve the particle block size
+	const int32 Size = Evolution->GetParticleRangeSize(Offset);
+
+	// Set mass from uniform mass
+	const TSet<int32> Vertices = Mesh.GetVertices();
+	TPBDParticles<float, 3>& Particles = Evolution->Particles();
+	for (int32 Index = Offset; Index < Offset + Size; ++Index)
+	{
+		Particles.M(Index) = Vertices.Contains(Index) ? UniformMass : 0.f;
+	}
+
+	// Clamp and enslave
+	ParticleMassClampAndEnslave(Offset, Size, MinPerParticleMass, KinematicPredicate);
+}
+
+void FClothingSimulationSolver::SetParticleMassFromTotalMass(int32 Offset, float TotalMass, float MinPerParticleMass, const TTriangleMesh<float>& Mesh, const TFunctionRef<bool(int32)>& KinematicPredicate)
+{
+	// Retrieve the particle block size
+	const int32 Size = Evolution->GetParticleRangeSize(Offset);
+
+	// Set mass per area
+	const float TotalArea = SetParticleMassPerArea(Offset, Size, Mesh);
+
+	// Find density
+	const float Density = TotalArea > 0.f ? TotalMass / TotalArea : 1.f;
+
+	// Update mass from mesh and density
+	ParticleMassUpdateDensity(Mesh, Density);
+
+	// Clamp and enslave
+	ParticleMassClampAndEnslave(Offset, Size, MinPerParticleMass, KinematicPredicate);
+}
+
+void FClothingSimulationSolver::SetParticleMassFromDensity(int32 Offset, float Density, float MinPerParticleMass, const TTriangleMesh<float>& Mesh, const TFunctionRef<bool(int32)>& KinematicPredicate)
+{
+	// Retrieve the particle block size
+	const int32 Size = Evolution->GetParticleRangeSize(Offset);
+
+	// Set mass per area
+	const float TotalArea = SetParticleMassPerArea(Offset, Size, Mesh);
+
+	// Set density from cm2 to m2
+	Density /= FMath::Square(ChaosClothingSimulationSolverConstant::WorldScale);
+
+	// Update mass from mesh and density
+	ParticleMassUpdateDensity(Mesh, Density);
+
+	// Clamp and enslave
+	ParticleMassClampAndEnslave(Offset, Size, MinPerParticleMass, KinematicPredicate);
+}
+
+void FClothingSimulationSolver::SetReferenceVelocityScale(
+	uint32 GroupId,
+	const TRigidTransform<float, 3>& OldReferenceSpaceTransform,
+	const TRigidTransform<float, 3>& ReferenceSpaceTransform,
+	const TVector<float, 3>& LinearVelocityScale,
+	float AngularVelocityScale)
+{
+	TRigidTransform<float, 3> OldRootBoneLocalTransform = OldReferenceSpaceTransform;
+	OldRootBoneLocalTransform.AddToTranslation(-OldLocalSpaceLocation);
+
+	// Calculate deltas
+	const FTransform DeltaTransform = ReferenceSpaceTransform.GetRelativeTransform(OldReferenceSpaceTransform);
+
+	// Apply linear velocity scale
+	const TVector<float, 3> LinearRatio = TVector<float, 3>(1.f) - LinearVelocityScale.BoundToBox(TVector<float, 3>(0.f), TVector<float, 3>(1.f));
+	const TVector<float, 3> DeltaPosition = LinearRatio * DeltaTransform.GetTranslation();
+
+	// Apply angular velocity scale
+	TRotation<float, 3> DeltaRotation = DeltaTransform.GetRotation();
+	TVector<float, 3> Axis;
+	float DeltaAngle;
+	DeltaRotation.ToAxisAndAngle(Axis, DeltaAngle);
+	if (DeltaAngle > PI)
+	{
+		DeltaAngle -= 2.f * PI;
+	}
+
+	const float AngularRatio = FMath::Clamp(1.f - AngularVelocityScale, 0.f, 1.f);
+	DeltaRotation = FQuat(Axis, DeltaAngle * AngularRatio);
+
+	// Transform points back into the previous frame of reference before applying the adjusted deltas 
+	PreSimulationTransforms[GroupId] = OldRootBoneLocalTransform.Inverse() * FTransform(DeltaRotation, DeltaPosition) * OldRootBoneLocalTransform;
+}
+
+float FClothingSimulationSolver::SetParticleMassPerArea(int32 Offset, int32 Size, const TTriangleMesh<float>& Mesh)
+{
+	// Zero out masses
+	TPBDParticles<float, 3>& Particles = Evolution->Particles();
+	for (int32 Index = Offset; Index < Offset + Size; ++Index)
+	{
+		Particles.M(Index) = 0.f;
+	}
+
+	// Assign per particle mass proportional to connected area.
+	const TArray<TVector<int32, 3>>& SurfaceElements = Mesh.GetSurfaceElements();
+	float TotalArea = 0.f;
+	for (const TVector<int32, 3>& Tri : SurfaceElements)
+	{
+		const float TriArea = 0.5f * TVector<float, 3>::CrossProduct(
+			Particles.X(Tri[1]) - Particles.X(Tri[0]),
+			Particles.X(Tri[2]) - Particles.X(Tri[0])).Size();
+		TotalArea += TriArea;
+		const float ThirdTriArea = TriArea / 3.f;
+		Particles.M(Tri[0]) += ThirdTriArea;
+		Particles.M(Tri[1]) += ThirdTriArea;
+		Particles.M(Tri[2]) += ThirdTriArea;
+	}
+
+	UE_LOG(LogChaosCloth, Verbose, TEXT("Total area: %f, SI total area: %f"), TotalArea, TotalArea / FMath::Square(ChaosClothingSimulationSolverConstant::WorldScale));
+	return TotalArea;
+}
+
+void FClothingSimulationSolver::ParticleMassUpdateDensity(const TTriangleMesh<float>& Mesh, float Density)
+{
+	const TSet<int32> Vertices = Mesh.GetVertices();
+	TPBDParticles<float, 3>& Particles = Evolution->Particles();
+	float TotalMass = 0.f;
+	for (const int32 Vertex : Vertices)
+	{
+		Particles.M(Vertex) *= Density;
+		TotalMass += Particles.M(Vertex);
+	}
+
+	UE_LOG(LogChaosCloth, Verbose, TEXT("Total mass: %f, "), TotalMass);
+}
+
+void FClothingSimulationSolver::ParticleMassClampAndEnslave(int32 Offset, int32 Size, float MinPerParticleMass, const TFunctionRef<bool(int32)>& KinematicPredicate)
+{
+	TPBDParticles<float, 3>& Particles = Evolution->Particles();
+	for (int32 Index = Offset; Index < Offset + Size; ++Index)
+	{
+		Particles.M(Index) = FMath::Max(Particles.M(Index), MinPerParticleMass);
+		Particles.InvM(Index) = KinematicPredicate(Index - Offset) ? 0.f : 1.f / Particles.M(Index);
+	}
+}
+
+void FClothingSimulationSolver::SetProperties(uint32 GroupId, float DampingCoefficient, float CollisionThickness, float FrictionCoefficient)
+{
+	Evolution->SetDamping(DampingCoefficient, GroupId);
+	Evolution->SetCollisionThickness(CollisionThickness, GroupId);
+	Evolution->SetCoefficientOfFriction(FrictionCoefficient, GroupId);
+}
+
+void FClothingSimulationSolver::SetGravity(uint32 GroupId, const TVector<float, 3>& InGravity)
+{
+	Evolution->GetGravityForces(GroupId).SetAcceleration(InGravity);
+}
+
+void FClothingSimulationSolver::SetWindVelocity(const TVector<float, 3>& InWindVelocity)
+{
+	WindVelocity = InWindVelocity * ChaosClothingSimulationSolverConstant::WorldScale;
+}
+
+void FClothingSimulationSolver::SetWindVelocityField(uint32 GroupId, float DragCoefficient, const TTriangleMesh<float>* TriangleMesh)
+{
+	TVelocityField<float, 3>& VelocityField = Evolution->GetVelocityField(GroupId);
+	VelocityField.SetGeometry(TriangleMesh);
+	VelocityField.SetDragCoefficient(DragCoefficient);
+}
+
+void FClothingSimulationSolver::SetSelfCollisions(uint32 GroupId, float SelfCollisionThickness, const TTriangleMesh<float>* TriangleMesh)
+{
+#ifdef CHAOS_CLOTH_MUST_IMPROVE_SELF_COLLISION  // TODO: Improve self-collision until it can run in engine tests without crashing the simulation.
+	if (TriangleMesh)
+	{
+		// TODO(mlentine): Parallelize these for multiple meshes
+		// TODO(Kriss.Gossart): Check/fix potential issue with CollisionTriangles/mesh indices
+		Evolution->CollisionTriangles().Append(TriangleMesh->GetSurfaceElements());
+
+		const TVector<int32, 2> VertexRange = TriangleMesh->GetVertexRange();
+		for (int32 Index = VertexRange[0]; Index <= VertexRange[1]; ++Index)
+		{
+			const TSet<int32> Neighbors = TriangleMesh->GetNRing(Index, 5);
+			for (int32 Element : Neighbors)
+			{
+				check(Index != Element);
+				Evolution->DisabledCollisionElements().Add(TVector<int32, 2>(Index, Element));
+				Evolution->DisabledCollisionElements().Add(TVector<int32, 2>(Element, Index));
+			}
+		}
+	}
+#endif  // #ifdef CHAOS_CLOTH_MUST_IMPROVE_SELF_COLLISION
+
+	// TODO: Note that enabling self collision also enable collisions between all cloths running within the same solver.
+	//       Therefore this function should be called nevertheless.
+	Evolution->SetSelfCollisionThickness(SelfCollisionThickness, GroupId);
+}
+
+void FClothingSimulationSolver::ApplyPreSimulationTransforms()
+{
+	const TVector<float, 3> DeltaLocalSpaceLocation = LocalSpaceLocation - OldLocalSpaceLocation;
+
+	const TPBDActiveView<TPBDParticles<float, 3>>& ParticlesActiveView = Evolution->ParticlesActiveView();
+	const TArray<uint32>& ParticleGroupIds = Evolution->ParticleGroupIds();
+
+	ParticlesActiveView.ParallelFor(
+		[this, &ParticleGroupIds, &DeltaLocalSpaceLocation](TPBDParticles<float, 3>& Particles, int32 Index)
+		{
+			const TRigidTransform<float, 3>& GroupSpaceTransform = PreSimulationTransforms[ParticleGroupIds[Index]];
+
+			// Update initial state for particles
+			Particles.P(Index) = Particles.X(Index) = GroupSpaceTransform.TransformPosition(Particles.X(Index)) - DeltaLocalSpaceLocation;
+			Particles.V(Index) = GroupSpaceTransform.TransformVector(Particles.V(Index));
+
+			// Update anim initial state (target updated by skinning)
+			OldAnimationPositions[Index] = GroupSpaceTransform.TransformPosition(OldAnimationPositions[Index]) - DeltaLocalSpaceLocation;
+		}, /*MinParallelBatchSize =*/ 1000);  // TODO: Profile this value
+
+	const TPBDActiveView<TKinematicGeometryClothParticles<float, 3>>& CollisionParticlesActiveView = Evolution->CollisionParticlesActiveView();
+	const TArray<uint32>& CollisionParticleGroupIds = Evolution->CollisionParticleGroupIds();
+
+	CollisionParticlesActiveView.SequentialFor(  // There's unlikely to ever have enough collision particles for a parallel for
+		[this, &CollisionParticleGroupIds, &DeltaLocalSpaceLocation](TKinematicGeometryClothParticles<float, 3>& CollisionParticles, int32 Index)
+		{
+			const TRigidTransform<float, 3>& GroupSpaceTransform = PreSimulationTransforms[CollisionParticleGroupIds[Index]];
+
+			// Update initial state for collisions
+			OldCollisionTransforms[Index] = GroupSpaceTransform * OldCollisionTransforms[Index];
+			OldCollisionTransforms[Index].AddToTranslation(-DeltaLocalSpaceLocation);
+			CollisionParticles.X(Index) = OldCollisionTransforms[Index].GetTranslation();
+			CollisionParticles.R(Index) = OldCollisionTransforms[Index].GetRotation();
+		});
+}
+
+void FClothingSimulationSolver::Update(float InDeltaTime)
+{
+	SCOPE_CYCLE_COUNTER(STAT_ChaosClothSolverUpdate);
+
+	// Filter delta time to smoothen time variations and prevent unwanted vibrations
+	static const float DeltaTimeDecay = 0.1f;
+	const float PrevDeltaTime = DeltaTime;
+	DeltaTime = DeltaTime + (InDeltaTime - DeltaTime) * DeltaTimeDecay;
+
+	// Update Solver Colliders
+	{
+		SCOPE_CYCLE_COUNTER(STAT_ChaosClothSolverUpdateColliders);
+
+		// Reset collision particles, but leave cloth's own colliders
+		Evolution->ResetCollisionParticles(CollisionParticlesOffset);
+
+		// Add/re-add the solver's dynamic collider particles that can change at every update
+		for (FClothingSimulationCollider* const Collider : Colliders)
+		{
+			// Add the collider's particles
+			Collider->Add(this);
+
+			if (bAreCollisionParticlesDirty)
+			{
+				// Reset initial state that will become the previous state after the swap below
+				Collider->Update(this);
+			}
+		}
+		bAreCollisionParticlesDirty = false;
+
+		// Update colliders
+		Swap(OldCollisionTransforms, CollisionTransforms);
+
+		for (FClothingSimulationCollider* Collider: Colliders)
+		{
+			Collider->Update(this);
+		}
+	}
+
+	// Update Cloths
+	{
+		SCOPE_CYCLE_COUNTER(STAT_ChaosClothSolverUpdateCloths);
+
+		Swap(OldAnimationPositions, AnimationPositions);
+
+		for (FClothingSimulationCloth* Cloth : Cloths)
+		{
+			const uint32 GroupId = Cloth->GetGroupId();
+
+			// Pre-update overridable solver properties first
+			Evolution->GetGravityForces(GroupId).SetAcceleration(Gravity);
+			Evolution->GetVelocityField(GroupId).SetVelocity(WindVelocity);
+			Evolution->GetVelocityField(GroupId).SetFluidDensity(WindFluidDensity);
+
+			Cloth->Update(this);
+		}
+	}
+
+	// Pre solver step, apply group space transforms for teleport and linear/delta ratios, ...etc
+	{
+		SCOPE_CYCLE_COUNTER(STAT_ChaosClothSolverUpdatePreSolverStep);
+
+		ApplyPreSimulationTransforms();
+	}
+
+	// Advance Sim
+	{
+		SCOPE_CYCLE_COUNTER(STAT_ChaosClothSolverUpdateSolverStep);
+
+		Evolution->SetIterations(NumIterations);
+
+		const float SubstepDeltaTime = DeltaTime / (float)NumSubsteps;
+	
+		for (int32 i = 0; i < NumSubsteps; ++i)
+		{
+			Evolution->AdvanceOneTimeStep(SubstepDeltaTime);
+		}
+
+		Time = Evolution->GetTime();
+		UE_LOG(LogChaosCloth, VeryVerbose, TEXT("DeltaTime: %.6f, FilteredDeltaTime: %.6f, Time = %.6f"), InDeltaTime, DeltaTime, Time);
+	}
+
+	// Post solver step, update normals, ...etc
+	{
+		SCOPE_CYCLE_COUNTER(STAT_ChaosClothSolverUpdatePostSolverStep);
+
+		for (FClothingSimulationCloth* Cloth : Cloths)
+		{
+			Cloth->PostUpdate(this);
+		}
+	}
+
+	// Save old space location for next update
+	OldLocalSpaceLocation = LocalSpaceLocation;
+}
+
+FBoxSphereBounds FClothingSimulationSolver::CalculateBounds() const
+{
+	SCOPE_CYCLE_COUNTER(STAT_ChaosClothSolverCalculateBounds);
+
+	const TPBDActiveView<TPBDParticles<float, 3>>& ParticlesActiveView = Evolution->ParticlesActiveView();
+
+	if (ParticlesActiveView.HasActiveRange())
+	{
+		// Calculate bounding box
+		TAABB<float, 3> BoundingBox = TAABB<float, 3>::EmptyAABB();
+
+		ParticlesActiveView.SequentialFor(
+			[&BoundingBox](TPBDParticles<float, 3>& Particles, int32 Index)
+			{
+				BoundingBox.GrowToInclude(Particles.X(Index));
+			});
+
+		// Calculate (squared) radius
+		const TVector<float, 3> Center = BoundingBox.Center();
+		float SquaredRadius = 0.f;
+
+		ParticlesActiveView.SequentialFor(
+			[&SquaredRadius, &Center](TPBDParticles<float, 3>& Particles, int32 Index)
+			{
+				SquaredRadius = FMath::Max(SquaredRadius, (Particles.X(Index) - Center).SizeSquared());
+			});
+
+		// Update bounds with this cloth
+		return FBoxSphereBounds(LocalSpaceLocation + BoundingBox.Center(), BoundingBox.Extents() * 0.5f, FMath::Sqrt(SquaredRadius));
+	}
+
+	return FBoxSphereBounds(LocalSpaceLocation, FVector(0.f), 0.f);
+}
diff --git a/Engine/Plugins/Experimental/ChaosCloth/Source/Chaos/Public/ChaosCloth/ChaosClothConfig.h b/Engine/Plugins/Experimental/ChaosCloth/Source/Chaos/Public/ChaosCloth/ChaosClothConfig.h
--- a/Engine/Plugins/Experimental/ChaosCloth/Source/Chaos/Public/ChaosCloth/ChaosClothConfig.h
+++ b/Engine/Plugins/Experimental/ChaosCloth/Source/Chaos/Public/ChaosCloth/ChaosClothConfig.h
@@ -28,6 +28,9 @@ public:
 	/** PostLoad override used to deal with updates/changes in properties. */
 	virtual void PostLoad() override;
 
+	/** Return the mass value, from whichever mass mode (Density, UniformMass, or TotalMass) is selected. */
+	float GetMassValue() const;
+
 	/**
 	 * How cloth particle mass is determined
 	 * -	Uniform Mass: Every particle's mass will be set to the value specified in the UniformMass setting. Mostly to avoid as it can cause some serious issues with irregular tessellations.
diff --git a/Engine/Plugins/Experimental/ChaosCloth/Source/Chaos/Public/ChaosCloth/ChaosClothConstraints.h b/Engine/Plugins/Experimental/ChaosCloth/Source/Chaos/Public/ChaosCloth/ChaosClothConstraints.h
new file mode 100644
--- /dev/null
+++ b/Engine/Plugins/Experimental/ChaosCloth/Source/Chaos/Public/ChaosCloth/ChaosClothConstraints.h
@@ -0,0 +1,119 @@
+// Copyright Epic Games, Inc. All Rights Reserved.
+#pragma once
+
+#include "Chaos/Vector.h"
+#include "Containers/ContainersFwd.h"
+#include "Templates/SharedPointer.h"
+
+namespace Chaos
+{
+	template<typename T, int d> class TPBDEvolution;
+
+	class FPBDSpringConstraints;
+	template<typename T, int d> class TXPBDSpringConstraints;
+	template<typename T> class TPBDBendingConstraints;
+	class FPBDAxialSpringConstraints;
+	template<typename T, int d> class TXPBDAxialSpringConstraints;
+	template<typename T> class TPBDVolumeConstraint;
+	template<typename T, int d> class TPBDLongRangeConstraints;
+	template<typename T, int d> class TXPBDLongRangeConstraints;
+	template<typename T, int d> class TPBDSphericalConstraint;
+	template<typename T, int d> class TPBDSphericalBackstopConstraint;
+	template<typename T, int d> class TPBDAnimDriveConstraint;
+	template<typename T, int d> class TPBDShapeConstraints;
+
+	class FClothConstraints final
+	{
+	public:
+		FClothConstraints();
+		~FClothConstraints();
+
+		// ---- Solver interface ----
+		void Initialize(
+			TPBDEvolution<float, 3>* InEvolution,
+			const TArray<TVector<float, 3>>& InAnimationPositions,
+			const TArray<TVector<float, 3>>& InAnimationNormals,
+			int32 InParticleOffset,
+			int32 InNumParticles);
+		// ---- End of Solver interface ----
+
+		// ---- Cloth interface ----
+		void SetEdgeConstraints(TArray<TVector<int32, 2>>&& Edges, float EdgeStiffness, bool bUseXPBDConstraints);
+		void SetEdgeConstraints(const TArray<TVector<int32, 3>>& SurfaceElements, float EdgeStiffness, bool bUseXPBDConstraints);
+		void SetBendingConstraints(TArray<TVector<int32, 2>>&& Edges, float BendingStiffness, bool bUseXPBDConstraints);
+		void SetBendingConstraints(TArray<TVector<int32, 4>>&& BendingElements, float BendingStiffness);
+		void SetAreaConstraints(TArray<TVector<int32, 3>>&& SurfaceElements, float AreaStiffness, bool bUseXPBDConstraints);
+		void SetVolumeConstraints(TArray<TVector<int32, 2>>&& DoubleBendingEdges, float VolumeStiffness);
+		void SetVolumeConstraints(TArray<TVector<int32, 3>>&& SurfaceElements, float VolumeStiffness);
+		void SetLongRangeConstraints(const TMap<int32, TSet<uint32>>& PointToNeighborsMap, float StrainLimitingStiffness, float LimitScale, bool bUseGeodesicDistance, bool bUseXPBDConstraints);
+		void SetMaximumDistanceConstraints(const TConstArrayView<float>& MaxDistances);
+		void SetBackstopConstraints(const TConstArrayView<float>& BackstopDistances, const TConstArrayView<float>& BackstopRadiuses);
+		void SetAnimDriveConstraints(const TConstArrayView<float>& AnimDriveMultipliers);
+		void SetShapeTargetConstraints(float ShapeTargetStiffness);
+
+		void CreateRules();
+		void Enable(bool bEnable);
+
+		void SetMaxDistancesMultiplier(float InMaxDistancesMultiplier) { MaxDistancesMultiplier = InMaxDistancesMultiplier; }
+		void SetAnimDriveSpringStiffness(float InAnimDriveSpringStiffness) { AnimDriveSpringStiffness = InAnimDriveSpringStiffness; }
+		// ---- End of Cloth interface ----
+
+		// ---- Debug functions ----
+		const TSharedPtr<FPBDSpringConstraints>& GetTwoEdgeConstraints() const { return TwoEdgeConstraints; }  
+		const TSharedPtr<TXPBDSpringConstraints<float, 3>>& GetXTwoEdgeConstraints() const { return XTwoEdgeConstraints; }
+		const TSharedPtr<FPBDSpringConstraints> GetThreeEdgeConstraints() const { return ThreeEdgeConstraints; }
+		const TSharedPtr<TXPBDSpringConstraints<float, 3>>& GetXThreeEdgeConstraints() const { return XThreeEdgeConstraints; } 
+		const TSharedPtr<FPBDSpringConstraints>& GetEdgeConstraints() const { return EdgeConstraints; }
+		const TSharedPtr<TXPBDSpringConstraints<float, 3>>& GetXEdgeConstraints() const { return XEdgeConstraints; }
+		const TSharedPtr<FPBDSpringConstraints>& GetBendingConstraints() const { return BendingConstraints; }  
+		const TSharedPtr<TXPBDSpringConstraints<float, 3>>& GetXBendingConstraints() const { return XBendingConstraints; }
+		const TSharedPtr<TPBDBendingConstraints<float>>& GetBendingElementConstraints() const { return BendingElementConstraints; }
+		const TSharedPtr<FPBDAxialSpringConstraints>& GetAreaConstraints() const { return AreaConstraints; }
+		const TSharedPtr<TXPBDAxialSpringConstraints<float,3>>& GetXAreaConstraints() const { return XAreaConstraints; }
+		const TSharedPtr<FPBDSpringConstraints>& GetThinShellVolumeConstraints() const { return ThinShellVolumeConstraints; }
+		const TSharedPtr<TPBDVolumeConstraint<float>>& GetVolumeConstraints() const { return VolumeConstraints; } 
+		const TSharedPtr<TPBDLongRangeConstraints<float, 3>>& GetLongRangeConstraints() const { return LongRangeConstraints; }
+		const TSharedPtr<TXPBDLongRangeConstraints<float,3>>& GetXLongRangeConstraints() const { return XLongRangeConstraints; } 
+		const TSharedPtr<TPBDSphericalConstraint<float, 3>>& GetMaximumDistanceConstraints() const { return MaximumDistanceConstraints; }
+		const TSharedPtr<TPBDSphericalBackstopConstraint<float, 3>>& GetBackstopConstraints() const { return BackstopConstraints; }
+		const TSharedPtr<TPBDAnimDriveConstraint<float, 3>>& GetAnimDriveConstraints() const { return AnimDriveConstraints; }
+		const TSharedPtr<TPBDShapeConstraints<float, 3>>& GetShapeConstraints() const { return ShapeConstraints; }
+		// ---- End of debug functions ----
+
+	private:
+		TSharedPtr<FPBDSpringConstraints> TwoEdgeConstraints;
+		TSharedPtr<TXPBDSpringConstraints<float, 3>> XTwoEdgeConstraints;
+		TSharedPtr<FPBDSpringConstraints> ThreeEdgeConstraints;
+		TSharedPtr<TXPBDSpringConstraints<float, 3>> XThreeEdgeConstraints;
+		TSharedPtr<FPBDSpringConstraints> EdgeConstraints;
+		TSharedPtr<TXPBDSpringConstraints<float, 3>> XEdgeConstraints;
+		TSharedPtr<FPBDSpringConstraints> BendingConstraints;
+		TSharedPtr<TXPBDSpringConstraints<float, 3>> XBendingConstraints;
+		TSharedPtr<TPBDBendingConstraints<float>> BendingElementConstraints;
+		TSharedPtr<FPBDAxialSpringConstraints> AreaConstraints;
+		TSharedPtr<TXPBDAxialSpringConstraints<float, 3>> XAreaConstraints;
+		TSharedPtr<FPBDSpringConstraints> ThinShellVolumeConstraints;
+		TSharedPtr<TPBDVolumeConstraint<float>> VolumeConstraints;
+		TSharedPtr<TPBDLongRangeConstraints<float, 3>> LongRangeConstraints;
+		TSharedPtr<TXPBDLongRangeConstraints<float, 3>> XLongRangeConstraints;
+		TSharedPtr<TPBDSphericalConstraint<float, 3>> MaximumDistanceConstraints;
+		TSharedPtr<TPBDSphericalBackstopConstraint<float, 3>> BackstopConstraints;
+		TSharedPtr<TPBDAnimDriveConstraint<float, 3>> AnimDriveConstraints;
+		TSharedPtr<TPBDShapeConstraints<float, 3>> ShapeConstraints;
+		
+		TPBDEvolution<float, 3>* Evolution;
+		const TArray<TVector<float, 3>>* AnimationPositions;
+		const TArray<TVector<float, 3>>* AnimationNormals;
+
+		int32 ParticleOffset;
+		int32 NumParticles;
+		int32 ConstraintInitOffset;
+		int32 ConstraintRuleOffset;
+		int32 NumConstraintInits;
+		int32 NumConstraintRules;
+
+		// Animatable parameters
+		float MaxDistancesMultiplier;
+		float AnimDriveSpringStiffness;
+	};
+} // namespace Chaos
diff --git a/Engine/Plugins/Experimental/ChaosCloth/Source/Chaos/Public/ChaosCloth/ChaosClothingSimulation.h b/Engine/Plugins/Experimental/ChaosCloth/Source/Chaos/Public/ChaosCloth/ChaosClothingSimulation.h
--- a/Engine/Plugins/Experimental/ChaosCloth/Source/Chaos/Public/ChaosCloth/ChaosClothingSimulation.h
+++ b/Engine/Plugins/Experimental/ChaosCloth/Source/Chaos/Public/ChaosCloth/ChaosClothingSimulation.h
@@ -3,35 +3,70 @@
 
 #include "ClothingSimulation.h"
 #include "ClothingAsset.h"
-#include "Chaos/ArrayCollectionArray.h"
-#include "Chaos/Transform.h"
+#include "ClothCollisionData.h"
 #include "Chaos/ChaosDebugDrawDeclares.h"
 #include "ChaosCloth/ChaosClothConfig.h"
-#include "Components/SkeletalMeshComponent.h"
+#include "Templates/Atomic.h"
+#include "Templates/UniquePtr.h"
+
+class USkeletalMeshComponent;
+class FClothingSimulationContextCommon;
+#if WITH_EDITOR
+class UMaterial;
+#endif
 
 namespace Chaos
 {
-	typedef FClothingSimulationContextCommon ClothingSimulationContext;
 	template<class T, int d> class TPBDLongRangeConstraintsBase;
 	template<typename T> class TTriangleMesh;
 	template<typename T, int d> class TPBDEvolution;
+	class FClothingSimulationSolver;
+	class FClothingSimulationMesh;
+	class FClothingSimulationCloth;
+	class FClothingSimulationCollider;
 
-	class ClothingSimulation : public FClothingSimulationCommon
+	typedef FClothingSimulationContextCommon FClothingSimulationContext;
+
+	class FClothingSimulation : public FClothingSimulationCommon
 #if WITH_EDITOR
 		, public FGCObject  // Add garbage collection for debug cloth material
 #endif  // #if WITH_EDITOR
 	{
 	public:
-		ClothingSimulation();
-		virtual ~ClothingSimulation() override;
+		FClothingSimulation();
+		virtual ~FClothingSimulation() override;
+
+	protected:
+		// IClothingSimulation interface
+		virtual void Initialize() override;
+		virtual void Shutdown() override;
 
+		virtual IClothingSimulationContext* CreateContext() override;
+		virtual void DestroyContext(IClothingSimulationContext* InContext) override { delete InContext; }
+
+		virtual void CreateActor(USkeletalMeshComponent* InOwnerComponent, UClothingAssetBase* InAsset, int32 SimDataIndex) override;
+		virtual void DestroyActors() override;
+
+		virtual bool ShouldSimulate() const override { return true; }
+		virtual void Simulate(IClothingSimulationContext* InContext) override;
+		virtual void GetSimulationData(TMap<int32, FClothSimulData>& OutData, USkeletalMeshComponent* InOwnerComponent, USkinnedMeshComponent* InOverrideComponent) const override;
+
+		// Return bounds in local space (or in world space if InOwnerComponent is null).
+		virtual FBoxSphereBounds GetBounds(const USkeletalMeshComponent* InOwnerComponent) const override;
+
+		virtual void AddExternalCollisions(const FClothCollisionData& InData) override;
+		virtual void ClearExternalCollisions() override;
+		virtual void GetCollisions(FClothCollisionData& OutCollisions, bool bIncludeExternal = true) const override;
+		// End of IClothingSimulation interface
+
+	public:
 		// Set the animation drive stiffness for all actors
-		void SetAnimDriveSpringStiffness(float InStiffness);
+		void SetAnimDriveSpringStiffness(float InAnimDriveSpringStiffness);
 		void SetGravityOverride(const FVector& InGravityOverride);
 		void DisableGravityOverride();
 
 		// Function to be called if any of the assets' configuration parameters have changed
-		void RefreshClothConfig();
+		void RefreshClothConfig(const IClothingSimulationContext* InContext);
 		// Function to be called if any of the assets' physics assets changes (colliders)
 		// This seems to only happen when UPhysicsAsset::RefreshPhysicsAssetChange is called with
 		// bFullClothRefresh set to false during changes created using the viewport manipulators.
@@ -52,51 +87,32 @@ namespace Chaos
 		// End of FGCObject interface
 
 		// Editor only debug draw function
-		CHAOSCLOTH_API void DebugDrawPhysMeshShaded(USkeletalMeshComponent* OwnerComponent, FPrimitiveDrawInterface* PDI) const;
-		CHAOSCLOTH_API void DebugDrawParticleIndices(USkeletalMeshComponent* OwnerComponent, FCanvas* Canvas, const FSceneView* SceneView) const;
-		CHAOSCLOTH_API void DebugDrawMaxDistanceValues(USkeletalMeshComponent* OwnerComponent, FCanvas* Canvas, const FSceneView* SceneView) const;
+		CHAOSCLOTH_API void DebugDrawPhysMeshShaded(FPrimitiveDrawInterface* PDI) const;
+		CHAOSCLOTH_API void DebugDrawParticleIndices(FCanvas* Canvas, const FSceneView* SceneView) const;
+		CHAOSCLOTH_API void DebugDrawMaxDistanceValues(FCanvas* Canvas, const FSceneView* SceneView) const;
 #endif  // #if WITH_EDITOR
 
 #if WITH_EDITOR || CHAOS_DEBUG_DRAW
 		// Editor & runtime debug draw functions
-		CHAOSCLOTH_API void DebugDrawPhysMeshWired(USkeletalMeshComponent* OwnerComponent = nullptr, FPrimitiveDrawInterface* PDI = nullptr) const;
-		CHAOSCLOTH_API void DebugDrawPointNormals(USkeletalMeshComponent* OwnerComponent = nullptr, FPrimitiveDrawInterface* PDI = nullptr) const;
-		CHAOSCLOTH_API void DebugDrawInversedPointNormals(USkeletalMeshComponent* OwnerComponent = nullptr, FPrimitiveDrawInterface* PDI = nullptr) const;
-		CHAOSCLOTH_API void DebugDrawFaceNormals(USkeletalMeshComponent* OwnerComponent = nullptr, FPrimitiveDrawInterface* PDI = nullptr) const;
-		CHAOSCLOTH_API void DebugDrawInversedFaceNormals(USkeletalMeshComponent* OwnerComponent = nullptr, FPrimitiveDrawInterface* PDI = nullptr) const;
-		CHAOSCLOTH_API void DebugDrawCollision(USkeletalMeshComponent* OwnerComponent = nullptr, FPrimitiveDrawInterface* PDI = nullptr) const;
-		CHAOSCLOTH_API void DebugDrawBackstops(USkeletalMeshComponent* OwnerComponent = nullptr, FPrimitiveDrawInterface* PDI = nullptr) const;
-		CHAOSCLOTH_API void DebugDrawMaxDistances(USkeletalMeshComponent* OwnerComponent = nullptr, FPrimitiveDrawInterface* PDI = nullptr) const;
-		CHAOSCLOTH_API void DebugDrawAnimDrive(USkeletalMeshComponent* OwnerComponent = nullptr, FPrimitiveDrawInterface* PDI = nullptr) const;
-		CHAOSCLOTH_API void DebugDrawLongRangeConstraint(USkeletalMeshComponent* OwnerComponent = nullptr, FPrimitiveDrawInterface* PDI = nullptr) const;
-		CHAOSCLOTH_API void DebugDrawWindDragForces(USkeletalMeshComponent* OwnerComponent = nullptr, FPrimitiveDrawInterface* PDI = nullptr) const;
-		CHAOSCLOTH_API void DebugDrawLocalSpace(USkeletalMeshComponent* OwnerComponent = nullptr, FPrimitiveDrawInterface* PDI = nullptr) const;
+		CHAOSCLOTH_API void DebugDrawPhysMeshWired(FPrimitiveDrawInterface* PDI = nullptr) const;
+		CHAOSCLOTH_API void DebugDrawAnimMeshWired(FPrimitiveDrawInterface* PDI = nullptr) const;
+		CHAOSCLOTH_API void DebugDrawPointNormals(FPrimitiveDrawInterface* PDI = nullptr) const;
+		CHAOSCLOTH_API void DebugDrawInversedPointNormals(FPrimitiveDrawInterface* PDI = nullptr) const;
+		CHAOSCLOTH_API void DebugDrawCollision(FPrimitiveDrawInterface* PDI = nullptr) const;
+		CHAOSCLOTH_API void DebugDrawBackstops(FPrimitiveDrawInterface* PDI = nullptr) const;
+		CHAOSCLOTH_API void DebugDrawBackstopDistances(FPrimitiveDrawInterface* PDI = nullptr) const;
+		CHAOSCLOTH_API void DebugDrawMaxDistances(FPrimitiveDrawInterface* PDI = nullptr) const;
+		CHAOSCLOTH_API void DebugDrawAnimDrive(FPrimitiveDrawInterface* PDI = nullptr) const;
+		CHAOSCLOTH_API void DebugDrawLongRangeConstraint(FPrimitiveDrawInterface* PDI = nullptr) const;
+		CHAOSCLOTH_API void DebugDrawWindDragForces(FPrimitiveDrawInterface* PDI = nullptr) const;
+		CHAOSCLOTH_API void DebugDrawLocalSpace(FPrimitiveDrawInterface* PDI = nullptr) const;
 #endif  // #if WITH_EDITOR || CHAOS_DEBUG_DRAW
 
-	protected:
-		// IClothingSimulation interface
-		virtual void Initialize() override;
-		virtual void CreateActor(USkeletalMeshComponent* InOwnerComponent, UClothingAssetBase* InAsset, int32 SimDataIndex) override;
-		virtual void PostActorCreationInitialize() override;
-		virtual IClothingSimulationContext* CreateContext() override { return new ClothingSimulationContext(); }
-		virtual void Shutdown() override;
-		virtual bool ShouldSimulate() const override { return true; }
-		virtual void Simulate(IClothingSimulationContext* InContext) override;
-		virtual void DestroyActors() override;
-		virtual void DestroyContext(IClothingSimulationContext* InContext) override { delete InContext; }
-		virtual void GetSimulationData(TMap<int32, FClothSimulData>& OutData, USkeletalMeshComponent* InOwnerComponent, USkinnedMeshComponent* InOverrideComponent) const override;
-
-		// Return bounds in local space (or in world space if InOwnerComponent is null).
-		virtual FBoxSphereBounds GetBounds(const USkeletalMeshComponent* InOwnerComponent) const override;
-
-		virtual void AddExternalCollisions(const FClothCollisionData& InData) override;
-		virtual void ClearExternalCollisions() override;
-		virtual void GetCollisions(FClothCollisionData& OutCollisions, bool bIncludeExternal = true) const override;
-		// End of IClothingSimulation interface
-
 	private:
 		void ResetStats();
-		void UpdateStats(int32 InSimDataIndex);
+		void UpdateStats(const FClothingSimulationCloth* Cloth);
+
+		void UpdateSimulationFromSharedSimConfig();
 
 #if CHAOS_DEBUG_DRAW
 		// Runtime only debug draw functions
@@ -104,74 +120,18 @@ namespace Chaos
 		void DebugDrawGravity() const;
 #endif  // #if CHAOS_DEBUG_DRAW
 
-		void UpdateSimulationFromSharedSimConfig();
-		void BuildMesh(const FClothPhysicalMeshData& PhysMesh, int32 InSimDataIndex);
-
-		// Reset particle positions to animation positions, masses and velocities are cleared
-		void ResetParticles(int32 InSimDataIndex);
-
-		void SetParticleMasses(const UChaosClothConfig* ChaosClothSimConfig, const FClothPhysicalMeshData& PhysMesh, int32 InSimDataIndex);
-
-		void AddConstraints(const UChaosClothConfig* ChaosClothSimConfig, const FClothPhysicalMeshData& PhysMesh, int32 InSimDataIndex);
-
-		void AddSelfCollisions(int32 InSimDataIndex);
-
-		// Extract all collisions (from the physics asset and from the legacy apex collisions)
-		void ExtractCollisions(const UClothingAssetCommon* Asset, int32 InSimDataIndex);
-		// Extract the collisions from the physics asset referenced by the specified clothing asset
-		void ExtractPhysicsAssetCollisions(const UClothingAssetCommon* Asset, int32 InSimDataIndex);
-		// Extract the legacy apex collisions from the specified clothing asset 
-		void ExtractLegacyAssetCollisions(const UClothingAssetCommon* Asset, int32 InSimDataIndex);
-		// Add collisions from a ClothCollisionData structure
-		void AddCollisions(const FClothCollisionData& ClothCollisionData, const TArray<int32>& UsedBoneIndices, int32 InSimDataIndex);
-		// Process the collision function for all collisions of the specified cloth
-		void ForAllCollisions(TFunction<void(TGeometryClothParticles<float, 3>&, uint32)> CollisionFunction, int32 SimDataIndex);
-		// Update the collision transforms using the specified context, also reset initial collision particle states if the number of collisions has changed
-		void UpdateCollisionTransforms(const ClothingSimulationContext& Context, int32 InSimDataIndex);
-		// Return the correct bone index based on the asset used bone index array
-		FORCEINLINE int32 GetMappedBoneIndex(const TArray<int32>& UsedBoneIndices, int32 BoneIndex)
-		{
-			return UsedBoneIndices.IsValidIndex(BoneIndex) ? UsedBoneIndices[BoneIndex] : INDEX_NONE;
-		}
-
 	private:
-		// Assets
-		TArray<UClothingAssetCommon*> Assets;
-		UChaosClothSharedSimConfig* ClothSharedSimConfig;
+		// Simulation objects
+		TUniquePtr<FClothingSimulationSolver> Solver;  // Default solver
+		TArray<TUniquePtr<FClothingSimulationMesh>> Meshes;
+		TArray<TUniquePtr<FClothingSimulationCloth>> Cloths;
+		TArray<TUniquePtr<FClothingSimulationCollider>> Colliders;
 
-		// Cloth Interaction Parameters
-		// These simulation parameters can be changed through blueprints
-		// They will only be updated when the simulation is not running (so are safe to use on any cloth thread)
-		TArray<float> AnimDriveSpringStiffness; // One for every Asset
-		TArray<float> MaxDistancesMultipliers;
+		// External collision Data
+		FClothCollisionData ExternalCollisionData;
 
-		// Collision Data
-		FClothCollisionData ExternalCollisions;  // External collisions
-		TArray<Chaos::TRigidTransform<float, 3>> OldCollisionTransforms;  // Used for the kinematic collision transform update
-		TArray<Chaos::TRigidTransform<float, 3>> CollisionTransforms;  // Used for the kinematic collision transform update
-		Chaos::TArrayCollectionArray<int32> BoneIndices;
-		Chaos::TArrayCollectionArray<Chaos::TRigidTransform<float, 3>> BaseTransforms;
-
-		// Animation Data
-		TArray<Chaos::TVector<float, 3>> OldAnimationPositions;
-		TArray<Chaos::TVector<float, 3>> AnimationPositions;
-		TArray<Chaos::TVector<float, 3>> AnimationNormals;
-
-		// Sim Data
-		TArray<Chaos::TVector<uint32, 2>> IndexToRangeMap;
-
-		TArray<TUniquePtr<Chaos::TTriangleMesh<float>>> Meshes;
-		mutable TArray<TArray<Chaos::TVector<float, 3>>> FaceNormals;
-		mutable TArray<TArray<Chaos::TVector<float, 3>>> PointNormals;
-
-		TUniquePtr<Chaos::TPBDEvolution<float, 3>> Evolution;
-
-		TArray<Chaos::TVector<uint32, 2>> CollisionsRangeMap;  // Collisions first particle index and number of collisions for each simulated cloth
-		TArray<TArray<Chaos::TVector<uint32, 2>>> ExternalCollisionsRangeMaps;  // External collisions first particle index and number of collisions for each simulated cloth
-		uint32 ExternalCollisionsOffset;  // External collisions first particle index
-
-		float Time;
-		float DeltaTime;
+		// Shared cloth config
+		UChaosClothSharedSimConfig* ClothSharedSimConfig;
 
 		// Properties that must be readable from all threads
 		TAtomic<int32> NumCloths;
@@ -181,20 +141,12 @@ namespace Chaos
 		TAtomic<int32> NumSubsteps;
 		TAtomic<float> SimulationTime;
 
+		// Overrides
+		bool bUseLocalSpaceSimulation;
 		bool bUseGravityOverride;
 		FVector GravityOverride;
-		TArray<float> GravityScales;
-		TArray<bool> bUseGravityConfigs; 
-		TArray<FVector> GravityConfigs;
-		TArray<FVector> WindVelocities;
-
-		TArray<TSharedPtr<TPBDLongRangeConstraintsBase<float, 3>>> LongRangeConstraints;
-
-		bool bUseLocalSpaceSimulation;  // The function of this is to help with floating point precision errors if the character is far away form the world origin
-		FVector LocalSpaceLocation;  // This is used to translate between world space and simulation space. Add this to simulation space coordinates to get world space coordinates
-		TArray<FTransform> RootBoneWorldTransforms;  // Reference bone transform used for teleportation and to decouple the simulation from reference bone's accelerations
-		TArray<FVector> LinearDeltaRatios;  // Linear ratios applied to the reference bone transforms
-		TArray<float> AngularDeltaRatios;  // Angular ratios factor applied to the reference bone transforms
+		float MaxDistancesMultipliers;
+		float AnimDriveSpringStiffness;
 
 #if WITH_EDITOR
 		// Visualization material
diff --git a/Engine/Plugins/Experimental/ChaosCloth/Source/Chaos/Public/ChaosCloth/ChaosClothingSimulationCloth.h b/Engine/Plugins/Experimental/ChaosCloth/Source/Chaos/Public/ChaosCloth/ChaosClothingSimulationCloth.h
new file mode 100644
--- /dev/null
+++ b/Engine/Plugins/Experimental/ChaosCloth/Source/Chaos/Public/ChaosCloth/ChaosClothingSimulationCloth.h
@@ -0,0 +1,201 @@
+// Copyright Epic Games, Inc. All Rights Reserved.
+#pragma once
+
+#include "Chaos/Transform.h"
+#include "Chaos/TriangleMesh.h"
+#include "Chaos/AABB.h"
+#include "Containers/ContainersFwd.h"
+
+namespace Chaos
+{
+	class FClothingSimulationSolver;
+	class FClothingSimulationMesh;
+	class FClothingSimulationCollider;
+
+	// Cloth simulation node
+	class FClothingSimulationCloth final
+	{
+	public:
+		enum EMassMode
+		{
+			UniformMass,
+			TotalMass,
+			Density
+		};
+
+		FClothingSimulationCloth(
+			FClothingSimulationMesh* InMesh,
+			TArray<FClothingSimulationCollider*>&& InColliders,
+			uint32 InGroupId,
+			EMassMode InMassMode,
+			float InMassValue,
+			float InMinPerParticleMass,
+			float InEdgeStiffness,
+			float InBendingStiffness,
+			bool bInUseBendingElements,
+			float InAreaStiffness,
+			float InVolumeStiffness,
+			bool bInUseThinShellVolumeConstraints,
+			float InStrainLimitingStiffness,
+			float InLimitScale,
+			bool bInUseGeodesicDistance,
+			float InMaxDistancesMultiplier,
+			float InAnimDriveSpringStiffness,
+			float InShapeTargetStiffness,
+			bool bInUseXPBDConstraints,
+			float InGravityScale,
+			bool bIsGravityOverridden,
+			const TVector<float, 3>& InGravityOverride,
+			const TVector<float, 3>& InLinearVelocityScale,
+			float InAngularVelocityScale,
+			float InDragCoefficient,
+			float InDampingCoefficient,
+			float InCollisionThickness,
+			float InFrictionCoefficient,
+			bool bInUseSelfCollisions,
+			float InSelfCollisionThickness,
+			bool bInUseLODIndexOverride,
+			int32 InLODIndexOverride);
+		~FClothingSimulationCloth();
+
+		uint32 GetGroupId() const { return GroupId; }
+		uint32 GetLODIndex(const FClothingSimulationSolver* Solver) const { return LODIndices.FindChecked(Solver); }
+
+		int32 GetNumActiveKinematicParticles() const { return NumActiveKinematicParticles; }
+		int32 GetNumActiveDynamicParticles() const { return NumActiveDynamicParticles; }
+
+		// ---- Animatable property setters ----
+		void SetMaxDistancesMultiplier(float InMaxDistancesMultiplier) { MaxDistancesMultiplier = InMaxDistancesMultiplier; }
+		void SetAnimDriveSpringStiffness(float InAnimDriveSpringStiffness) { AnimDriveSpringStiffness = InAnimDriveSpringStiffness; }
+
+		void Reset() { bNeedsReset = true; }
+		void Teleport() { bNeedsTeleport = true; }
+		// ---- End of the animatable property setters ----
+
+		// ---- Node property getters/setters
+		FClothingSimulationMesh* GetMesh() const { return Mesh; }
+		void SetMesh(FClothingSimulationMesh* InMesh);
+
+		const TArray<FClothingSimulationCollider*>& GetColliders() const { return Colliders; }
+		void SetColliders(TArray<FClothingSimulationCollider*>&& InColliders);
+		void RemoveColliders();
+		//void AddCollider(FClothingSimulationCollider* InCollider);  // TODO
+		//void RemoveCollider(FClothingSimulationCollider* InCollider);  // TODO
+		// ---- End of the Node property getters/setters
+
+		// ---- Debugging/visualization functions
+		// Return the solver's input positions for this cloth source current LOD, not thread safe, call must be done right after the solver update.
+		TConstArrayView<TVector<float, 3>> GetAnimationPositions(const FClothingSimulationSolver* Solver) const;
+		// Return the solver's input normals for this cloth source current LOD, not thread safe, call must be done right after the solver update.
+		TConstArrayView<TVector<float, 3>> GetAnimationNormals(const FClothingSimulationSolver* Solver) const;
+		// Return the solver's positions for this cloth current LOD, not thread safe, call must be done right after the solver update.
+		TConstArrayView<TVector<float, 3>> GetParticlePositions(const FClothingSimulationSolver* Solver) const;
+		// Return the solver's normals for this cloth current LOD, not thread safe, call must be done right after the solver update.
+		TConstArrayView<TVector<float, 3>> GetParticleNormals(const FClothingSimulationSolver* Solver) const;
+		// Return the solver's normals for this cloth current LOD, not thread safe, call must be done right after the solver update.
+		TConstArrayView<float> GetParticleInvMasses(const FClothingSimulationSolver* Solver) const;
+		// Return the current gravity as applied by the solver using the various overrides, not thread safe, call must be done right after the solver update.
+		TVector<float, 3> GetGravity(const FClothingSimulationSolver* Solver) const;
+		// Return the current bounding box based on a given solver, not thread safe, call must be done right after the solver update.
+		TAABB<float, 3> CalculateBoundingBox(const FClothingSimulationSolver* Solver) const;
+		// Return the current LOD Offset in the solver's particle array, or INDEX_NONE if no LOD is currently selected
+		int32 GetOffset(const FClothingSimulationSolver* Solver) const;
+		// Return the current LOD Mesh
+		const TTriangleMesh<float>& GetTriangleMesh(const FClothingSimulationSolver* Solver) const;
+		// Return the current LOD Weightmaps
+		const TArray<TConstArrayView<float>>& GetWeightMaps(const FClothingSimulationSolver* Solver) const;
+		// Return the local reference space transform for this cloth
+		const TRigidTransform<float, 3>& GetReferenceSpaceTransform() const { return ReferenceSpaceTransform;  }
+		// ---- End of the debugging/visualization functions
+
+		// ---- Solver interface ----
+		void Add(FClothingSimulationSolver* Solver);
+		void Remove(FClothingSimulationSolver* Solver);
+
+		void Update(FClothingSimulationSolver* Solver);
+		void PostUpdate(FClothingSimulationSolver* Solver);
+		// ---- End of the Solver interface ----
+
+	private:
+		int32 GetNumParticles(int32 InLODIndex) const;
+		int32 GetOffset(const FClothingSimulationSolver* Solver, int32 InLODIndex) const;
+
+	private:
+		struct FLODData
+		{
+			// Input mesh
+			const int32 NumParticles;
+			const TConstArrayView<uint32> Indices;
+			const TArray<TConstArrayView<float>> WeightMaps;
+
+			// Per Solver data
+			struct FSolverData
+			{
+				int32 Offset;
+				TTriangleMesh<float> TriangleMesh;  // TODO: Triangle Mesh shouldn't really be solver dependent (ie not use an offset)
+			};
+			TMap<FClothingSimulationSolver*, FSolverData> SolverData;
+
+			// Stats
+			int32 NumKinenamicParticles;
+			int32 NumDynammicParticles;
+
+			FLODData(int32 InNumParticles, const TConstArrayView<uint32>& InIndices, const TArray<TConstArrayView<float>>& InWeightMaps);
+
+			void Add(FClothingSimulationSolver* Solver, FClothingSimulationCloth* Cloth, int32 LODIndex);
+
+			void Enable(FClothingSimulationSolver* Solver, bool bEnable) const;
+
+			void ResetStartPose(FClothingSimulationSolver* Solver) const;
+
+			void UpdateNormals(FClothingSimulationSolver* Solver) const;
+		};
+
+		// Cloth parameters
+		FClothingSimulationMesh* Mesh;
+		TArray<FClothingSimulationCollider*> Colliders;
+		uint32 GroupId;
+		EMassMode MassMode;
+		float MassValue;
+		float MinPerParticleMass;
+		float EdgeStiffness;
+		float BendingStiffness;
+		bool bUseBendingElements;
+		float AreaStiffness;
+		float VolumeStiffness;
+		bool bUseThinShellVolumeConstraints;
+		float StrainLimitingStiffness;
+		float LimitScale;
+		bool bUseGeodesicDistance;
+		float MaxDistancesMultiplier;  // Animatable
+		float AnimDriveSpringStiffness;  // Animatable
+		float ShapeTargetStiffness;
+		bool bUseXPBDConstraints;
+		float GravityScale;
+		bool bIsGravityOverridden;
+		TVector<float, 3> GravityOverride;
+		TVector<float, 3> LinearVelocityScale;  // Linear ratio applied to the reference bone transforms
+		float AngularVelocityScale;  // Angular ratio factor applied to the reference bone transforms
+		float DragCoefficient;
+		float DampingCoefficient;
+		float CollisionThickness;
+		float FrictionCoefficient;
+		bool bUseSelfCollisions;
+		float SelfCollisionThickness;
+		bool bUseLODIndexOverride;
+		int32 LODIndexOverride;
+		bool bNeedsReset;
+		bool bNeedsTeleport;
+
+		// Reference space transform
+		TRigidTransform<float, 3> ReferenceSpaceTransform;  // TODO: Add override in the style of LODIndexOverride
+
+		// LOD data
+		TArray<FLODData> LODData;
+		TMap<FClothingSimulationSolver*, int32> LODIndices;
+
+		// Stats
+		int32 NumActiveKinematicParticles;
+		int32 NumActiveDynamicParticles;
+	};
+} // namespace Chaos
\ No newline at end of file
diff --git a/Engine/Plugins/Experimental/ChaosCloth/Source/Chaos/Public/ChaosCloth/ChaosClothingSimulationCollider.h b/Engine/Plugins/Experimental/ChaosCloth/Source/Chaos/Public/ChaosCloth/ChaosClothingSimulationCollider.h
new file mode 100644
--- /dev/null
+++ b/Engine/Plugins/Experimental/ChaosCloth/Source/Chaos/Public/ChaosCloth/ChaosClothingSimulationCollider.h
@@ -0,0 +1,102 @@
+// Copyright Epic Games, Inc. All Rights Reserved.
+#pragma once
+
+#include "ClothCollisionData.h"
+#include "Containers/ContainersFwd.h"
+
+class USkeletalMeshComponent;
+class UClothingAssetCommon;
+class FClothingSimulationContextCommon;
+
+namespace Chaos
+{
+	class FClothingSimulationSolver;
+	class FClothingSimulationCloth;
+
+	// Collider simulation node
+	class FClothingSimulationCollider final
+	{
+	public:
+		enum class ECollisionDataType : int32
+		{
+			Global = 0,  // Global collision slot (physics collisions)
+			Dynamic,  // Dynamic collision slot (aka external collisions)
+			LODs,  // LODIndex based start slot (LODs collisions)
+		};
+
+		FClothingSimulationCollider(
+			const UClothingAssetCommon* InAsset,  // Cloth asset for collision data, can be nullptr
+			const USkeletalMeshComponent* InSkeletalMeshComponent,  // For asset LODs management, can be nullptr
+			bool bInUseLODIndexOverride,
+			int32 InLODIndexOverride);
+		~FClothingSimulationCollider();
+
+		int32 GetNumGeometries() const { int32 NumGeometries = 0; for (const FLODData& LODDatum : LODData) { NumGeometries += LODDatum.NumGeometries; } return NumGeometries; }
+
+		// Return source (untransformed) collision data for global, dynamic and active LODs.
+		FClothCollisionData GetCollisionData() const;
+
+		// ---- Animatable property setters ----
+		// Set dynamic collision data, will only get updated when used as a Solver Collider TODO: Subclass collider?
+		void SetCollisionData(const FClothCollisionData* InCollisionData) { CollisionData = InCollisionData; }
+		// ---- End of the animatable property setters ----
+
+		// ---- Cloth/Solver interface ----
+		void Add(FClothingSimulationSolver* Solver, FClothingSimulationCloth* Cloth = nullptr);
+		void Remove(FClothingSimulationSolver* /*Solver*/, FClothingSimulationCloth* /*Cloth*/ = nullptr) {}
+
+		void Update(FClothingSimulationSolver* Solver, FClothingSimulationCloth* Cloth = nullptr);
+		// ---- End of the Cloth/Solver interface ----
+
+		// ---- Debugging and visualization functions ----
+		// Return currently LOD active collision particles translations, not thread safe, to use after solver update.
+		TConstArrayView<TVector<float, 3>> GetCollisionTranslations(const FClothingSimulationSolver* Solver, const FClothingSimulationCloth* Cloth, ECollisionDataType CollisionDataType) const;
+
+		// Return currently LOD active collision particles rotations, not thread safe, to use after solver update.
+		TConstArrayView<TRotation<float, 3>> GetCollisionRotations(const FClothingSimulationSolver* Solver, const FClothingSimulationCloth* Cloth, ECollisionDataType CollisionDataType) const;
+
+		// Return currently LOD active collision geometries, not thread safe, to use after solver update.
+		TConstArrayView<TUniquePtr<FImplicitObject>> GetCollisionGeometries(const FClothingSimulationSolver* Solver, const FClothingSimulationCloth* Cloth, ECollisionDataType CollisionDataType) const;
+		// ---- End of the debugging and visualization functions ----
+
+	private:
+		void ExtractPhysicsAssetCollision(FClothCollisionData& ClothCollisionData, TArray<int32>& UsedBoneIndices);
+
+	private:
+		struct FLODData
+		{
+			FClothCollisionData ClothCollisionData;
+			int32 NumGeometries;  // Number of collision bodies
+			TMap<TPair<const FClothingSimulationSolver*, const FClothingSimulationCloth*>, int32> Offsets;  // Solver particle offset
+
+			FLODData() : NumGeometries(0) {}
+
+			void Add(
+				FClothingSimulationSolver* Solver,
+				FClothingSimulationCloth* Cloth,
+				const FClothCollisionData& InClothCollisionData,
+				const TArray<int32>& UsedBoneIndices = TArray<int32>());
+
+			void Update(FClothingSimulationSolver* Solver, FClothingSimulationCloth* Cloth, const FClothingSimulationContextCommon* Context);
+
+			void Enable(FClothingSimulationSolver* Solver, FClothingSimulationCloth* Cloth, bool bEnable);
+
+			void ResetStartPose(FClothingSimulationSolver* Solver, FClothingSimulationCloth* Cloth);
+
+			FORCEINLINE static int32 GetMappedBoneIndex(const TArray<int32>& UsedBoneIndices, int32 BoneIndex)
+			{
+				return UsedBoneIndices.IsValidIndex(BoneIndex) ? UsedBoneIndices[BoneIndex] : INDEX_NONE;
+			}
+		};
+
+		const UClothingAssetCommon* Asset;
+		const USkeletalMeshComponent* SkeletalMeshComponent;
+		const FClothCollisionData* CollisionData;
+		bool bUseLODIndexOverride;
+		int32 LODIndexOverride;
+
+		// Collision primitives
+		TArray<FLODData> LODData;  // Actual LODs start at LODStart
+		int32 LODIndex;  // TODO: Have map of LODIndex per solver
+	};
+} // namespace Chaos
\ No newline at end of file
diff --git a/Engine/Plugins/Experimental/ChaosCloth/Source/Chaos/Public/ChaosCloth/ChaosClothingSimulationInteractor.h b/Engine/Plugins/Experimental/ChaosCloth/Source/Chaos/Public/ChaosCloth/ChaosClothingSimulationInteractor.h
--- a/Engine/Plugins/Experimental/ChaosCloth/Source/Chaos/Public/ChaosCloth/ChaosClothingSimulationInteractor.h
+++ b/Engine/Plugins/Experimental/ChaosCloth/Source/Chaos/Public/ChaosCloth/ChaosClothingSimulationInteractor.h
@@ -7,7 +7,7 @@
 #include "ChaosClothingSimulationInteractor.generated.h"
 
 // Command signature for handling synced command buffer
-DECLARE_DELEGATE_TwoParams(ChaosClothInteractorCommand, Chaos::ClothingSimulation*, Chaos::ClothingSimulationContext*)
+DECLARE_DELEGATE_TwoParams(ChaosClothInteractorCommand, Chaos::FClothingSimulation*, FClothingSimulationContextCommon*)
 
 UCLASS(BlueprintType)
 class CHAOSCLOTH_API UChaosClothingSimulationInteractor : public UClothingSimulationInteractor
diff --git a/Engine/Plugins/Experimental/ChaosCloth/Source/Chaos/Public/ChaosCloth/ChaosClothingSimulationMesh.h b/Engine/Plugins/Experimental/ChaosCloth/Source/Chaos/Public/ChaosCloth/ChaosClothingSimulationMesh.h
new file mode 100644
--- /dev/null
+++ b/Engine/Plugins/Experimental/ChaosCloth/Source/Chaos/Public/ChaosCloth/ChaosClothingSimulationMesh.h
@@ -0,0 +1,72 @@
+// Copyright Epic Games, Inc. All Rights Reserved.
+#pragma once
+
+#include "Containers/ContainersFwd.h"
+#include "Chaos/Transform.h"
+
+class USkeletalMeshComponent;
+class UClothingAssetCommon;
+
+namespace Chaos
+{
+	class FClothingSimulationSolver;
+
+	// Mesh simulation node
+	class FClothingSimulationMesh final
+	{
+	public:
+		FClothingSimulationMesh(const UClothingAssetCommon* InAsset, const USkeletalMeshComponent* InSkeletalMeshComponent);
+		~FClothingSimulationMesh();
+
+		// ---- Node property getters/setters
+		const UClothingAssetCommon* GetAsset() const { return Asset; }
+		const USkeletalMeshComponent* GetSkeletalMeshComponent() const { return SkeletalMeshComponent; }
+		// ---- End of node property getters/setters
+
+		// ---- Cloth interface ----
+		void Update(
+			FClothingSimulationSolver* Solver,
+			int32 PrevLODIndex,
+			int32 LODIndex,
+			int32 PrevOffset,
+			int32 Offset);
+
+		// Return the LOD Index specified by the input SkeletalMeshComponent. Note that this is not the same as the current LOD index, and that Mesh LOD changes are always driven by the Cloth.
+		int32 GetLODIndex() const;
+		int32 GetNumLODs() const;
+		int32 GetNumPoints(int32 LODIndex) const;
+		TConstArrayView<const uint32> GetIndices(int32 LODIndex) const;
+		TArray<TConstArrayView<const float>> GetWeightMaps(int32 LODIndex) const;
+		TRigidTransform<float, 3> GetReferenceBoneTransform() const;
+
+		bool WrapDeformLOD(
+			int32 PrevLODIndex,
+			int32 LODIndex,
+			const TVector<float, 3>* Normals,
+			const TVector<float, 3>* Positions,
+			TVector<float, 3>* OutPositions) const;
+
+		bool WrapDeformLOD(
+			int32 PrevLODIndex,
+			int32 LODIndex,
+			const TVector<float, 3>* Normals,
+			const TVector<float, 3>* Positions,
+			const TVector<float, 3>* Velocities,
+			TVector<float, 3>* OutPositions0,
+			TVector<float, 3>* OutPositions1,
+			TVector<float, 3>* OutVelocities) const;
+		// ---- End of the Cloth interface ----
+
+	private:
+		int32 GetReferenceBoneIndex() const;
+		void SkinPhysicsMesh(
+			int32 LODIndex,
+			const TVector<float, 3>& LocalSpaceLocation,
+			TVector<float, 3>* OutPositions,
+			TVector<float, 3>* OutNormals) const;
+
+	private:
+		const UClothingAssetCommon* Asset;
+		const USkeletalMeshComponent* SkeletalMeshComponent;
+	};
+} // namespace Chaos
diff --git a/Engine/Plugins/Experimental/ChaosCloth/Source/Chaos/Public/ChaosCloth/ChaosClothingSimulationSolver.h b/Engine/Plugins/Experimental/ChaosCloth/Source/Chaos/Public/ChaosCloth/ChaosClothingSimulationSolver.h
new file mode 100644
--- /dev/null
+++ b/Engine/Plugins/Experimental/ChaosCloth/Source/Chaos/Public/ChaosCloth/ChaosClothingSimulationSolver.h
@@ -0,0 +1,190 @@
+// Copyright Epic Games, Inc. All Rights Reserved.
+#pragma once
+
+#include "ChaosCloth/ChaosClothConstraints.h"
+#include "Chaos/ArrayCollectionArray.h"
+#include "Chaos/Transform.h"
+
+class USkeletalMeshComponent;
+class UClothingAssetCommon;
+
+namespace Chaos
+{
+	template<typename T, int d> class TPBDEvolution;
+	template<typename T> class TTriangleMesh;
+
+	class FClothingSimulationCloth;
+	class FClothingSimulationCollider;
+	class FClothingSimulationMesh;
+
+	// Solver simulation node
+	class FClothingSimulationSolver final
+	{
+	public:
+		FClothingSimulationSolver();
+		~FClothingSimulationSolver();
+
+		// ---- Animatable property setters ----
+		void SetLocalSpaceLocation(const TVector<float, 3>& InLocalSpaceLocation) { LocalSpaceLocation = InLocalSpaceLocation; }
+		const TVector<float, 3>& GetLocalSpaceLocation() const { return LocalSpaceLocation; }
+
+		// Disables all Cloths gravity override mechanism
+		void EnableClothGravityOverride(bool bInIsClothGravityOverrideEnabled) { bIsClothGravityOverrideEnabled = bInIsClothGravityOverrideEnabled; }
+		bool IsClothGravityOverrideEnabled() const { return bIsClothGravityOverrideEnabled; }
+		void SetGravity(const TVector<float, 3>& InGravity) { Gravity = InGravity; }
+		const TVector<float, 3>& GetGravity() const { return Gravity; }
+
+		void SetWindVelocity(const TVector<float, 3>& InWindVelocity);
+		void SetWindFluidDensity(float InWindFluidDensity) { WindFluidDensity = InWindFluidDensity; }
+
+		void SetNumIterations(int32 InNumIterations) { NumIterations = InNumIterations; }
+		int32 GetNumIterations() const { return NumIterations; }
+		void SetNumSubsteps(int32 InNumSubsteps) { NumSubsteps = InNumSubsteps; }
+		int32 GetNumSubsteps() const { return NumSubsteps; }
+		// ---- End of the animatable property setters ----
+
+		// ---- Object management functions ----
+		void SetCloths(TArray<FClothingSimulationCloth*>&& InCloths); 
+		void AddCloth(FClothingSimulationCloth* InCloth);
+		void RemoveCloth(FClothingSimulationCloth* InCloth);
+		void RemoveCloths();
+
+		void SetColliders(TArray<FClothingSimulationCollider*>&& InColliders);
+		void AddCollider(FClothingSimulationCollider* InCollider);
+		void RemoveCollider(FClothingSimulationCollider* InCollider);
+		void RemoveColliders();
+
+		TConstArrayView<const FClothingSimulationCloth*> GetCloths() const { return Cloths; }
+		TConstArrayView<const FClothingSimulationCollider*> GetColliders() const { return Colliders; }
+
+		// Update solver properties before simulation
+		void Update(float InDeltaTime);
+
+		FBoxSphereBounds CalculateBounds() const;
+		// ---- End of the object management functions ----
+
+		// ---- Cloth interface ----
+		int32 AddParticles(int32 NumParticles, uint32 GroupId);
+		void EnableParticles(int32 Offset, bool bEnable);
+
+		void SetParticleMassUniform(int32 Offset, float UniformMass, float MinPerParticleMass, const TTriangleMesh<float>& Mesh, const TFunctionRef<bool(int32)>& KinematicPredicate);
+		void SetParticleMassFromTotalMass(int32 Offset, float TotalMass, float MinPerParticleMass, const TTriangleMesh<float>& Mesh, const TFunctionRef<bool(int32)>& KinematicPredicate);
+		void SetParticleMassFromDensity(int32 Offset, float Density, float MinPerParticleMass, const TTriangleMesh<float>& Mesh, const TFunctionRef<bool(int32)>& KinematicPredicate);
+
+		// Set the amount of velocity allowed to filter from the given change in reference space transform, including local simulation space
+		void SetReferenceVelocityScale(uint32 GroupId,
+			const TRigidTransform<float, 3>& OldReferenceSpaceTransform,
+			const TRigidTransform<float, 3>& ReferenceSpaceTransform,
+			const TVector<float, 3>& LinearVelocityScale,
+			float AngularVelocityScale);
+
+		// Set general cloth simulation properties
+		void SetProperties(uint32 GroupId, float DampingCoefficient, float CollisionThickness, float FrictionCoefficient);
+
+		// Set per group gravity, used to override solver's gravity. Must be called during cloth update.
+		void SetGravity(uint32 GroupId, const TVector<float, 3>& Gravity);
+
+		// Set the geometry affected by wind, or disable if TriangleMesh is null.
+		void SetWindVelocityField(uint32 GroupId, float DragCoefficient, const TTriangleMesh<float>* TriangleMesh = nullptr);
+
+		// Enable self collision, or disable if TriangleMesh is null.
+		void SetSelfCollisions(uint32 GroupId, float SelfCollisionThickness, const TTriangleMesh<float>* TriangleMesh = nullptr);
+
+		const TVector<float, 3>* GetOldAnimationPositions(int32 Offset) const { return OldAnimationPositions.GetData() + Offset; }
+		TVector<float, 3>* GetOldAnimationPositions(int32 Offset) { return OldAnimationPositions.GetData() + Offset; }
+		const TVector<float, 3>* GetAnimationPositions(int32 Offset) const { return AnimationPositions.GetData() + Offset; }
+		TVector<float, 3>* GetAnimationPositions(int32 Offset) { return AnimationPositions.GetData() + Offset; }
+		const TVector<float, 3>* GetAnimationNormals(int32 Offset) const { return AnimationNormals.GetData() + Offset; }
+		TVector<float, 3>* GetAnimationNormals(int32 Offset) { return AnimationNormals.GetData() + Offset; }
+		const TVector<float, 3>* GetNormals(int32 Offset) const { return Normals.GetData() + Offset; }
+		TVector<float, 3>* GetNormals(int32 Offset) { return Normals.GetData() + Offset; }
+		const TVector<float, 3>* GetParticlePs(int32 Offset) const;
+		TVector<float, 3>* GetParticlePs(int32 Offset);
+		const TVector<float, 3>* GetParticleXs(int32 Offset) const;
+		TVector<float, 3>* GetParticleXs(int32 Offset);
+		const TVector<float, 3>* GetParticleVs(int32 Offset) const;
+		TVector<float, 3>* GetParticleVs(int32 Offset);
+		const float* GetParticleInvMasses(int32 Offset) const;
+		const FClothConstraints& GetClothConstraints(int32 Offset) const { return ClothsConstraints.FindChecked(Offset);  }
+		FClothConstraints& GetClothConstraints(int32 Offset) { return ClothsConstraints.FindChecked(Offset);  }
+		// ---- End of the Cloth interface ----
+
+		// ---- Collider interface ----
+		int32 AddCollisionParticles(int32 NumCollisionParticles, uint32 GroupId, int32 RecycledOffset = 0);
+		void EnableCollisionParticles(int32 Offset, bool bEnable);
+
+		const int32* GetCollisionBoneIndices(int32 Offset) const { return CollisionBoneIndices.GetData() + Offset; }
+		int32* GetCollisionBoneIndices(int32 Offset) { return CollisionBoneIndices.GetData() + Offset; }
+		const TRigidTransform<float, 3>* GetCollisionBaseTransforms(int32 Offset) const { return CollisionBaseTransforms.GetData() + Offset; }
+		TRigidTransform<float, 3>* GetCollisionBaseTransforms(int32 Offset) { return CollisionBaseTransforms.GetData() + Offset; }
+		const TRigidTransform<float, 3>* GetOldCollisionTransforms(int32 Offset) const { return OldCollisionTransforms.GetData() + Offset; }
+		TRigidTransform<float, 3>* GetOldCollisionTransforms(int32 Offset) { return OldCollisionTransforms.GetData() + Offset; }
+		const TRigidTransform<float, 3>* GetCollisionTransforms(int32 Offset) const { return CollisionTransforms.GetData() + Offset; }
+		TRigidTransform<float, 3>* GetCollisionTransforms(int32 Offset) { return CollisionTransforms.GetData() + Offset; }
+		const TVector<float, 3>* GetCollisionParticleXs(int32 Offset) const;
+		TVector<float, 3>* GetCollisionParticleXs(int32 Offset);
+		const TRotation<float, 3>* GetCollisionParticleRs(int32 Offset) const;
+		TRotation<float, 3>* GetCollisionParticleRs(int32 Offset);
+		void SetCollisionGeometry(int32 Offset, int32 Index, TUniquePtr<FImplicitObject>&& Geometry);
+		const TUniquePtr<FImplicitObject>* GetCollisionGeometries(int32 Offset) const;
+		// ---- End of the Collider interface ----
+
+	private:
+		void ResetParticles();
+		void ApplyPreSimulationTransforms();
+		float SetParticleMassPerArea(int32 Offset, int32 Size, const TTriangleMesh<float>& Mesh);
+		void ParticleMassUpdateDensity(const TTriangleMesh<float>& Mesh, float Density);
+		void ParticleMassClampAndEnslave(int32 Offset, int32 Size, float MinPerParticleMass, const TFunctionRef<bool(int32)>& KinematicPredicate);
+
+	private:
+		TUniquePtr<TPBDEvolution<float, 3>> Evolution;
+
+		// Object arrays
+		TArray<FClothingSimulationCloth*> Cloths;
+		TArray<FClothingSimulationCollider*> Colliders;
+
+		// Simulation group attributes
+		TArrayCollectionArray<TRigidTransform<float, 3>> PreSimulationTransforms;  // Allow a different frame of reference for each cloth groups
+
+		// Particle attributes
+		TArrayCollectionArray<TVector<float, 3>> Normals;
+		TArrayCollectionArray<TVector<float, 3>> OldAnimationPositions;
+		TArrayCollectionArray<TVector<float, 3>> AnimationPositions;
+		TArrayCollectionArray<TVector<float, 3>> AnimationNormals;
+
+		// Collision particle attributes
+		TArrayCollectionArray<int32> CollisionBoneIndices;
+		TArrayCollectionArray<TRigidTransform<float, 3>> CollisionBaseTransforms;
+		TArrayCollectionArray<TRigidTransform<float, 3>> OldCollisionTransforms;
+		TArrayCollectionArray<TRigidTransform<float, 3>> CollisionTransforms;
+
+		// Cloth constraints
+		TMap<int32, FClothConstraints> ClothsConstraints;
+
+		// Local space simulation
+		TVector<float, 3> OldLocalSpaceLocation;
+		TVector<float, 3> LocalSpaceLocation;  // This is used to translate between world space and simulation space. Add this to simulation space coordinates to get world space coordinates
+
+		// Time stepping
+		float Time;
+		float DeltaTime;
+		int32 NumIterations;
+		int32 NumSubsteps;
+
+		// Solver colliders offset
+		int32 CollisionParticlesOffset;  // Collision particle offset on the first solver/non cloth collider
+		int32 CollisionParticlesSize;  // Number of solver only colliders
+		bool bAreCollisionParticlesDirty;  // Whereas adding/re-add a solver collider has triggered a new particle allocation
+
+		// Active particle view
+		TArray<TPair<FClothingSimulationCloth*, TVector<uint32, 2>>> ClothsMap;
+		TArray<TPair<FClothingSimulationCollider*, TVector<uint32, 2>>> CollidersMap;
+
+		// Solver parameters
+		TVector<float, 3> Gravity;
+		TVector<float, 3> WindVelocity;
+		float WindFluidDensity;
+		bool bIsClothGravityOverrideEnabled;
+	};
+
+} // namespace Chaos
diff --git a/Engine/Plugins/Experimental/ChaosClothEditor/Source/Chaos/Private/ChaosClothEditor/ChaosSimulationEditorExtender.cpp b/Engine/Plugins/Experimental/ChaosClothEditor/Source/Chaos/Private/ChaosClothEditor/ChaosSimulationEditorExtender.cpp
--- a/Engine/Plugins/Experimental/ChaosClothEditor/Source/Chaos/Private/ChaosClothEditor/ChaosSimulationEditorExtender.cpp
+++ b/Engine/Plugins/Experimental/ChaosClothEditor/Source/Chaos/Private/ChaosClothEditor/ChaosSimulationEditorExtender.cpp
@@ -23,8 +23,8 @@ namespace Chaos
 		static const uint32 Count;
 
 		// Chaos debug draw function
-		typedef void (Chaos::ClothingSimulation::*FDebugDrawFunction)(USkeletalMeshComponent*, FPrimitiveDrawInterface*) const;
-		typedef void (Chaos::ClothingSimulation::*FDebugDrawTextsFunction)(USkeletalMeshComponent*, FCanvas*, const FSceneView*) const;
+		typedef void (Chaos::FClothingSimulation::*FDebugDrawFunction)(FPrimitiveDrawInterface*) const;
+		typedef void (Chaos::FClothingSimulation::*FDebugDrawTextsFunction)(FCanvas*, const FSceneView*) const;
 		FDebugDrawFunction DebugDrawFunction;
 		FDebugDrawTextsFunction DebugDrawTextsFunction;
 
@@ -54,21 +54,21 @@ using namespace Chaos;
 
 const FVisualizationOption FVisualizationOption::OptionData[] = 
 {
-	FVisualizationOption(&Chaos::ClothingSimulation::DebugDrawPhysMeshShaded      , LOCTEXT("ChaosVisName_PhysMesh"            , "Physical Mesh (Flat Shaded)"), LOCTEXT("ChaosVisName_PhysMeshShaded_ToolTip"      , "Draws the current physical result as a doubled sided flat shaded mesh"), /*bDisablesSimulation =*/false, /*bHidesClothSections=*/true),
-	FVisualizationOption(&Chaos::ClothingSimulation::DebugDrawPhysMeshWired       , LOCTEXT("ChaosVisName_PhysMeshWire"        , "Physical Mesh (Wireframe)"  ), LOCTEXT("ChaosVisName_PhysMeshWired_ToolTip"       , "Draws the current physical mesh result in wireframe")),
-	FVisualizationOption(&Chaos::ClothingSimulation::DebugDrawParticleIndices     , LOCTEXT("ChaosVisName_ParticleIndices"     , "Particle Indices"           ), LOCTEXT("ChaosVisName_ParticleIndices_ToolTip"     , "Draws the particle indices as instanciated by the solver")),
-	FVisualizationOption(&Chaos::ClothingSimulation::DebugDrawPointNormals        , LOCTEXT("ChaosVisName_PointNormals"        , "Point Normals"              ), LOCTEXT("ChaosVisName_PointNormals_ToolTip"        , "Draws the current point normals for the simulation mesh")),
-	FVisualizationOption(&Chaos::ClothingSimulation::DebugDrawInversedPointNormals, LOCTEXT("ChaosVisName_InversedPointNormals", "Inversed Point Normals"     ), LOCTEXT("ChaosVisName_InversedPointNormals_ToolTip", "Draws the inversed point normals for the simulation mesh")),
-	FVisualizationOption(&Chaos::ClothingSimulation::DebugDrawFaceNormals         , LOCTEXT("ChaosVisName_FaceNormals"         , "Face Normals"               ), LOCTEXT("ChaosVisName_FaceNormals_ToolTip"         , "Draws the current face normals for the simulation mesh")),
-	FVisualizationOption(&Chaos::ClothingSimulation::DebugDrawInversedFaceNormals , LOCTEXT("ChaosVisName_InversedFaceNormals" , "Inversed Face Normals"      ), LOCTEXT("ChaosVisName_InversedFaceNormals_ToolTip" , "Draws the inversed face normals for the simulation mesh")),
-	FVisualizationOption(&Chaos::ClothingSimulation::DebugDrawCollision           , LOCTEXT("ChaosVisName_Collision"           , "Collisions"                 ), LOCTEXT("ChaosVisName_Collision_ToolTip"           , "Draws the collision bodies the simulation is currently using")),
-	FVisualizationOption(&Chaos::ClothingSimulation::DebugDrawBackstops           , LOCTEXT("ChaosVisName_Backstop"            , "Backstops"                  ), LOCTEXT("ChaosVisName_Backstop_ToolTip"            , "Draws the backstop offset for each simulation particle"), /*bDisablesSimulation =*/true),
-	FVisualizationOption(&Chaos::ClothingSimulation::DebugDrawMaxDistances        , LOCTEXT("ChaosVisName_MaxDistance"         , "Max Distances"              ), LOCTEXT("ChaosVisName_MaxDistance_ToolTip"         , "Draws the current max distances for the sim particles as a line along its normal"), /*bDisablesSimulation =*/true),
-	FVisualizationOption(&Chaos::ClothingSimulation::DebugDrawMaxDistanceValues   , LOCTEXT("ChaosVisName_MaxDistanceValue"    , "Max Distances As Numbers"   ), LOCTEXT("ChaosVisName_MaxDistanceValue_ToolTip"    , "Draws the current max distances as numbers")),
-	FVisualizationOption(&Chaos::ClothingSimulation::DebugDrawAnimDrive           , LOCTEXT("ChaosVisName_AnimDrive"           , "Anim Drive"                 ), LOCTEXT("ChaosVisName_AnimDrive_Tooltip"           , "Draws the current skinned reference mesh for the simulation which anim drive will attempt to reach if enabled")),
-	FVisualizationOption(&Chaos::ClothingSimulation::DebugDrawLongRangeConstraint , LOCTEXT("ChaosVisName_LongRangeConstraint" , "Long Range Constraint"      ), LOCTEXT("ChaosVisName_LongRangeConstraint_Tooltip" , "Draws the long range attachment constraint distances")),
-	FVisualizationOption(&Chaos::ClothingSimulation::DebugDrawWindDragForces      , LOCTEXT("ChaosVisName_WindDragForces"      , "Wind drag forces"           ), LOCTEXT("ChaosVisName_Wind_Tooltip"                , "Draws the Wind drag forces")),
-	FVisualizationOption(&Chaos::ClothingSimulation::DebugDrawLocalSpace          , LOCTEXT("ChaosVisName_LocalSpace"          , "Local space reference bone" ), LOCTEXT("ChaosVisName_LocalSpace_Tooltip"          , "Draws the local space reference bone")),
+	FVisualizationOption(&Chaos::FClothingSimulation::DebugDrawPhysMeshShaded      , LOCTEXT("ChaosVisName_PhysMesh"            , "Physical Mesh (Flat Shaded)"), LOCTEXT("ChaosVisName_PhysMeshShaded_ToolTip"      , "Draws the current physical result as a doubled sided flat shaded mesh"), /*bDisablesSimulation =*/false, /*bHidesClothSections=*/true),
+	FVisualizationOption(&Chaos::FClothingSimulation::DebugDrawPhysMeshWired       , LOCTEXT("ChaosVisName_PhysMeshWire"        , "Physical Mesh (Wireframe)"  ), LOCTEXT("ChaosVisName_PhysMeshWired_ToolTip"       , "Draws the current physical mesh result in wireframe")),
+	FVisualizationOption(&Chaos::FClothingSimulation::DebugDrawAnimMeshWired       , LOCTEXT("ChaosVisName_AnimMeshWire"        , "Animated Mesh (Wireframe)"  ), LOCTEXT("ChaosVisName_AnimMeshWired_ToolTip"       , "Draws the current animated mesh input in wireframe")),
+	FVisualizationOption(&Chaos::FClothingSimulation::DebugDrawParticleIndices     , LOCTEXT("ChaosVisName_ParticleIndices"     , "Particle Indices"           ), LOCTEXT("ChaosVisName_ParticleIndices_ToolTip"     , "Draws the particle indices as instantiated by the solver")),
+	FVisualizationOption(&Chaos::FClothingSimulation::DebugDrawPointNormals        , LOCTEXT("ChaosVisName_PointNormals"        , "Point Normals"              ), LOCTEXT("ChaosVisName_PointNormals_ToolTip"        , "Draws the current point normals for the simulation mesh")),
+	FVisualizationOption(&Chaos::FClothingSimulation::DebugDrawInversedPointNormals, LOCTEXT("ChaosVisName_InversedPointNormals", "Inversed Point Normals"     ), LOCTEXT("ChaosVisName_InversedPointNormals_ToolTip", "Draws the inversed point normals for the simulation mesh")),
+	FVisualizationOption(&Chaos::FClothingSimulation::DebugDrawCollision           , LOCTEXT("ChaosVisName_Collision"           , "Collisions"                 ), LOCTEXT("ChaosVisName_Collision_ToolTip"           , "Draws the collision bodies the simulation is currently using")),
+	FVisualizationOption(&Chaos::FClothingSimulation::DebugDrawBackstops           , LOCTEXT("ChaosVisName_Backstop"            , "Backstops"                  ), LOCTEXT("ChaosVisName_Backstop_ToolTip"            , "Draws the backstop radius and position for each simulation particle")),
+	FVisualizationOption(&Chaos::FClothingSimulation::DebugDrawBackstopDistances   , LOCTEXT("ChaosVisName_BackstopDistance"    , "Backstop Distances"         ), LOCTEXT("ChaosVisName_BackstopDistance_ToolTip"    , "Draws the backstop distance offset for each simulation particle"), /*bDisablesSimulation =*/true),
+	FVisualizationOption(&Chaos::FClothingSimulation::DebugDrawMaxDistances        , LOCTEXT("ChaosVisName_MaxDistance"         , "Max Distances"              ), LOCTEXT("ChaosVisName_MaxDistance_ToolTip"         , "Draws the current max distances for the sim particles as a line along its normal"), /*bDisablesSimulation =*/true),
+	FVisualizationOption(&Chaos::FClothingSimulation::DebugDrawMaxDistanceValues   , LOCTEXT("ChaosVisName_MaxDistanceValue"    , "Max Distances As Numbers"   ), LOCTEXT("ChaosVisName_MaxDistanceValue_ToolTip"    , "Draws the current max distances as numbers")),
+	FVisualizationOption(&Chaos::FClothingSimulation::DebugDrawAnimDrive           , LOCTEXT("ChaosVisName_AnimDrive"           , "Anim Drive"                 ), LOCTEXT("ChaosVisName_AnimDrive_Tooltip"           , "Draws the current skinned reference mesh for the simulation which anim drive will attempt to reach if enabled")),
+	FVisualizationOption(&Chaos::FClothingSimulation::DebugDrawLongRangeConstraint , LOCTEXT("ChaosVisName_LongRangeConstraint" , "Long Range Constraint"      ), LOCTEXT("ChaosVisName_LongRangeConstraint_Tooltip" , "Draws the long range attachment constraint distances")),
+	FVisualizationOption(&Chaos::FClothingSimulation::DebugDrawWindDragForces      , LOCTEXT("ChaosVisName_WindDragForces"      , "Wind drag forces"           ), LOCTEXT("ChaosVisName_Wind_Tooltip"                , "Draws the Wind drag forces")),
+	FVisualizationOption(&Chaos::FClothingSimulation::DebugDrawLocalSpace          , LOCTEXT("ChaosVisName_LocalSpace"          , "Local space reference bone" ), LOCTEXT("ChaosVisName_LocalSpace_Tooltip"          , "Draws the local space reference bone")),
 };
 const uint32 FVisualizationOption::Count = sizeof(OptionData) / sizeof(FVisualizationOption);
 
@@ -125,32 +125,32 @@ void FSimulationEditorExtender::ExtendViewportShowMenu(FMenuBuilder& MenuBuilder
 	MenuBuilder.EndSection();
 }
 
-void FSimulationEditorExtender::DebugDrawSimulation(const IClothingSimulation* Simulation, USkeletalMeshComponent* OwnerComponent, FPrimitiveDrawInterface* PDI)
+void FSimulationEditorExtender::DebugDrawSimulation(const IClothingSimulation* Simulation, USkeletalMeshComponent* /*OwnerComponent*/, FPrimitiveDrawInterface* PDI)
 {
 	if (!ensure(Simulation)) { return; }
 
-	const ClothingSimulation* const ChaosSimulation = static_cast<const ClothingSimulation*>(Simulation);
+	const FClothingSimulation* const ChaosSimulation = static_cast<const FClothingSimulation*>(Simulation);
 
 	for (int32 OptionIndex = 0; OptionIndex < FVisualizationOption::Count; ++OptionIndex)
 	{
 		if (Flags[OptionIndex] && FVisualizationOption::OptionData[OptionIndex].DebugDrawFunction)
 		{
-			(ChaosSimulation->*(FVisualizationOption::OptionData[OptionIndex].DebugDrawFunction))(OwnerComponent, PDI);
+			(ChaosSimulation->*(FVisualizationOption::OptionData[OptionIndex].DebugDrawFunction))(PDI);
 		}
 	}
 }
 
-void FSimulationEditorExtender::DebugDrawSimulationTexts(const IClothingSimulation* Simulation, USkeletalMeshComponent* OwnerComponent, FCanvas* Canvas, const FSceneView* SceneView)
+void FSimulationEditorExtender::DebugDrawSimulationTexts(const IClothingSimulation* Simulation, USkeletalMeshComponent* /*OwnerComponent*/, FCanvas* Canvas, const FSceneView* SceneView)
 {
 	if (!ensure(Simulation)) { return; }
 
-	const ClothingSimulation* const ChaosSimulation = static_cast<const ClothingSimulation*>(Simulation);
+	const FClothingSimulation* const ChaosSimulation = static_cast<const FClothingSimulation*>(Simulation);
 
 	for (int32 OptionIndex = 0; OptionIndex < FVisualizationOption::Count; ++OptionIndex)
 	{
 		if (Flags[OptionIndex] && FVisualizationOption::OptionData[OptionIndex].DebugDrawTextsFunction)
 		{
-			(ChaosSimulation->*(FVisualizationOption::OptionData[OptionIndex].DebugDrawTextsFunction))(OwnerComponent, Canvas, SceneView);
+			(ChaosSimulation->*(FVisualizationOption::OptionData[OptionIndex].DebugDrawTextsFunction))(Canvas, SceneView);
 		}
 	}
 }
diff --git a/Engine/Source/Editor/ClothPainter/Private/SClothAssetSelector.cpp b/Engine/Source/Editor/ClothPainter/Private/SClothAssetSelector.cpp
--- a/Engine/Source/Editor/ClothPainter/Private/SClothAssetSelector.cpp
+++ b/Engine/Source/Editor/ClothPainter/Private/SClothAssetSelector.cpp
@@ -674,6 +674,11 @@ private:
 
 					// Set the flag
 					Mask->bEnabled = bNewEnableState;
+					
+					if(UClothingAssetCommon* Asset = InItem->ClothingAsset.Get())
+					{
+						Asset->ApplyParameterMasks();
+					}
 				}
 			}
 		}
diff --git a/Engine/Source/Runtime/ClothingSystemRuntimeCommon/Private/ClothPhysicalMeshData.cpp b/Engine/Source/Runtime/ClothingSystemRuntimeCommon/Private/ClothPhysicalMeshData.cpp
--- a/Engine/Source/Runtime/ClothingSystemRuntimeCommon/Private/ClothPhysicalMeshData.cpp
+++ b/Engine/Source/Runtime/ClothingSystemRuntimeCommon/Private/ClothPhysicalMeshData.cpp
@@ -76,10 +76,8 @@ void FClothPhysicalMeshData::ClearWeightMaps()
 	// Clear all weight maps (and reserve 4 slots)
 	WeightMaps.Empty(4);
 
-	// Max distances must be present, so fill to zero on clear so we still have valid mesh data.
-	AddWeightMap(EWeightMapTargetCommon::MaxDistance).Initialize(Vertices.Num());
-
 	// Add default (empty) optional maps, as these are always expected to be found
+	AddWeightMap(EWeightMapTargetCommon::MaxDistance);
 	AddWeightMap(EWeightMapTargetCommon::BackstopDistance);
 	AddWeightMap(EWeightMapTargetCommon::BackstopRadius);
 	AddWeightMap(EWeightMapTargetCommon::AnimDriveMultiplier);
diff --git a/Engine/Source/Runtime/ClothingSystemRuntimeCommon/Private/ClothingAsset.cpp b/Engine/Source/Runtime/ClothingSystemRuntimeCommon/Private/ClothingAsset.cpp
--- a/Engine/Source/Runtime/ClothingSystemRuntimeCommon/Private/ClothingAsset.cpp
+++ b/Engine/Source/Runtime/ClothingSystemRuntimeCommon/Private/ClothingAsset.cpp
@@ -1052,6 +1052,17 @@ void UClothingAssetCommon::PostUpdateAllAssets()
 
 void UClothingAssetCommon::InvalidateCachedData()
 {
+#if WITH_CHAOS_CLOTHING
+	ForEachInteractorUsingClothing([](UClothingSimulationInteractor* InInteractor)
+	{
+		if (InInteractor)
+		{
+			InInteractor->ClothConfigUpdated();
+		}
+	});
+#endif  // #if WITH_CHAOS_CLOTHING
+
+	// TODO: This mass calculation isn't used by Chaos, check what can be stripped out
 	for(FClothLODDataCommon& CurrentLodData : LodData)
 	{
 		// Recalculate inverse masses for the physical mesh particles
diff --git a/Engine/Source/Runtime/ClothingSystemRuntimeCommon/Private/ClothingSimulation.cpp b/Engine/Source/Runtime/ClothingSystemRuntimeCommon/Private/ClothingSimulation.cpp
--- a/Engine/Source/Runtime/ClothingSystemRuntimeCommon/Private/ClothingSimulation.cpp
+++ b/Engine/Source/Runtime/ClothingSystemRuntimeCommon/Private/ClothingSimulation.cpp
@@ -21,6 +21,7 @@ FClothingSimulationContextCommon::FClothingSimulationContextCommon()
 	, DeltaSeconds(0.f)
 	, TeleportMode(EClothingTeleportMode::None)
 	, MaxDistanceScale(1.f)
+	, PredictedLod(INDEX_NONE)
 {}
 
 FClothingSimulationContextCommon::~FClothingSimulationContextCommon()
@@ -37,6 +38,8 @@ void FClothingSimulationContextCommon::Fill(const USkeletalMeshComponent* InComp
 	FillDeltaSeconds(InDeltaSeconds, InMaxPhysicsDelta);
 	FillTeleportMode(InComponent, InDeltaSeconds, InMaxPhysicsDelta);
 	FillMaxDistanceScale(InComponent);
+
+	PredictedLod = InComponent->PredictedLODLevel;
 }
 
 void FClothingSimulationContextCommon::FillBoneTransforms(const USkeletalMeshComponent* InComponent)
@@ -99,7 +102,7 @@ void FClothingSimulationContextCommon::FillBoneTransforms(const USkeletalMeshCom
 void FClothingSimulationContextCommon::FillRefToLocals(const USkeletalMeshComponent* InComponent)
 {
 	RefToLocals.Reset();
-	InComponent->GetCurrentRefToLocalMatrices(RefToLocals, 0);
+	InComponent->GetCurrentRefToLocalMatrices(RefToLocals, InComponent->PredictedLODLevel);
 }
 
 void FClothingSimulationContextCommon::FillComponentToWorld(const USkeletalMeshComponent* InComponent)
diff --git a/Engine/Source/Runtime/ClothingSystemRuntimeCommon/Public/ClothingSimulation.h b/Engine/Source/Runtime/ClothingSystemRuntimeCommon/Public/ClothingSimulation.h
--- a/Engine/Source/Runtime/ClothingSystemRuntimeCommon/Public/ClothingSimulation.h
+++ b/Engine/Source/Runtime/ClothingSystemRuntimeCommon/Public/ClothingSimulation.h
@@ -33,7 +33,7 @@ protected:
 	float SetWindFromComponent(const USkeletalMeshComponent* Component);
 
 public:
-	// World space bone transforms of the owning component
+	// Component space bone transforms of the owning component
 	TArray<FTransform> BoneTransforms;
 
 	// Ref to local matrices from the owning component (for skinning fixed verts)
@@ -56,6 +56,9 @@ public:
 
 	// Scale for the max distance constraints of the simulation mesh
 	float MaxDistanceScale;
+
+	// The predicted LOD of the skeletal mesh component running the simulation
+	int32 PredictedLod;
 };
 
 // Base simulation to fill in common data for the base context
diff --git a/Engine/Source/Runtime/ClothingSystemRuntimeInterface/Public/ClothingSimulationInterface.h b/Engine/Source/Runtime/ClothingSystemRuntimeInterface/Public/ClothingSimulationInterface.h
--- a/Engine/Source/Runtime/ClothingSystemRuntimeInterface/Public/ClothingSimulationInterface.h
+++ b/Engine/Source/Runtime/ClothingSystemRuntimeInterface/Public/ClothingSimulationInterface.h
@@ -47,8 +47,9 @@ protected:
 	virtual void CreateActor(USkeletalMeshComponent* InOwnerComponent, UClothingAssetBase* InAsset, int32 SimDataIndex) = 0;
 
 	/** Do any required initialization after all the actors have been loaded */
-	virtual void PostActorCreationInitialize() = 0;
-	
+	UE_DEPRECATED(4.26, "No longer required.")
+	virtual void PostActorCreationInitialize() {}
+
 	/** Create a new context, will not be filled, call FillContext before simulating with this context */
 	virtual IClothingSimulationContext* CreateContext() = 0;
 
diff --git a/Engine/Source/Runtime/ClothingSystemRuntimeNv/Private/ClothingSimulationNv.cpp b/Engine/Source/Runtime/ClothingSystemRuntimeNv/Private/ClothingSimulationNv.cpp
--- a/Engine/Source/Runtime/ClothingSystemRuntimeNv/Private/ClothingSimulationNv.cpp
+++ b/Engine/Source/Runtime/ClothingSystemRuntimeNv/Private/ClothingSimulationNv.cpp
@@ -39,8 +39,7 @@ DECLARE_CYCLE_STAT(TEXT("Update Anim Drive"), STAT_NvClothUpdateAnimDrive, STATG
 //=============================================================================
 
 FClothingSimulationContextNv::FClothingSimulationContextNv()
-	: PredictedLod(0)
-	, WindAdaption(0.0f)
+	: WindAdaption(0.0f)
 {
 }
 
@@ -54,14 +53,6 @@ void FClothingSimulationContextNv::Fill(const USkeletalMeshComponent* InComponen
 	LLM_SCOPE(ELLMTag::SkeletalMesh);
 
 	FClothingSimulationContextCommon::Fill(InComponent, InDeltaSeconds, InMaxPhysicsDelta);
-
-	PredictedLod = InComponent->PredictedLODLevel;
-}
-
-void FClothingSimulationContextNv::FillRefToLocals(const USkeletalMeshComponent* InComponent)
-{
-	FClothingSimulationContextCommon::RefToLocals.Reset();
-	InComponent->GetCurrentRefToLocalMatrices(FClothingSimulationContextCommon::RefToLocals, InComponent->PredictedLODLevel);
 }
 
 void FClothingSimulationContextNv::FillWorldGravity(const USkeletalMeshComponent* InComponent)
