Chaos Cloth - Fix RefreshPhysicsAsset and RefreshClothConfig that was asserting due to missing code to remove the solver pointer from the cloth and colllider maps.

diff --git a/Engine/Plugins/Experimental/ChaosCloth/Source/Chaos/Private/ChaosCloth/ChaosClothingSimulation.cpp b/Engine/Plugins/Experimental/ChaosCloth/Source/Chaos/Private/ChaosCloth/ChaosClothingSimulation.cpp
--- a/Engine/Plugins/Experimental/ChaosCloth/Source/Chaos/Private/ChaosCloth/ChaosClothingSimulation.cpp
+++ b/Engine/Plugins/Experimental/ChaosCloth/Source/Chaos/Private/ChaosCloth/ChaosClothingSimulation.cpp
@@ -385,7 +385,7 @@ void FClothingSimulation::GetCollisions(FClothCollisionData& OutCollisions, bool
 	{
 		for (const FClothingSimulationCollider* const Collider : Cloth->GetColliders())
 		{
-			OutCollisions.Append(Collider->GetCollisionData());
+			OutCollisions.Append(Collider->GetCollisionData(Solver.Get(), Cloth.Get()));
 		}
 	}
 
diff --git a/Engine/Plugins/Experimental/ChaosCloth/Source/Chaos/Private/ChaosCloth/ChaosClothingSimulationCloth.cpp b/Engine/Plugins/Experimental/ChaosCloth/Source/Chaos/Private/ChaosCloth/ChaosClothingSimulationCloth.cpp
--- a/Engine/Plugins/Experimental/ChaosCloth/Source/Chaos/Private/ChaosCloth/ChaosClothingSimulationCloth.cpp
+++ b/Engine/Plugins/Experimental/ChaosCloth/Source/Chaos/Private/ChaosCloth/ChaosClothingSimulationCloth.cpp
@@ -261,6 +261,11 @@ void FClothingSimulationCloth::FLODData::Add(FClothingSimulationSolver* Solver,
 	}
 }
 
+void FClothingSimulationCloth::FLODData::Remove(FClothingSimulationSolver* Solver)
+{
+	SolverData.Remove(Solver);
+}
+
 void FClothingSimulationCloth::FLODData::Enable(FClothingSimulationSolver* Solver, bool bEnable) const
 {
 	check(Solver);
@@ -461,9 +466,11 @@ void FClothingSimulationCloth::RemoveColliders()
 void FClothingSimulationCloth::Add(FClothingSimulationSolver* Solver)
 {
 	check(Solver);
+
+	// Can't add a cloth twice to the same solver
 	check(!LODIndices.Find(Solver));
 
-	// Reset LODIndex
+	// Initialize LODIndex
 	int32& LODIndex = LODIndices.Add(Solver);
 	LODIndex = INDEX_NONE;
 
@@ -482,8 +489,12 @@ void FClothingSimulationCloth::Add(FClothingSimulationSolver* Solver)
 
 void FClothingSimulationCloth::Remove(FClothingSimulationSolver* Solver)
 {
-	// Remove LODIndex from map
+	// Remove solver from maps
 	LODIndices.Remove(Solver);
+	for (FLODData& LODDatum: LODData)
+	{
+		LODDatum.Remove(Solver);
+	}
 
 	// Remove Colliders
 	for (FClothingSimulationCollider* Collider : Colliders)
@@ -494,12 +505,12 @@ void FClothingSimulationCloth::Remove(FClothingSimulationSolver* Solver)
 
 int32 FClothingSimulationCloth::GetNumParticles(int32 InLODIndex) const
 {
-	return (LODData.IsValidIndex(InLODIndex)) ? LODData[InLODIndex].NumParticles : 0;
+	return LODData.IsValidIndex(InLODIndex) ? LODData[InLODIndex].NumParticles : 0;
 }
 
 int32 FClothingSimulationCloth::GetOffset(const FClothingSimulationSolver* Solver, int32 InLODIndex) const
 {
-	return (LODData.IsValidIndex(InLODIndex)) ? LODData[InLODIndex].SolverData.FindChecked(Solver).Offset : 0;
+	return LODData.IsValidIndex(InLODIndex) ? LODData[InLODIndex].SolverData.FindChecked(Solver).Offset : 0;
 }
 
 TVector<float, 3> FClothingSimulationCloth::GetGravity(const FClothingSimulationSolver* Solver) const
diff --git a/Engine/Plugins/Experimental/ChaosCloth/Source/Chaos/Private/ChaosCloth/ChaosClothingSimulationCollider.cpp b/Engine/Plugins/Experimental/ChaosCloth/Source/Chaos/Private/ChaosCloth/ChaosClothingSimulationCollider.cpp
--- a/Engine/Plugins/Experimental/ChaosCloth/Source/Chaos/Private/ChaosCloth/ChaosClothingSimulationCollider.cpp
+++ b/Engine/Plugins/Experimental/ChaosCloth/Source/Chaos/Private/ChaosCloth/ChaosClothingSimulationCollider.cpp
@@ -56,11 +56,10 @@ void FClothingSimulationCollider::FLODData::Add(
 
 	// The offset will be set to the first collision particle's index
 	// Try to reuse existing offsets when Add is called during the solver's collider update
-	const TPair<const FClothingSimulationSolver*, const FClothingSimulationCloth*> SolverClothPair(Solver, Cloth);
-	int32* const OffsetPtr = Offsets.Find(SolverClothPair);
+	int32* const OffsetPtr = Offsets.Find(FSolverClothPair(Solver, Cloth));
 	const bool bIsNewCollider = !OffsetPtr;
 	check(!Cloth || bIsNewCollider);  // TODO: We need to be able to update cloth colliders too at some point, but this won't work as it is
-	int32& Offset = bIsNewCollider ? Offsets.Add(SolverClothPair) : *OffsetPtr;
+	int32& Offset = bIsNewCollider ? Offsets.Add(FSolverClothPair(Solver, Cloth)) : *OffsetPtr;
 	Offset = Solver->AddCollisionParticles(NumGeometries, GroupId, bIsNewCollider ? INDEX_NONE : Offset);
 
 	// Capsules
@@ -271,13 +270,17 @@ void FClothingSimulationCollider::FLODData::Add(
 	UE_LOG(LogChaosCloth, VeryVerbose, TEXT("Added collisions: %d spheres, %d capsules, %d convexes, %d boxes."), NumSpheres, NumCapsules, NumConvexes, NumBoxes);
 }
 
+void FClothingSimulationCollider::FLODData::Remove(FClothingSimulationSolver* Solver, FClothingSimulationCloth* Cloth)
+{
+	Offsets.Remove(FSolverClothPair(Solver, Cloth));
+}
+
 void FClothingSimulationCollider::FLODData::Update(FClothingSimulationSolver* Solver, FClothingSimulationCloth* Cloth, const FClothingSimulationContextCommon* Context)
 {
 	check(Solver);
 	if (NumGeometries)
 	{
-		const TPair<const FClothingSimulationSolver*, const FClothingSimulationCloth*> SolverClothPair(Solver, Cloth);
-		const int32 Offset = Offsets.FindChecked(SolverClothPair);
+		const int32 Offset = Offsets.FindChecked(FSolverClothPair(Solver, Cloth));
 		const int32* const BoneIndices = Solver->GetCollisionBoneIndices(Offset);
 		const TRigidTransform<float, 3>* BaseTransforms = Solver->GetCollisionBaseTransforms(Offset);
 		TRigidTransform<float, 3>* const CollisionTransforms = Solver->GetCollisionTransforms(Offset);
@@ -303,8 +306,7 @@ void FClothingSimulationCollider::FLODData::Enable(FClothingSimulationSolver* So
 	check(Solver);
 	if (NumGeometries)
 	{
-		const TPair<const FClothingSimulationSolver*, const FClothingSimulationCloth*> SolverClothPair(Solver, Cloth);
-		const int32 Offset = Offsets.FindChecked(SolverClothPair);
+		const int32 Offset = Offsets.FindChecked(FSolverClothPair(Solver, Cloth));
 		Solver->EnableCollisionParticles(Offset, bEnable);
 	}
 }
@@ -314,8 +316,7 @@ void FClothingSimulationCollider::FLODData::ResetStartPose(FClothingSimulationSo
 	check(Solver);
 	if (NumGeometries)
 	{
-		const TPair<const FClothingSimulationSolver*, const FClothingSimulationCloth*> SolverClothPair(Solver, Cloth);
-		const int32 Offset = Offsets.FindChecked(SolverClothPair);
+		const int32 Offset = Offsets.FindChecked(FSolverClothPair(Solver, Cloth));
 		const TRigidTransform<float, 3>* const CollisionTransforms = Solver->GetCollisionTransforms(Offset);
 		TRigidTransform<float, 3>* const OldCollisionTransforms = Solver->GetOldCollisionTransforms(Offset);
 		TRotation<float, 3>* const Rs = Solver->GetCollisionParticleRs(Offset);
@@ -340,7 +341,6 @@ FClothingSimulationCollider::FClothingSimulationCollider(
 	, CollisionData(nullptr)
 	, bUseLODIndexOverride(bInUseLODIndexOverride)
 	, LODIndexOverride(InLODIndexOverride)
-	, LODIndex(INDEX_NONE)
 {
 	// Prepare LOD array
 	const int32 NumLODs = Asset ? Asset->LodData.Num() : 0;
@@ -351,11 +351,13 @@ FClothingSimulationCollider::~FClothingSimulationCollider()
 {
 }
 
-FClothCollisionData FClothingSimulationCollider::GetCollisionData() const
+FClothCollisionData FClothingSimulationCollider::GetCollisionData(const FClothingSimulationSolver* Solver, const FClothingSimulationCloth* Cloth) const
 {
 	FClothCollisionData ClothCollisionData;
 	ClothCollisionData.Append(LODData[(int32)ECollisionDataType::Global].ClothCollisionData);
 	ClothCollisionData.Append(LODData[(int32)ECollisionDataType::Dynamic].ClothCollisionData);
+
+	const int32 LODIndex = LODIndices.FindChecked(FSolverClothPair(Solver, Cloth));
 	if (LODIndex >= (int32)ECollisionDataType::LODs)
 	{
 		ClothCollisionData.Append(LODData[LODIndex].ClothCollisionData);
@@ -540,12 +542,27 @@ void FClothingSimulationCollider::ExtractPhysicsAssetCollision(FClothCollisionDa
 	}  // End if Asset->PhysicsAsset
 }
 
+
+int32 FClothingSimulationCollider::GetNumGeometries(int32 InLODIndex) const
+{
+	return LODData.IsValidIndex(InLODIndex) ? LODData[InLODIndex].NumGeometries : 0;
+}
+
+int32 FClothingSimulationCollider::GetOffset(const FClothingSimulationSolver* Solver, const FClothingSimulationCloth* Cloth, int32 InLODIndex) const
+{
+	return LODData.IsValidIndex(InLODIndex) ? LODData[InLODIndex].Offsets.FindChecked(FSolverClothPair(Solver, Cloth)) : INDEX_NONE;
+}
+
 void FClothingSimulationCollider::Add(FClothingSimulationSolver* Solver, FClothingSimulationCloth* Cloth)
 {
 	check(Solver);
 
-	// Reset LOD
-	LODIndex = INDEX_NONE;  // TODO: Have map of LODIndex per solver
+	// Can't add a collider twice to the same solver/cloth pair unless it's a solver collider
+	check(!Cloth || !LODIndices.Find(FSolverClothPair(Solver, Cloth)));
+
+	// Initialize LODIndex
+	int32& LODIndex = LODIndices.FindOrAdd(FSolverClothPair(Solver, Cloth));
+	LODIndex = INDEX_NONE;
 
 	// Create physics asset collisions, this will affect all LODs, so store it at index 0
 	FClothCollisionData PhysicsAssetCollisionData;
@@ -579,6 +596,16 @@ void FClothingSimulationCollider::Add(FClothingSimulationSolver* Solver, FClothi
 	}
 }
 
+void FClothingSimulationCollider::Remove(FClothingSimulationSolver* Solver, FClothingSimulationCloth* Cloth)
+{
+	LODIndices.Remove(FSolverClothPair(Solver, Cloth));
+
+	for (FLODData& LODDatum: LODData)
+	{
+		LODDatum.Remove(Solver, Cloth);
+	}
+}
+
 void FClothingSimulationCollider::Update(FClothingSimulationSolver* Solver, FClothingSimulationCloth* Cloth)
 {
 	check(Solver);
@@ -596,6 +623,7 @@ void FClothingSimulationCollider::Update(FClothingSimulationSolver* Solver, FClo
 	}
 
 	// Update current LOD index
+	int32& LODIndex = LODIndices.FindChecked(FSolverClothPair(Solver, Cloth));
 	const int32 PrevLODIndex = LODIndex;
 	LODIndex = (int32)ECollisionDataType::LODs + (bUseLODIndexOverride ? LODIndexOverride : Cloth ? Cloth->GetLODIndex(Solver) : INDEX_NONE);
 	if (!LODData.IsValidIndex(LODIndex))
@@ -637,41 +665,49 @@ TConstArrayView<TVector<float, 3>> FClothingSimulationCollider::GetCollisionTran
 {
 	check(Solver);
 
+	const int32 LODIndex = LODIndices.FindChecked(FSolverClothPair(Solver, Cloth));
 	const int32 Index = 
 		(CollisionDataType < ECollisionDataType::LODs) ? (int32)CollisionDataType :
 		(LODIndex >= (int32)ECollisionDataType::LODs) ? LODIndex : INDEX_NONE;
 
-	const TPair<const FClothingSimulationSolver*, const FClothingSimulationCloth*> SolverClothPair(Solver, Cloth);
+	const int32 Offset = GetOffset(Solver, Cloth, Index);
+	const int32 NumGeometries = GetNumGeometries(Index);
 
-	return (Index != INDEX_NONE && LODData[Index].NumGeometries) ?
-		TConstArrayView<TVector<float, 3>>(Solver->GetCollisionParticleXs(LODData[Index].Offsets.FindChecked(SolverClothPair)), LODData[Index].NumGeometries) :
+	return Offset != INDEX_NONE && NumGeometries ?
+		TConstArrayView<TVector<float, 3>>(Solver->GetCollisionParticleXs(Offset), NumGeometries) :
 		TConstArrayView<TVector<float, 3>>();
 }
 
 TConstArrayView<TRotation<float, 3>> FClothingSimulationCollider::GetCollisionRotations(const FClothingSimulationSolver* Solver, const FClothingSimulationCloth* Cloth, ECollisionDataType CollisionDataType) const
 {
 	check(Solver);
+
+	const int32 LODIndex = LODIndices.FindChecked(FSolverClothPair(Solver, Cloth));
 	const int32 Index = 
 		(CollisionDataType < ECollisionDataType::LODs) ? (int32)CollisionDataType :
 		(LODIndex >= (int32)ECollisionDataType::LODs) ? LODIndex : INDEX_NONE;
 
-	const TPair<const FClothingSimulationSolver*, const FClothingSimulationCloth*> SolverClothPair(Solver, Cloth);
+	const int32 Offset = GetOffset(Solver, Cloth, Index);
+	const int32 NumGeometries = GetNumGeometries(Index);
 
-	return (Index != INDEX_NONE && LODData[Index].NumGeometries) ?
-		TConstArrayView<TRotation<float, 3>>(Solver->GetCollisionParticleRs(LODData[Index].Offsets.FindChecked(SolverClothPair)), LODData[Index].NumGeometries) :
+	return Offset != INDEX_NONE && NumGeometries ?
+		TConstArrayView<TRotation<float, 3>>(Solver->GetCollisionParticleRs(Offset), NumGeometries) :
 		TConstArrayView<TRotation<float, 3>>();
 }
 
 TConstArrayView<TUniquePtr<FImplicitObject>> FClothingSimulationCollider::GetCollisionGeometries(const FClothingSimulationSolver* Solver, const FClothingSimulationCloth* Cloth, ECollisionDataType CollisionDataType) const
 {
 	check(Solver);
+
+	const int32 LODIndex = LODIndices.FindChecked(FSolverClothPair(Solver, Cloth));
 	const int32 Index = 
 		(CollisionDataType < ECollisionDataType::LODs) ? (int32)CollisionDataType :
 		(LODIndex >= (int32)ECollisionDataType::LODs) ? LODIndex : INDEX_NONE;
 
-	const TPair<const FClothingSimulationSolver*, const FClothingSimulationCloth*> SolverClothPair(Solver, Cloth);
+	const int32 Offset = GetOffset(Solver, Cloth, Index);
+	const int32 NumGeometries = GetNumGeometries(Index);
 
-	return (Index != INDEX_NONE && LODData[Index].NumGeometries) ?
-		TConstArrayView<TUniquePtr<FImplicitObject>>(Solver->GetCollisionGeometries(LODData[Index].Offsets.FindChecked(SolverClothPair)), LODData[Index].NumGeometries) :
+	return Offset != INDEX_NONE && NumGeometries ?
+		TConstArrayView<TUniquePtr<FImplicitObject>>(Solver->GetCollisionGeometries(Offset), NumGeometries) :
 		TConstArrayView<TUniquePtr<FImplicitObject>>();
 }
diff --git a/Engine/Plugins/Experimental/ChaosCloth/Source/Chaos/Public/ChaosCloth/ChaosClothingSimulationCloth.h b/Engine/Plugins/Experimental/ChaosCloth/Source/Chaos/Public/ChaosCloth/ChaosClothingSimulationCloth.h
--- a/Engine/Plugins/Experimental/ChaosCloth/Source/Chaos/Public/ChaosCloth/ChaosClothingSimulationCloth.h
+++ b/Engine/Plugins/Experimental/ChaosCloth/Source/Chaos/Public/ChaosCloth/ChaosClothingSimulationCloth.h
@@ -143,6 +143,7 @@ namespace Chaos
 			FLODData(int32 InNumParticles, const TConstArrayView<uint32>& InIndices, const TArray<TConstArrayView<float>>& InWeightMaps);
 
 			void Add(FClothingSimulationSolver* Solver, FClothingSimulationCloth* Cloth, int32 LODIndex);
+			void Remove(FClothingSimulationSolver* Solver);
 
 			void Enable(FClothingSimulationSolver* Solver, bool bEnable) const;
 
diff --git a/Engine/Plugins/Experimental/ChaosCloth/Source/Chaos/Public/ChaosCloth/ChaosClothingSimulationCollider.h b/Engine/Plugins/Experimental/ChaosCloth/Source/Chaos/Public/ChaosCloth/ChaosClothingSimulationCollider.h
--- a/Engine/Plugins/Experimental/ChaosCloth/Source/Chaos/Public/ChaosCloth/ChaosClothingSimulationCollider.h
+++ b/Engine/Plugins/Experimental/ChaosCloth/Source/Chaos/Public/ChaosCloth/ChaosClothingSimulationCollider.h
@@ -34,7 +34,7 @@ namespace Chaos
 		int32 GetNumGeometries() const { int32 NumGeometries = 0; for (const FLODData& LODDatum : LODData) { NumGeometries += LODDatum.NumGeometries; } return NumGeometries; }
 
 		// Return source (untransformed) collision data for global, dynamic and active LODs.
-		FClothCollisionData GetCollisionData() const;
+		FClothCollisionData GetCollisionData(const FClothingSimulationSolver* Solver, const FClothingSimulationCloth* Cloth) const;
 
 		// ---- Animatable property setters ----
 		// Set dynamic collision data, will only get updated when used as a Solver Collider TODO: Subclass collider?
@@ -43,7 +43,7 @@ namespace Chaos
 
 		// ---- Cloth/Solver interface ----
 		void Add(FClothingSimulationSolver* Solver, FClothingSimulationCloth* Cloth = nullptr);
-		void Remove(FClothingSimulationSolver* /*Solver*/, FClothingSimulationCloth* /*Cloth*/ = nullptr) {}
+		void Remove(FClothingSimulationSolver* Solver, FClothingSimulationCloth* Cloth = nullptr);
 
 		void Update(FClothingSimulationSolver* Solver, FClothingSimulationCloth* Cloth = nullptr);
 		// ---- End of the Cloth/Solver interface ----
@@ -62,12 +62,17 @@ namespace Chaos
 	private:
 		void ExtractPhysicsAssetCollision(FClothCollisionData& ClothCollisionData, TArray<int32>& UsedBoneIndices);
 
+		int32 GetNumGeometries(int32 InLODIndex) const;
+		int32 GetOffset(const FClothingSimulationSolver* Solver, const FClothingSimulationCloth* Cloth, int32 InLODIndex) const;
+
 	private:
+		typedef TPair<const FClothingSimulationSolver*, const FClothingSimulationCloth*> FSolverClothPair;
+
 		struct FLODData
 		{
 			FClothCollisionData ClothCollisionData;
 			int32 NumGeometries;  // Number of collision bodies
-			TMap<TPair<const FClothingSimulationSolver*, const FClothingSimulationCloth*>, int32> Offsets;  // Solver particle offset
+			TMap<FSolverClothPair, int32> Offsets;  // Solver particle offset
 
 			FLODData() : NumGeometries(0) {}
 
@@ -76,6 +81,7 @@ namespace Chaos
 				FClothingSimulationCloth* Cloth,
 				const FClothCollisionData& InClothCollisionData,
 				const TArray<int32>& UsedBoneIndices = TArray<int32>());
+			void Remove(FClothingSimulationSolver* Solver, FClothingSimulationCloth* Cloth);
 
 			void Update(FClothingSimulationSolver* Solver, FClothingSimulationCloth* Cloth, const FClothingSimulationContextCommon* Context);
 
@@ -97,6 +103,6 @@ namespace Chaos
 
 		// Collision primitives
 		TArray<FLODData> LODData;  // Actual LODs start at LODStart
-		int32 LODIndex;  // TODO: Have map of LODIndex per solver
+		TMap<FSolverClothPair, int32> LODIndices;
 	};
 } // namespace Chaos
\ No newline at end of file
