Chaos Cloth - Fix MaxDistances debug draw crash when the mask is disable or missing.

diff --git a/Engine/Plugins/Experimental/ChaosCloth/Source/Chaos/Private/ChaosCloth/ChaosClothingSimulation.cpp b/Engine/Plugins/Experimental/ChaosCloth/Source/Chaos/Private/ChaosCloth/ChaosClothingSimulation.cpp
--- a/Engine/Plugins/Experimental/ChaosCloth/Source/Chaos/Private/ChaosCloth/ChaosClothingSimulation.cpp
+++ b/Engine/Plugins/Experimental/ChaosCloth/Source/Chaos/Private/ChaosCloth/ChaosClothingSimulation.cpp
@@ -626,6 +626,11 @@ void FClothingSimulation::DebugDrawMaxDistanceValues(FCanvas* Canvas, const FSce
 		}
 
 		const TConstArrayView<float>& MaxDistances = Cloth->GetWeightMaps(Solver.Get())[(int32)EChaosWeightMapTarget::MaxDistance];
+		if (!MaxDistances.Num())
+		{
+			continue;
+		}
+
 		const TConstArrayView<TVector<float, 3>> Positions = Cloth->GetAnimationPositions(Solver.Get());
 		const TConstArrayView<float> InvMasses = Cloth->GetParticleInvMasses(Solver.Get());
 		check(MaxDistances.Num() == Positions.Num());
@@ -1213,6 +1218,11 @@ void FClothingSimulation::DebugDrawMaxDistances(FPrimitiveDrawInterface* PDI) co
 		}
 
 		const TConstArrayView<float>& MaxDistances = Cloth->GetWeightMaps(Solver.Get())[(int32)EChaosWeightMapTarget::MaxDistance];
+		if (!MaxDistances.Num())
+		{
+			continue;
+		}
+
 		const TConstArrayView<float> InvMasses = Cloth->GetParticleInvMasses(Solver.Get());
 		const TConstArrayView<TVector<float, 3>> Positions = Cloth->GetAnimationPositions(Solver.Get());
 		const TConstArrayView<TVector<float, 3>> Normals = Cloth->GetAnimationNormals(Solver.Get());
