Chaos Cloth - Rework external collider code to be cloth bound rather than solver bound to fix collider issues with the local velocity scales.

diff --git a/Engine/Plugins/Experimental/ChaosCloth/Source/Chaos/Private/ChaosCloth/ChaosClothingSimulation.cpp b/Engine/Plugins/Experimental/ChaosCloth/Source/Chaos/Private/ChaosCloth/ChaosClothingSimulation.cpp
--- a/Engine/Plugins/Experimental/ChaosCloth/Source/Chaos/Private/ChaosCloth/ChaosClothingSimulation.cpp
+++ b/Engine/Plugins/Experimental/ChaosCloth/Source/Chaos/Private/ChaosCloth/ChaosClothingSimulation.cpp
@@ -91,11 +91,6 @@ void FClothingSimulation::Initialize()
 	// Create solver
 	Solver = MakeUnique<FClothingSimulationSolver>();
 
-	// Create external collision dynamic collider
-	Colliders.Emplace(MakeUnique<FClothingSimulationCollider>(nullptr, nullptr, false, INDEX_NONE));
-	Colliders[0]->SetCollisionData(&ExternalCollisionData);
-	Solver->AddCollider(Colliders[0].Get());
-
 	ResetStats();
 }
 
@@ -169,6 +164,9 @@ void FClothingSimulation::CreateActor(USkeletalMeshComponent* InOwnerComponent,
 		/*bInUseLODIndexOverride =*/ false,
 		/*InLODIndexOverride =*/ INDEX_NONE));
 
+	// Set the external collision data to get updated at every frame
+	Colliders[ColliderIndex]->SetCollisionData(&ExternalCollisionData);
+
 	// Create cloth node
 	AnimDriveSpringStiffness = ClothConfig->AnimDriveSpringStiffness;
 	const int32 ClothIndex = Cloths.Emplace(MakeUnique<FClothingSimulationCloth>(
@@ -465,26 +463,9 @@ void FClothingSimulation::RefreshClothConfig(const IClothingSimulationContext* I
 
 void FClothingSimulation::RefreshPhysicsAsset()
 {
-	// Reset stats
-	ResetStats();
-
-	// A collider update cannot be re-triggered for now, remove all cloths from the solver instead
-	Solver->RemoveCloths();
-
-	for (const TUniquePtr<FClothingSimulationCloth>& Cloth : Cloths)
-	{
-		// Backup collider(s)
-		TArray<FClothingSimulationCollider*> ClothColliders = Cloth->GetColliders();
+	// A collider update cannot be re-triggered for now, refresh all cloths from the solver instead
+	Solver->RefreshCloths();
 
-		// Refresh collider
-		Cloth->SetColliders(MoveTemp(ClothColliders));
-
-		// Re-add cloth to the solver
-		Solver->AddCloth(Cloth.Get());
-
-		// Update stats
-		UpdateStats(Cloth.Get());
-	}
 	UE_LOG(LogChaosCloth, VeryVerbose, TEXT("RefreshPhysicsAsset, all collisions have been re-added for all clothing assets"));
 }
 
@@ -1016,8 +997,8 @@ void FClothingSimulation::DebugDrawCollision(FPrimitiveDrawInterface* PDI) const
 			static const FLinearColor LODsColor(FColor::Silver);
 
 			const FLinearColor Color =
-				(CollisionDataType == FClothingSimulationCollider::ECollisionDataType::Global) ? GlobalColor :
-				(CollisionDataType == FClothingSimulationCollider::ECollisionDataType::Dynamic) ? DynamicColor : LODsColor;
+				(CollisionDataType == FClothingSimulationCollider::ECollisionDataType::LODless) ? GlobalColor :
+				(CollisionDataType == FClothingSimulationCollider::ECollisionDataType::External) ? DynamicColor : LODsColor;
 
 			const TVector<float, 3>& LocalSpaceLocation = Solver->GetLocalSpaceLocation();
 
@@ -1084,15 +1065,11 @@ void FClothingSimulation::DebugDrawCollision(FPrimitiveDrawInterface* PDI) const
 	{
 		for (const FClothingSimulationCollider* const Collider : Cloth->GetColliders())
 		{
-			DrawCollision(Collider, Cloth, FClothingSimulationCollider::ECollisionDataType::Global);
+			DrawCollision(Collider, Cloth, FClothingSimulationCollider::ECollisionDataType::LODless);
+			DrawCollision(Collider, Cloth, FClothingSimulationCollider::ECollisionDataType::External);
 			DrawCollision(Collider, Cloth, FClothingSimulationCollider::ECollisionDataType::LODs);
 		}
 	}
-	// Draw external collisions
-	for (const FClothingSimulationCollider* const Collider : Solver->GetColliders())
-	{
-		DrawCollision(Collider, nullptr, FClothingSimulationCollider::ECollisionDataType::Dynamic);
-	}
 }
 
 void FClothingSimulation::DebugDrawBackstops(FPrimitiveDrawInterface* PDI) const
diff --git a/Engine/Plugins/Experimental/ChaosCloth/Source/Chaos/Private/ChaosCloth/ChaosClothingSimulationCloth.cpp b/Engine/Plugins/Experimental/ChaosCloth/Source/Chaos/Private/ChaosCloth/ChaosClothingSimulationCloth.cpp
--- a/Engine/Plugins/Experimental/ChaosCloth/Source/Chaos/Private/ChaosCloth/ChaosClothingSimulationCloth.cpp
+++ b/Engine/Plugins/Experimental/ChaosCloth/Source/Chaos/Private/ChaosCloth/ChaosClothingSimulationCloth.cpp
@@ -395,71 +395,111 @@ void FClothingSimulationCloth::SetMesh(FClothingSimulationMesh* InMesh)
 {
 	Mesh = InMesh;
 
-	// Backup all currently attached solvers
-	TArray<FClothingSimulationSolver*> Solvers;
-	LODIndices.GetKeys(Solvers);
-
-	// Remove this cloth from all solvers
-	for (FClothingSimulationSolver* const Solver : Solvers)
-	{
-		Solver->RemoveCloth(this);
-	}
-
 	// Reset LODs
 	const int32 NumLODs = Mesh ? Mesh->GetNumLODs() : 0;
-
 	LODData.Reset(NumLODs);
 	for (int32 Index = 0; Index < NumLODs; ++Index)
 	{
 		LODData.Emplace(Mesh->GetNumPoints(Index), Mesh->GetIndices(Index), Mesh->GetWeightMaps(Index));
 	}
 
-	// Re-add this cloth to existing solvers
+	// Iterate all known solvers
+	TArray<FClothingSimulationSolver*> Solvers;
+	LODIndices.GetKeys(Solvers);
 	for (FClothingSimulationSolver* const Solver : Solvers)
 	{
-		Solver->AddCloth(this);
+		// Refresh this cloth to recreate particles
+		Solver->RefreshCloth(this);
 	}
 }
 
 void FClothingSimulationCloth::SetColliders(TArray<FClothingSimulationCollider*>&& InColliders)
 {
+	// Empty the collider list, but keep the pointers around for the removal operation below
+	const TArray<FClothingSimulationCollider*> TempColliders = MoveTemp(Colliders);
+
+	// Replace with the new colliders
 	Colliders = InColliders;
 
-	// Backup all currently attached solvers
+	// Iterate all known solvers
 	TArray<FClothingSimulationSolver*> Solvers;
 	LODIndices.GetKeys(Solvers);
-
-	// Remove this cloth from all solvers
 	for (FClothingSimulationSolver* const Solver : Solvers)
 	{
-		Solver->RemoveCloth(this);
+		// Remove any held collider data related to this cloth simulation
+		for (FClothingSimulationCollider* const Collider : TempColliders)
+		{
+			Collider->Remove(Solver, this);
+		}
+
+		// Refresh this cloth to recreate collision particles
+		Solver->RefreshCloth(this);
 	}
+}
+
+void FClothingSimulationCloth::AddCollider(FClothingSimulationCollider* InCollider)
+{
+	check(InCollider);
 
-	// Re-add this cloth to existing solvers
+	if (Colliders.Find(InCollider) != INDEX_NONE)
+	{
+		return;
+	}
+
+	// Add the collider to the solver update array
+	Colliders.Emplace(InCollider);
+
+	// Iterate all known solvers
+	TArray<FClothingSimulationSolver*> Solvers;
+	LODIndices.GetKeys(Solvers);
 	for (FClothingSimulationSolver* const Solver : Solvers)
 	{
-		Solver->AddCloth(this);
+		// Refresh this cloth to recreate collision particles
+		Solver->RefreshCloth(this);
 	}
 }
 
-void FClothingSimulationCloth::RemoveColliders()
+void FClothingSimulationCloth::RemoveCollider(FClothingSimulationCollider* InCollider)
 {
-	Colliders.Reset();
+	if (Colliders.Find(InCollider) == INDEX_NONE)
+	{
+		return;
+	}
 
-	// Backup all currently attached solvers
+	// Remove collider from array
+	Colliders.RemoveSwap(InCollider);
+
+	// Iterate all known solvers
 	TArray<FClothingSimulationSolver*> Solvers;
 	LODIndices.GetKeys(Solvers);
-
-	// Remove this cloth from all solvers
 	for (FClothingSimulationSolver* const Solver : Solvers)
 	{
-		Solver->RemoveCloth(this);
+		// Remove any held collider data related to this cloth simulation
+		InCollider->Remove(Solver, this);
+
+		// Refresh this cloth to recreate collision particles
+		Solver->RefreshCloth(this);
 	}
+}
 
-	// Re-add this cloth to existing solvers
+void FClothingSimulationCloth::RemoveColliders()
+{
+	// Empty the collider list, but keep the pointers around for the removal operation below
+	const TArray<FClothingSimulationCollider*> TempColliders = MoveTemp(Colliders);
+
+	// Iterate all known solvers
+	TArray<FClothingSimulationSolver*> Solvers;
+	LODIndices.GetKeys(Solvers);
 	for (FClothingSimulationSolver* const Solver : Solvers)
 	{
-		Solver->AddCloth(this);
+		// Remove any held collider data related to this cloth simulation
+		for (FClothingSimulationCollider* const Collider : TempColliders)
+		{
+			Collider->Remove(Solver, this);
+		}
+
+		// Refresh this cloth to recreate collision particles
+		Solver->RefreshCloth(this);
 	}
 }
 
@@ -489,18 +529,18 @@ void FClothingSimulationCloth::Add(FClothingSimulationSolver* Solver)
 
 void FClothingSimulationCloth::Remove(FClothingSimulationSolver* Solver)
 {
+	// Remove Colliders
+	for (FClothingSimulationCollider* Collider : Colliders)
+	{
+		Collider->Remove(Solver, this);
+	}
+
 	// Remove solver from maps
 	LODIndices.Remove(Solver);
 	for (FLODData& LODDatum: LODData)
 	{
 		LODDatum.Remove(Solver);
 	}
-
-	// Remove Colliders
-	for (FClothingSimulationCollider* Collider : Colliders)
-	{
-		Collider->Remove(Solver, this);
-	}
 }
 
 int32 FClothingSimulationCloth::GetNumParticles(int32 InLODIndex) const
diff --git a/Engine/Plugins/Experimental/ChaosCloth/Source/Chaos/Private/ChaosCloth/ChaosClothingSimulationCollider.cpp b/Engine/Plugins/Experimental/ChaosCloth/Source/Chaos/Private/ChaosCloth/ChaosClothingSimulationCollider.cpp
--- a/Engine/Plugins/Experimental/ChaosCloth/Source/Chaos/Private/ChaosCloth/ChaosClothingSimulationCollider.cpp
+++ b/Engine/Plugins/Experimental/ChaosCloth/Source/Chaos/Private/ChaosCloth/ChaosClothingSimulationCollider.cpp
@@ -24,6 +24,7 @@ void FClothingSimulationCollider::FLODData::Add(
 	const TArray<int32>& UsedBoneIndices)
 {
 	check(Solver);
+	check(Cloth);
 
 	// Keep a list of all collisions
 	ClothCollisionData = InClothCollisionData;
@@ -55,10 +56,9 @@ void FClothingSimulationCollider::FLODData::Add(
 	const uint32 GroupId = Cloth ? Cloth->GetGroupId() : INDEX_NONE;
 
 	// The offset will be set to the first collision particle's index
-	// Try to reuse existing offsets when Add is called during the solver's collider update
+	// Try to reuse existing offsets when Add is called during the collider update (ie Offset isn't INDEX_NONE)
 	int32* const OffsetPtr = Offsets.Find(FSolverClothPair(Solver, Cloth));
 	const bool bIsNewCollider = !OffsetPtr;
-	check(!Cloth || bIsNewCollider);  // TODO: We need to be able to update cloth colliders too at some point, but this won't work as it is
 	int32& Offset = bIsNewCollider ? Offsets.Add(FSolverClothPair(Solver, Cloth)) : *OffsetPtr;
 	Offset = Solver->AddCollisionParticles(NumGeometries, GroupId, bIsNewCollider ? INDEX_NONE : Offset);
 
@@ -278,6 +278,7 @@ void FClothingSimulationCollider::FLODData::Remove(FClothingSimulationSolver* So
 void FClothingSimulationCollider::FLODData::Update(FClothingSimulationSolver* Solver, FClothingSimulationCloth* Cloth, const FClothingSimulationContextCommon* Context)
 {
 	check(Solver);
+	check(Cloth);
 	if (NumGeometries)
 	{
 		const int32 Offset = Offsets.FindChecked(FSolverClothPair(Solver, Cloth));
@@ -304,6 +305,7 @@ void FClothingSimulationCollider::FLODData::Update(FClothingSimulationSolver* So
 void FClothingSimulationCollider::FLODData::Enable(FClothingSimulationSolver* Solver, FClothingSimulationCloth* Cloth, bool bEnable)
 {
 	check(Solver);
+	check(Cloth);
 	if (NumGeometries)
 	{
 		const int32 Offset = Offsets.FindChecked(FSolverClothPair(Solver, Cloth));
@@ -314,6 +316,7 @@ void FClothingSimulationCollider::FLODData::Enable(FClothingSimulationSolver* So
 void FClothingSimulationCollider::FLODData::ResetStartPose(FClothingSimulationSolver* Solver, FClothingSimulationCloth* Cloth)
 {
 	check(Solver);
+	check(Cloth);
 	if (NumGeometries)
 	{
 		const int32 Offset = Offsets.FindChecked(FSolverClothPair(Solver, Cloth));
@@ -354,8 +357,8 @@ FClothingSimulationCollider::~FClothingSimulationCollider()
 FClothCollisionData FClothingSimulationCollider::GetCollisionData(const FClothingSimulationSolver* Solver, const FClothingSimulationCloth* Cloth) const
 {
 	FClothCollisionData ClothCollisionData;
-	ClothCollisionData.Append(LODData[(int32)ECollisionDataType::Global].ClothCollisionData);
-	ClothCollisionData.Append(LODData[(int32)ECollisionDataType::Dynamic].ClothCollisionData);
+	ClothCollisionData.Append(LODData[(int32)ECollisionDataType::LODless].ClothCollisionData);
+	ClothCollisionData.Append(LODData[(int32)ECollisionDataType::External].ClothCollisionData);
 
 	const int32 LODIndex = LODIndices.FindChecked(FSolverClothPair(Solver, Cloth));
 	if (LODIndex >= (int32)ECollisionDataType::LODs)
@@ -550,18 +553,20 @@ int32 FClothingSimulationCollider::GetNumGeometries(int32 InLODIndex) const
 
 int32 FClothingSimulationCollider::GetOffset(const FClothingSimulationSolver* Solver, const FClothingSimulationCloth* Cloth, int32 InLODIndex) const
 {
-	return LODData.IsValidIndex(InLODIndex) ? LODData[InLODIndex].Offsets.FindChecked(FSolverClothPair(Solver, Cloth)) : INDEX_NONE;
+	const int32* const Offset = LODData.IsValidIndex(InLODIndex) ? LODData[InLODIndex].Offsets.Find(FSolverClothPair(Solver, Cloth)) : nullptr;
+	return Offset ? *Offset : INDEX_NONE;
 }
 
 void FClothingSimulationCollider::Add(FClothingSimulationSolver* Solver, FClothingSimulationCloth* Cloth)
 {
 	check(Solver);
+	check(Cloth);
 
-	// Can't add a collider twice to the same solver/cloth pair unless it's a solver collider
-	check(!Cloth || !LODIndices.Find(FSolverClothPair(Solver, Cloth)));
+	// Can't add a collider twice to the same solver/cloth pair
+	check(!LODIndices.Find(FSolverClothPair(Solver, Cloth)));
 
 	// Initialize LODIndex
-	int32& LODIndex = LODIndices.FindOrAdd(FSolverClothPair(Solver, Cloth));
+	int32& LODIndex = LODIndices.Add(FSolverClothPair(Solver, Cloth));
 	LODIndex = INDEX_NONE;
 
 	// Create physics asset collisions, this will affect all LODs, so store it at index 0
@@ -569,10 +574,7 @@ void FClothingSimulationCollider::Add(FClothingSimulationSolver* Solver, FClothi
 	TArray<int32> UsedBoneIndices;
 	ExtractPhysicsAssetCollision(PhysicsAssetCollisionData, UsedBoneIndices);
 
-	LODData[(int32)ECollisionDataType::Global].Add(Solver, Cloth, PhysicsAssetCollisionData, UsedBoneIndices);
-
-	// Create user collision data
-	LODData[(int32)ECollisionDataType::Dynamic].Add(Solver, Cloth, CollisionData ? *CollisionData : FClothCollisionData());
+	LODData[(int32)ECollisionDataType::LODless].Add(Solver, Cloth, PhysicsAssetCollisionData, UsedBoneIndices);
 
 	// Create legacy asset LOD collisions
 	const int32 NumLODs = Asset ? Asset->LodData.Num() : 0;
@@ -609,14 +611,18 @@ void FClothingSimulationCollider::Remove(FClothingSimulationSolver* Solver, FClo
 void FClothingSimulationCollider::Update(FClothingSimulationSolver* Solver, FClothingSimulationCloth* Cloth)
 {
 	check(Solver);
+	check(Cloth);
 
 	SCOPE_CYCLE_COUNTER(STAT_ChaosClothingSimulationColliderUpdate);
 
-	// Retrieve context (can be null when no skeletal mesh component is specified)
-	const FClothingSimulationContextCommon* const Context = SkeletalMeshComponent ? static_cast<const FClothingSimulationContextCommon*>(SkeletalMeshComponent->GetClothingSimulationContext()) : nullptr;
+	// Add or re-add the external collision particles
+	const int32 ExternalCollisionNumGeometries = GetNumGeometries((int32)ECollisionDataType::External);
+	const int32 ExternalCollisionOffset = GetOffset(Solver, Cloth, (int32)ECollisionDataType::External);
+	LODData[(int32)ECollisionDataType::External].Add(Solver, Cloth, CollisionData ? *CollisionData : FClothCollisionData());
 
 	// Update the collision transforms
 	// Note: All the LODs are updated at once, so that the previous transforms are always correct during LOD switching
+	const FClothingSimulationContextCommon* const Context = SkeletalMeshComponent ? static_cast<const FClothingSimulationContextCommon*>(SkeletalMeshComponent->GetClothingSimulationContext()) : nullptr;
 	for (FLODData& LODDatum : LODData)
 	{
 		LODDatum.Update(Solver, Cloth, Context);
@@ -631,16 +637,27 @@ void FClothingSimulationCollider::Update(FClothingSimulationSolver* Solver, FClo
 		LODIndex = (int32)ECollisionDataType::LODs + INDEX_NONE;
 	}
 
-	// Enable/disable particle LOD
+	// Enable particle if the external collisions have changed
+	// TODO: Find a better way in case the same number but different collisions are being re-added (hash collision data? Provide user dirty function?)
+	if (ExternalCollisionNumGeometries != GetNumGeometries((int32)ECollisionDataType::External) ||
+		ExternalCollisionOffset != GetOffset(Solver, Cloth, (int32)ECollisionDataType::External))
+	{
+		// Enable non LOD collisions at first initialization
+		LODData[(int32)ECollisionDataType::External].Enable(Solver, Cloth, true);
+
+		// Update initial state for collisions
+		LODData[(int32)ECollisionDataType::External].ResetStartPose(Solver, Cloth);
+	}
+
 	if (LODIndex != PrevLODIndex)
 	{
-		if (PrevLODIndex == INDEX_NONE)
+		if (PrevLODIndex == INDEX_NONE)  // First run
 		{
 			// Enable non LOD collisions at first initialization
-			LODData[(int32)ECollisionDataType::Global].Enable(Solver, Cloth, true);
+			LODData[(int32)ECollisionDataType::LODless].Enable(Solver, Cloth, true);
 
 			// Update initial state for collisions
-			LODData[(int32)ECollisionDataType::Global].ResetStartPose(Solver, Cloth);
+			LODData[(int32)ECollisionDataType::LODless].ResetStartPose(Solver, Cloth);
 		}
 		else if (PrevLODIndex > 0)
 		{
@@ -661,9 +678,18 @@ void FClothingSimulationCollider::Update(FClothingSimulationSolver* Solver, FClo
 	}
 }
 
+void FClothingSimulationCollider::ResetStartPose(FClothingSimulationSolver* Solver, FClothingSimulationCloth* Cloth)
+{
+	for (FLODData& LODDatum : LODData)
+	{
+		LODDatum.ResetStartPose(Solver, Cloth);
+	}
+}
+
 TConstArrayView<TVector<float, 3>> FClothingSimulationCollider::GetCollisionTranslations(const FClothingSimulationSolver* Solver, const FClothingSimulationCloth* Cloth, ECollisionDataType CollisionDataType) const
 {
 	check(Solver);
+	check(Cloth);
 
 	const int32 LODIndex = LODIndices.FindChecked(FSolverClothPair(Solver, Cloth));
 	const int32 Index = 
@@ -681,6 +707,7 @@ TConstArrayView<TVector<float, 3>> FClothingSimulationCollider::GetCollisionTran
 TConstArrayView<TRotation<float, 3>> FClothingSimulationCollider::GetCollisionRotations(const FClothingSimulationSolver* Solver, const FClothingSimulationCloth* Cloth, ECollisionDataType CollisionDataType) const
 {
 	check(Solver);
+	check(Cloth);
 
 	const int32 LODIndex = LODIndices.FindChecked(FSolverClothPair(Solver, Cloth));
 	const int32 Index = 
@@ -698,6 +725,7 @@ TConstArrayView<TRotation<float, 3>> FClothingSimulationCollider::GetCollisionRo
 TConstArrayView<TUniquePtr<FImplicitObject>> FClothingSimulationCollider::GetCollisionGeometries(const FClothingSimulationSolver* Solver, const FClothingSimulationCloth* Cloth, ECollisionDataType CollisionDataType) const
 {
 	check(Solver);
+	check(Cloth);
 
 	const int32 LODIndex = LODIndices.FindChecked(FSolverClothPair(Solver, Cloth));
 	const int32 Index = 
diff --git a/Engine/Plugins/Experimental/ChaosCloth/Source/Chaos/Private/ChaosCloth/ChaosClothingSimulationSolver.cpp b/Engine/Plugins/Experimental/ChaosCloth/Source/Chaos/Private/ChaosCloth/ChaosClothingSimulationSolver.cpp
--- a/Engine/Plugins/Experimental/ChaosCloth/Source/Chaos/Private/ChaosCloth/ChaosClothingSimulationSolver.cpp
+++ b/Engine/Plugins/Experimental/ChaosCloth/Source/Chaos/Private/ChaosCloth/ChaosClothingSimulationSolver.cpp
@@ -11,7 +11,6 @@
 using namespace Chaos;
 
 DECLARE_CYCLE_STAT(TEXT("Chaos Cloth Solver Update"), STAT_ChaosClothSolverUpdate, STATGROUP_ChaosCloth);
-DECLARE_CYCLE_STAT(TEXT("Chaos Cloth Solver Update Colliders"), STAT_ChaosClothSolverUpdateColliders, STATGROUP_ChaosCloth);
 DECLARE_CYCLE_STAT(TEXT("Chaos Cloth Solver Update Cloths"), STAT_ChaosClothSolverUpdateCloths, STATGROUP_ChaosCloth);
 DECLARE_CYCLE_STAT(TEXT("Chaos Cloth Solver Update Pre Solver Step"), STAT_ChaosClothSolverUpdatePreSolverStep, STATGROUP_ChaosCloth);
 DECLARE_CYCLE_STAT(TEXT("Chaos Cloth Solver Update Solver Step"), STAT_ChaosClothSolverUpdateSolverStep, STATGROUP_ChaosCloth);
@@ -46,7 +45,6 @@ FClothingSimulationSolver::FClothingSimulationSolver()
 	, NumSubsteps(ChaosClothingSimulationSolverDefault::NumSubsteps)
 	, CollisionParticlesOffset(0)
 	, CollisionParticlesSize(0)
-	, bAreCollisionParticlesDirty(false)
 	, Gravity(ChaosClothingSimulationSolverDefault::Gravity)
 	, WindVelocity(ChaosClothingSimulationSolverDefault::WindVelocity)
 	, WindFluidDensity(ChaosClothingSimulationSolverDefault::WindFluidDensity)
@@ -114,11 +112,8 @@ FClothingSimulationSolver::~FClothingSimulationSolver()
 
 void FClothingSimulationSolver::SetCloths(TArray<FClothingSimulationCloth*>&& InCloths)
 {
-	// Reset solver collisions so that there is never any dynamic collision particles below cloth's own ones
-	Evolution->ResetCollisionParticles();
-
-	// Reset cloth particles and associated elements
-	ResetParticles();
+	// Remove old cloths
+	RemoveCloths();
 
 	// Update array
 	Cloths = MoveTemp(InCloths);
@@ -135,9 +130,8 @@ void FClothingSimulationSolver::SetCloths(TArray<FClothingSimulationCloth*>&& In
 		Cloth->Update(this);
 	}
 
-	// Update solver collider's offset
+	// Update external collision's offset
 	CollisionParticlesOffset = Evolution->CollisionParticles().Size();
-	CollisionParticlesSize = 0;
 }
 
 void FClothingSimulationSolver::AddCloth(FClothingSimulationCloth* InCloth)
@@ -149,21 +143,20 @@ void FClothingSimulationSolver::AddCloth(FClothingSimulationCloth* InCloth)
 		return;
 	}
 
-	// Reset solver collisions so that there is never any dynamic collision particles below cloth's own ones
-	Evolution->ResetCollisionParticles(CollisionParticlesOffset);
-
 	// Add the cloth to the solver update array
 	Cloths.Emplace(InCloth);
 
+	// Reset external collisions so that there is never any external collision particles below cloth's ones
+	ResetCollisionParticles(CollisionParticlesOffset);
+
 	// Add the cloth's particles
 	InCloth->Add(this);
 
 	// Set initial state
 	InCloth->Update(this);
 
-	// Update solver collider's offset
+	// Update external collision's offset
 	CollisionParticlesOffset = Evolution->CollisionParticles().Size();
-	CollisionParticlesSize = 0;
 }
 
 void FClothingSimulationSolver::RemoveCloth(FClothingSimulationCloth* InCloth)
@@ -173,18 +166,18 @@ void FClothingSimulationSolver::RemoveCloth(FClothingSimulationCloth* InCloth)
 		return;
 	}
 
-	// Reset solver collisions so that there is never any dynamic collision particles below cloth's own ones
-	Evolution->ResetCollisionParticles();
-
-	// Reset cloth particles and associated elements
-	ResetParticles();
-
 	// Remove reference to this solver
 	InCloth->Remove(this);
 
 	// Remove collider from array
 	Cloths.RemoveSwap(InCloth);
 
+	// Reset collisions so that there is never any external collision particles below the cloth's ones
+	ResetCollisionParticles();
+
+	// Reset cloth particles and associated elements
+	ResetParticles();
+
 	// Re-add the remaining cloths' particles
 	for (FClothingSimulationCloth* const Cloth : Cloths)
 	{
@@ -195,9 +188,8 @@ void FClothingSimulationSolver::RemoveCloth(FClothingSimulationCloth* InCloth)
 		Cloth->Update(this);
 	}
 
-	// Update solver collider's offset
+	// Update external collision's offset
 	CollisionParticlesOffset = Evolution->CollisionParticles().Size();
-	CollisionParticlesSize = 0;
 }
 
 void FClothingSimulationSolver::RemoveCloths()
@@ -210,14 +202,50 @@ void FClothingSimulationSolver::RemoveCloths()
 	Cloths.Reset();
 
 	// Reset solver collisions
-	Evolution->ResetCollisionParticles();
+	ResetCollisionParticles();
+
+	// Reset cloth particles and associated elements
+	ResetParticles();
+}
+
+void FClothingSimulationSolver::RefreshCloth(FClothingSimulationCloth* InCloth)
+{
+	if (Cloths.Find(InCloth) == INDEX_NONE)
+	{
+		return;
+	}
+
+	// TODO: Add different ways to refresh cloths without recreating everything (collisions, constraints, particles)
+	RefreshCloths();
+}
+
+void FClothingSimulationSolver::RefreshCloths()
+{
+	// Remove the cloths' & collisions' particles
+	for (FClothingSimulationCloth* const Cloth : Cloths)
+	{
+		// Remove any solver data held by the cloth 
+		Cloth->Remove(this);
+	}
+
+	// Reset collision particles
+	ResetCollisionParticles();
 
 	// Reset cloth particles and associated elements
 	ResetParticles();
 
+	// Re-add the cloths' & collisions' particles
+	for (FClothingSimulationCloth* const Cloth : Cloths)
+	{
+		// Re-Add the cloth's and collisions' particles
+		Cloth->Add(this);
+
+		// Set initial state
+		Cloth->Update(this);
+	}
+
 	// Update solver collider's offset
-	CollisionParticlesOffset = 0;
-	CollisionParticlesSize = 0;
+	CollisionParticlesOffset = Evolution->CollisionParticles().Size();
 }
 
 void FClothingSimulationSolver::ResetParticles()
@@ -286,60 +314,19 @@ const float* FClothingSimulationSolver::GetParticleInvMasses(int32 Offset) const
 	return &Evolution->Particles().InvM(Offset);
 }
 
-void FClothingSimulationSolver::SetColliders(TArray<FClothingSimulationCollider*>&& InColliders)
+void FClothingSimulationSolver::ResetCollisionParticles(int32 InCollisionParticlesOffset)
 {
-	// Update array
-	Colliders = MoveTemp(InColliders);
-
-	// Colliders have changed, must prevent from trying reusing the ranges
-	Evolution->ResetCollisionParticles(CollisionParticlesOffset);
-	CollisionParticlesSize = 0;
-}
-
-void FClothingSimulationSolver::AddCollider(FClothingSimulationCollider* InCollider)
-{
-	check(InCollider);
-
-	if (Colliders.Find(InCollider) != INDEX_NONE)
-	{
-		return;
-	}
-
-	// Add the collider to the solver update array
-	Colliders.Emplace(InCollider);
-}
-
-void FClothingSimulationSolver::RemoveCollider(FClothingSimulationCollider* InCollider)
-{
-	if (Colliders.Find(InCollider) == INDEX_NONE)
-	{
-		return;
-	}
-
-	// Remove collider from array
-	Colliders.RemoveSwap(InCollider);
-
-	// Colliders have changed, must prevent from trying reusing the ranges
-	Evolution->ResetCollisionParticles(CollisionParticlesOffset);
-	CollisionParticlesSize = 0;
-}
-
-void FClothingSimulationSolver::RemoveColliders()
-{
-	// Remove all colliders from array
-	Colliders.Reset();
-
-	// Colliders have changed, must prevent from trying reusing the ranges
-	Evolution->ResetCollisionParticles(CollisionParticlesOffset);
+	Evolution->ResetCollisionParticles(InCollisionParticlesOffset);
+	CollisionParticlesOffset = InCollisionParticlesOffset;
 	CollisionParticlesSize = 0;
 }
 
 int32 FClothingSimulationSolver::AddCollisionParticles(int32 NumCollisionParticles, uint32 GroupId, int32 RecycledOffset)
 {
-	// Try reusing the particle range so that solver collisions can be added/removed dynamically
+	// Try reusing the particle range
+	// This is used by external collisions so that they can be added/removed between every solver update.
 	// If it doesn't match then remove all ranges above the given offset to start again.
-	// This rely on the assumption that these ranges are added again in the same order
-	// by the solver at every update
+	// This rely on the assumption that these ranges are added again in the same update order.
 	if (RecycledOffset == CollisionParticlesOffset + CollisionParticlesSize)
 	{
 		CollisionParticlesSize += NumCollisionParticles;
@@ -376,9 +363,6 @@ int32 FClothingSimulationSolver::AddCollisionParticles(int32 NumCollisionParticl
 	// Always starts with particles disabled
 	EnableCollisionParticles(Offset, false);
 
-	// Any change in particle allocation should trigger an initial state update
-	bAreCollisionParticlesDirty = true;  // TODO: Dirtying collision initial state update only based on the number of particles is not bug proof, but hashing the collision geometries is a bit heavy handed too
-
 	return Offset;
 }
 
@@ -653,42 +637,16 @@ void FClothingSimulationSolver::Update(float InDeltaTime)
 	const float PrevDeltaTime = DeltaTime;
 	DeltaTime = DeltaTime + (InDeltaTime - DeltaTime) * DeltaTimeDecay;
 
-	// Update Solver Colliders
-	{
-		SCOPE_CYCLE_COUNTER(STAT_ChaosClothSolverUpdateColliders);
-
-		// Reset collision particles, but leave cloth's own colliders
-		Evolution->ResetCollisionParticles(CollisionParticlesOffset);
-
-		// Add/re-add the solver's dynamic collider particles that can change at every update
-		for (FClothingSimulationCollider* const Collider : Colliders)
-		{
-			// Add the collider's particles
-			Collider->Add(this);
-
-			if (bAreCollisionParticlesDirty)
-			{
-				// Reset initial state that will become the previous state after the swap below
-				Collider->Update(this);
-			}
-		}
-		bAreCollisionParticlesDirty = false;
-
-		// Update colliders
-		Swap(OldCollisionTransforms, CollisionTransforms);
-
-		for (FClothingSimulationCollider* Collider: Colliders)
-		{
-			Collider->Update(this);
-		}
-	}
-
-	// Update Cloths
+	// Update Cloths and cloth colliders
 	{
 		SCOPE_CYCLE_COUNTER(STAT_ChaosClothSolverUpdateCloths);
 
+		Swap(OldCollisionTransforms, CollisionTransforms);
 		Swap(OldAnimationPositions, AnimationPositions);
 
+		// Clear external collisions so that they can be re-added
+		CollisionParticlesSize = 0;
+
 		for (FClothingSimulationCloth* Cloth : Cloths)
 		{
 			const uint32 GroupId = Cloth->GetGroupId();
diff --git a/Engine/Plugins/Experimental/ChaosCloth/Source/Chaos/Public/ChaosCloth/ChaosClothingSimulationCloth.h b/Engine/Plugins/Experimental/ChaosCloth/Source/Chaos/Public/ChaosCloth/ChaosClothingSimulationCloth.h
--- a/Engine/Plugins/Experimental/ChaosCloth/Source/Chaos/Public/ChaosCloth/ChaosClothingSimulationCloth.h
+++ b/Engine/Plugins/Experimental/ChaosCloth/Source/Chaos/Public/ChaosCloth/ChaosClothingSimulationCloth.h
@@ -78,9 +78,9 @@ namespace Chaos
 
 		const TArray<FClothingSimulationCollider*>& GetColliders() const { return Colliders; }
 		void SetColliders(TArray<FClothingSimulationCollider*>&& InColliders);
+		void AddCollider(FClothingSimulationCollider* InCollider);
+		void RemoveCollider(FClothingSimulationCollider* InCollider);
 		void RemoveColliders();
-		//void AddCollider(FClothingSimulationCollider* InCollider);  // TODO
-		//void RemoveCollider(FClothingSimulationCollider* InCollider);  // TODO
 		// ---- End of the Node property getters/setters
 
 		// ---- Debugging/visualization functions
diff --git a/Engine/Plugins/Experimental/ChaosCloth/Source/Chaos/Public/ChaosCloth/ChaosClothingSimulationCollider.h b/Engine/Plugins/Experimental/ChaosCloth/Source/Chaos/Public/ChaosCloth/ChaosClothingSimulationCollider.h
--- a/Engine/Plugins/Experimental/ChaosCloth/Source/Chaos/Public/ChaosCloth/ChaosClothingSimulationCollider.h
+++ b/Engine/Plugins/Experimental/ChaosCloth/Source/Chaos/Public/ChaosCloth/ChaosClothingSimulationCollider.h
@@ -19,9 +19,9 @@ namespace Chaos
 	public:
 		enum class ECollisionDataType : int32
 		{
-			Global = 0,  // Global collision slot (physics collisions)
-			Dynamic,  // Dynamic collision slot (aka external collisions)
-			LODs,  // LODIndex based start slot (LODs collisions)
+			LODless = 0,  // Global LODless collision slot filled with physics collisions
+			External,  // External collision slot added/removed at every frame
+			LODs,  // LODIndex based start slot for LODs collisions
 		};
 
 		FClothingSimulationCollider(
@@ -33,20 +33,21 @@ namespace Chaos
 
 		int32 GetNumGeometries() const { int32 NumGeometries = 0; for (const FLODData& LODDatum : LODData) { NumGeometries += LODDatum.NumGeometries; } return NumGeometries; }
 
-		// Return source (untransformed) collision data for global, dynamic and active LODs.
+		// Return source (untransformed) collision data for LODless, external and active LODs.
 		FClothCollisionData GetCollisionData(const FClothingSimulationSolver* Solver, const FClothingSimulationCloth* Cloth) const;
 
 		// ---- Animatable property setters ----
-		// Set dynamic collision data, will only get updated when used as a Solver Collider TODO: Subclass collider?
+		// Set external collision data, will only get updated when used as a Solver Collider TODO: Subclass collider?
 		void SetCollisionData(const FClothCollisionData* InCollisionData) { CollisionData = InCollisionData; }
 		// ---- End of the animatable property setters ----
 
-		// ---- Cloth/Solver interface ----
-		void Add(FClothingSimulationSolver* Solver, FClothingSimulationCloth* Cloth = nullptr);
-		void Remove(FClothingSimulationSolver* Solver, FClothingSimulationCloth* Cloth = nullptr);
+		// ---- Cloth interface ----
+		void Add(FClothingSimulationSolver* Solver, FClothingSimulationCloth* Cloth);
+		void Remove(FClothingSimulationSolver* Solver, FClothingSimulationCloth* Cloth);
 
-		void Update(FClothingSimulationSolver* Solver, FClothingSimulationCloth* Cloth = nullptr);
-		// ---- End of the Cloth/Solver interface ----
+		void Update(FClothingSimulationSolver* Solver, FClothingSimulationCloth* Cloth);
+		void ResetStartPose(FClothingSimulationSolver* Solver, FClothingSimulationCloth* Cloth);
+		// ---- End of the Cloth interface ----
 
 		// ---- Debugging and visualization functions ----
 		// Return currently LOD active collision particles translations, not thread safe, to use after solver update.
diff --git a/Engine/Plugins/Experimental/ChaosCloth/Source/Chaos/Public/ChaosCloth/ChaosClothingSimulationSolver.h b/Engine/Plugins/Experimental/ChaosCloth/Source/Chaos/Public/ChaosCloth/ChaosClothingSimulationSolver.h
--- a/Engine/Plugins/Experimental/ChaosCloth/Source/Chaos/Public/ChaosCloth/ChaosClothingSimulationSolver.h
+++ b/Engine/Plugins/Experimental/ChaosCloth/Source/Chaos/Public/ChaosCloth/ChaosClothingSimulationSolver.h
@@ -49,13 +49,10 @@ namespace Chaos
 		void RemoveCloth(FClothingSimulationCloth* InCloth);
 		void RemoveCloths();
 
-		void SetColliders(TArray<FClothingSimulationCollider*>&& InColliders);
-		void AddCollider(FClothingSimulationCollider* InCollider);
-		void RemoveCollider(FClothingSimulationCollider* InCollider);
-		void RemoveColliders();
+		void RefreshCloth(FClothingSimulationCloth* InCloth);
+		void RefreshCloths();
 
 		TConstArrayView<const FClothingSimulationCloth*> GetCloths() const { return Cloths; }
-		TConstArrayView<const FClothingSimulationCollider*> GetColliders() const { return Colliders; }
 
 		// Update solver properties before simulation
 		void Update(float InDeltaTime);
@@ -131,6 +128,7 @@ namespace Chaos
 
 	private:
 		void ResetParticles();
+		void ResetCollisionParticles(int32 InCollisionParticlesOffset = 0);
 		void ApplyPreSimulationTransforms();
 		float SetParticleMassPerArea(int32 Offset, int32 Size, const TTriangleMesh<float>& Mesh);
 		void ParticleMassUpdateDensity(const TTriangleMesh<float>& Mesh, float Density);
@@ -141,7 +139,6 @@ namespace Chaos
 
 		// Object arrays
 		TArray<FClothingSimulationCloth*> Cloths;
-		TArray<FClothingSimulationCollider*> Colliders;
 
 		// Simulation group attributes
 		TArrayCollectionArray<TRigidTransform<float, 3>> PreSimulationTransforms;  // Allow a different frame of reference for each cloth groups
@@ -174,11 +171,6 @@ namespace Chaos
 		// Solver colliders offset
 		int32 CollisionParticlesOffset;  // Collision particle offset on the first solver/non cloth collider
 		int32 CollisionParticlesSize;  // Number of solver only colliders
-		bool bAreCollisionParticlesDirty;  // Whereas adding/re-add a solver collider has triggered a new particle allocation
-
-		// Active particle view
-		TArray<TPair<FClothingSimulationCloth*, TVector<uint32, 2>>> ClothsMap;
-		TArray<TPair<FClothingSimulationCollider*, TVector<uint32, 2>>> CollidersMap;
 
 		// Solver parameters
 		TVector<float, 3> Gravity;
