Chaos Cloth - Fix Collision debug draw that is missing the local space location.

diff --git a/Engine/Plugins/Experimental/ChaosCloth/Source/Chaos/Private/ChaosCloth/ChaosClothingSimulation.cpp b/Engine/Plugins/Experimental/ChaosCloth/Source/Chaos/Private/ChaosCloth/ChaosClothingSimulation.cpp
--- a/Engine/Plugins/Experimental/ChaosCloth/Source/Chaos/Private/ChaosCloth/ChaosClothingSimulation.cpp
+++ b/Engine/Plugins/Experimental/ChaosCloth/Source/Chaos/Private/ChaosCloth/ChaosClothingSimulation.cpp
@@ -1012,18 +1012,21 @@ void FClothingSimulation::DebugDrawCollision(FPrimitiveDrawInterface* PDI) const
 			{
 				if (const FImplicitObject* const Object = CollisionGeometries[Index].Get())
 				{
+					const TVector<float, 3> Position = LocalSpaceLocation + Translations[Index];
+					const TRotation<float, 3> & Rotation = Rotations[Index];
+
 					switch (Object->GetType())
 					{
 					case ImplicitObjectType::Sphere:
-						DrawSphere(PDI, Object->GetObjectChecked<TSphere<float, 3>>(), Rotations[Index], Translations[Index], Color);
+						DrawSphere(PDI, Object->GetObjectChecked<TSphere<float, 3>>(), Rotation, Position, Color);
 						break;
 
 					case ImplicitObjectType::Box:
-						DrawBox(PDI, Object->GetObjectChecked<TBox<float, 3>>(), Rotations[Index], Translations[Index], Color);
+						DrawBox(PDI, Object->GetObjectChecked<TBox<float, 3>>(), Rotation, Position, Color);
 						break;
 
 					case ImplicitObjectType::Capsule:
-						DrawCapsule(PDI, Object->GetObjectChecked<TCapsule<float>>(), Rotations[Index], Translations[Index], Color);
+						DrawCapsule(PDI, Object->GetObjectChecked<TCapsule<float>>(), Rotation, Position, Color);
 						break;
 
 					case ImplicitObjectType::Union:  // Union only used as collision tapered capsules
@@ -1034,11 +1037,11 @@ void FClothingSimulation::DebugDrawCollision(FPrimitiveDrawInterface* PDI) const
 								switch (SubObject->GetType())
 								{
 								case ImplicitObjectType::Sphere:
-									DrawSphere(PDI, SubObject->GetObjectChecked<TSphere<float, 3>>(), Rotations[Index], Translations[Index], Color);
+									DrawSphere(PDI, SubObject->GetObjectChecked<TSphere<float, 3>>(), Rotation, Position, Color);
 									break;
 
 								case ImplicitObjectType::TaperedCylinder:
-									DrawTaperedCylinder(PDI, SubObject->GetObjectChecked<TTaperedCylinder<float>>(), Rotations[Index], Translations[Index], Color);
+									DrawTaperedCylinder(PDI, SubObject->GetObjectChecked<TTaperedCylinder<float>>(), Rotation, Position, Color);
 									break;
 
 								default:
@@ -1049,11 +1052,11 @@ void FClothingSimulation::DebugDrawCollision(FPrimitiveDrawInterface* PDI) const
 						break;
 
 					case ImplicitObjectType::Convex:
-						DrawConvex(PDI, Object->GetObjectChecked<FConvex>(), Rotations[Index], Translations[Index], Color);
+						DrawConvex(PDI, Object->GetObjectChecked<FConvex>(), Rotation, Position, Color);
 						break;
 
 					default:
-						DrawCoordinateSystem(PDI, Rotations[Index], Translations[Index]);  // Draw everything else as a coordinate for now
+						DrawCoordinateSystem(PDI, Rotation, Position);  // Draw everything else as a coordinate for now
 						break;
 					}
 				}
