Chaos Cloth - Fix external collision update.

diff --git a/Engine/Plugins/Experimental/ChaosCloth/Source/Chaos/Private/ChaosCloth/ChaosClothingSimulationSolver.cpp b/Engine/Plugins/Experimental/ChaosCloth/Source/Chaos/Private/ChaosCloth/ChaosClothingSimulationSolver.cpp
--- a/Engine/Plugins/Experimental/ChaosCloth/Source/Chaos/Private/ChaosCloth/ChaosClothingSimulationSolver.cpp
+++ b/Engine/Plugins/Experimental/ChaosCloth/Source/Chaos/Private/ChaosCloth/ChaosClothingSimulationSolver.cpp
@@ -335,8 +335,6 @@ int32 FClothingSimulationSolver::AddCollisionParticles(int32 NumCollisionParticl
 		if (CollisionParticlesOffset + CollisionParticlesSize <= (int32)Evolution->CollisionParticles().Size() &&  // Check first that the range hasn't been reset
 			NumCollisionParticles == Evolution->GetCollisionParticleRangeSize(RecycledOffset))  // This will assert if range has been reset
 		{
-			// Always starts with particles disabled
-			EnableCollisionParticles(RecycledOffset, false);
 			return RecycledOffset;
 		}
 		// Size has changed. must reset this collision range (and all of those following up) and reallocate some new particles
