CIS - Fix more non unity build errors due to Cloth changes. [Merging CL 14167977]
CIS - Fix another non unity build error due to recent Cloth refactor. [Merging CL 14167992]

diff --git a/Engine/Plugins/Experimental/ChaosCloth/Source/Chaos/Private/ChaosCloth/ChaosClothingSimulationCollider.cpp b/Engine/Plugins/Experimental/ChaosCloth/Source/Chaos/Private/ChaosCloth/ChaosClothingSimulationCollider.cpp
--- a/Engine/Plugins/Experimental/ChaosCloth/Source/Chaos/Private/ChaosCloth/ChaosClothingSimulationCollider.cpp
+++ b/Engine/Plugins/Experimental/ChaosCloth/Source/Chaos/Private/ChaosCloth/ChaosClothingSimulationCollider.cpp
@@ -9,6 +9,10 @@
 #include "Chaos/Convex.h"
 #include "ClothingSimulation.h" // For context
 #include "ClothingAsset.h"
+#include "Engine/SkeletalMesh.h"
+#include "Components/SkeletalMeshComponent.h"
+#include "PhysicsEngine/PhysicsAsset.h"
+#include "PhysicsEngine/AggregateGeom.h"
 #include "HAL/PlatformMath.h"
 #include "Containers/ArrayView.h"
 #include "Containers/BitArray.h"
