Chaos Cloth - Iterate on Bobs' constraint crash fix. [Merging CL 14207030]

#UE4 Avoid crash in FClothConstraints due to lambdas in CreateRules passing through "this" pointers though FClothConstraints are stored in a TMap which can reallocate. For now, I made this map of pointers instead of value types instead so the "this" address doesnt change on reallocate [Merging CL 14205540]

diff --git a/Engine/Plugins/Experimental/ChaosCloth/Source/Chaos/Private/ChaosCloth/ChaosClothingSimulationSolver.cpp b/Engine/Plugins/Experimental/ChaosCloth/Source/Chaos/Private/ChaosCloth/ChaosClothingSimulationSolver.cpp
--- a/Engine/Plugins/Experimental/ChaosCloth/Source/Chaos/Private/ChaosCloth/ChaosClothingSimulationSolver.cpp
+++ b/Engine/Plugins/Experimental/ChaosCloth/Source/Chaos/Private/ChaosCloth/ChaosClothingSimulationSolver.cpp
@@ -265,7 +265,10 @@ int32 FClothingSimulationSolver::AddParticles(int32 NumParticles, uint32 GroupId
 	const int32 Offset = Evolution->AddParticleRange(NumParticles, GroupId, /*bActivate =*/ false);
 
 	// Add an empty constraints container for this range
-	ClothsConstraints.Add(Offset).Initialize(Evolution.Get(), AnimationPositions, AnimationNormals, Offset, NumParticles);
+	check(!ClothsConstraints.Find(Offset));  // We cannot already have this Offset in the map, particle ranges are always added, never removed (unless reset)
+
+	ClothsConstraints.Emplace(Offset, MakeUnique<FClothConstraints>())
+		->Initialize(Evolution.Get(), AnimationPositions, AnimationNormals, Offset, NumParticles);
 
 	// Always starts with particles disabled
 	EnableParticles(Offset, false);
diff --git a/Engine/Plugins/Experimental/ChaosCloth/Source/Chaos/Public/ChaosCloth/ChaosClothingSimulationSolver.h b/Engine/Plugins/Experimental/ChaosCloth/Source/Chaos/Public/ChaosCloth/ChaosClothingSimulationSolver.h
--- a/Engine/Plugins/Experimental/ChaosCloth/Source/Chaos/Public/ChaosCloth/ChaosClothingSimulationSolver.h
+++ b/Engine/Plugins/Experimental/ChaosCloth/Source/Chaos/Public/ChaosCloth/ChaosClothingSimulationSolver.h
@@ -103,8 +103,8 @@ namespace Chaos
 		const TVector<float, 3>* GetParticleVs(int32 Offset) const;
 		TVector<float, 3>* GetParticleVs(int32 Offset);
 		const float* GetParticleInvMasses(int32 Offset) const;
-		const FClothConstraints& GetClothConstraints(int32 Offset) const { return ClothsConstraints.FindChecked(Offset);  }
-		FClothConstraints& GetClothConstraints(int32 Offset) { return ClothsConstraints.FindChecked(Offset);  }
+		const FClothConstraints& GetClothConstraints(int32 Offset) const { return *ClothsConstraints.FindChecked(Offset); }
+		FClothConstraints& GetClothConstraints(int32 Offset) { return *ClothsConstraints.FindChecked(Offset); }
 		// ---- End of the Cloth interface ----
 
 		// ---- Collider interface ----
@@ -157,7 +157,7 @@ namespace Chaos
 		TArrayCollectionArray<TRigidTransform<float, 3>> CollisionTransforms;
 
 		// Cloth constraints
-		TMap<int32, FClothConstraints> ClothsConstraints;
+		TMap<int32, TUniquePtr<FClothConstraints>> ClothsConstraints;
 
 		// Local space simulation
 		TVector<float, 3> OldLocalSpaceLocation;
