Chaos Cloth - Merging recent UE4 cloth fixes and improvements.
- Increase the Wind Strength slider max value to 10 in the Mesh Editor Character menu. [Merging 14207388]
- Fix already applied scale on cloth simulation data. [Merging 14207398]
- Add wind liift force. [ Merging 14207443]
- Update default Drag and Lift coefficients, and add Lift migration from NvCloth. [Merging 14207938]
- Re-enable wind/velocity fields debug drawing. Remove Drag from the name of the debug draw option since it now also draws Lift. [Merging 14208392]
- Fix config migration due to config object name collision. [Merging 14208623]
- Fix bounds calculation of empty simulation (causing the thumbnail to vanish) due to a missing implementation of ShouldSimulate(). [Merging 14226731]

diff --git a/Engine/Plugins/Experimental/ChaosCloth/Source/Chaos/Private/ChaosCloth/ChaosClothConfig.cpp b/Engine/Plugins/Experimental/ChaosCloth/Source/Chaos/Private/ChaosCloth/ChaosClothConfig.cpp
--- a/Engine/Plugins/Experimental/ChaosCloth/Source/Chaos/Private/ChaosCloth/ChaosClothConfig.cpp
+++ b/Engine/Plugins/Experimental/ChaosCloth/Source/Chaos/Private/ChaosCloth/ChaosClothConfig.cpp
@@ -62,6 +62,7 @@ void UChaosClothConfig::MigrateFrom(const FClothConfig_Legacy& ClothConfig)
 	ShapeTargetStiffness = 0.f;
 
 	DragCoefficient = (ClothConfig.WindMethod == EClothingWindMethod_Legacy::Accurate) ? ClothConfig.WindDragCoefficient: 0.07f;  // Only Accurate wind uses the WindDragCoefficient
+	LiftCoefficient = (ClothConfig.WindMethod == EClothingWindMethod_Legacy::Accurate) ? ClothConfig.WindLiftCoefficient: 0.035f;  // Only Accurate wind uses the WindLiftCoefficient
 
 	const float Damping = (ClothConfig.Damping.X + ClothConfig.Damping.Y + ClothConfig.Damping.Z) / 3.f;
 	DampingCoefficient = FMath::Clamp(Damping * Damping * 0.7f, 0.f, 1.f);  // Nv Cloth seems to have a different damping formulation.
diff --git a/Engine/Plugins/Experimental/ChaosCloth/Source/Chaos/Private/ChaosCloth/ChaosClothingSimulation.cpp b/Engine/Plugins/Experimental/ChaosCloth/Source/Chaos/Private/ChaosCloth/ChaosClothingSimulation.cpp
--- a/Engine/Plugins/Experimental/ChaosCloth/Source/Chaos/Private/ChaosCloth/ChaosClothingSimulation.cpp
+++ b/Engine/Plugins/Experimental/ChaosCloth/Source/Chaos/Private/ChaosCloth/ChaosClothingSimulation.cpp
@@ -19,6 +19,7 @@
 #include "Chaos/PBDSphericalConstraint.h"
 #include "Chaos/PBDAnimDriveConstraint.h"
 #include "Chaos/PBDLongRangeConstraints.h"
+#include "Chaos/VelocityField.h"
 #endif  // #if WITH_EDITOR || CHAOS_DEBUG_DRAW
 
 #if WITH_EDITOR
@@ -51,7 +52,7 @@ namespace ChaosClothingSimulationConsoleVariables
 	TAutoConsoleVariable<bool> CVarDebugMaxDistances        (TEXT("p.ChaosCloth.DebugDrawMaxDistances"        ), false, TEXT("Whether to debug draw the Chaos Cloth max distances"), ECVF_Cheat);
 	TAutoConsoleVariable<bool> CVarDebugAnimDrive           (TEXT("p.ChaosCloth.DebugDrawAnimDrive"           ), false, TEXT("Whether to debug draw the Chaos Cloth anim drive"), ECVF_Cheat);
 	TAutoConsoleVariable<bool> CVarDebugLongRangeConstraint (TEXT("p.ChaosCloth.DebugDrawLongRangeConstraint" ), false, TEXT("Whether to debug draw the Chaos Cloth long range constraint (aka tether constraint)"), ECVF_Cheat);
-	TAutoConsoleVariable<bool> CVarDebugWindDragForces      (TEXT("p.ChaosCloth.DebugDrawWindDragForces"      ), false, TEXT("Whether to debug draw the Chaos Cloth wind drag forces"), ECVF_Cheat);
+	TAutoConsoleVariable<bool> CVarDebugWindForces          (TEXT("p.ChaosCloth.DebugDrawWindForces"          ), false, TEXT("Whether to debug draw the Chaos Cloth wind forces"), ECVF_Cheat);
 }
 #endif  // #if CHAOS_DEBUG_DRAW
 
@@ -195,6 +196,7 @@ void FClothingSimulation::CreateActor(USkeletalMeshComponent* InOwnerComponent,
 		ClothConfig->LinearVelocityScale,
 		ClothConfig->AngularVelocityScale,
 		ClothConfig->DragCoefficient,
+		ClothConfig->LiftCoefficient,
 		ClothConfig->DampingCoefficient,
 		ClothConfig->CollisionThickness,
 		ClothConfig->FrictionCoefficient,
@@ -244,6 +246,18 @@ void FClothingSimulation::UpdateSimulationFromSharedSimConfig()
 	}
 }
 
+bool FClothingSimulation::ShouldSimulate() const
+{
+	for (const TUniquePtr<FClothingSimulationCloth>& Cloth : Cloths)
+	{
+		if (Cloth->GetLODIndex(Solver.Get()) != INDEX_NONE && Cloth->GetOffset(Solver.Get()) != INDEX_NONE)
+		{
+			return true;
+		}
+	}
+	return false;
+}
+
 void FClothingSimulation::Simulate(IClothingSimulationContext* InContext)
 {
 	SCOPE_CYCLE_COUNTER(STAT_ChaosClothSimulate);
@@ -305,7 +319,7 @@ void FClothingSimulation::Simulate(IClothingSimulationContext* InContext)
 	if (ChaosClothingSimulationConsoleVariables::CVarDebugMaxDistances        .GetValueOnAnyThread()) { DebugDrawMaxDistances        (); }
 	if (ChaosClothingSimulationConsoleVariables::CVarDebugAnimDrive           .GetValueOnAnyThread()) { DebugDrawAnimDrive           (); }
 	if (ChaosClothingSimulationConsoleVariables::CVarDebugLongRangeConstraint .GetValueOnAnyThread()) { DebugDrawLongRangeConstraint (); }
-	if (ChaosClothingSimulationConsoleVariables::CVarDebugWindDragForces      .GetValueOnAnyThread()) { DebugDrawWindDragForces      (); }
+	if (ChaosClothingSimulationConsoleVariables::CVarDebugWindForces          .GetValueOnAnyThread()) { DebugDrawWindForces          (); }
 #endif  // #if CHAOS_DEBUG_DRAW
 }
 
@@ -340,8 +354,8 @@ void FClothingSimulation::GetSimulationData(
 
 		for (int32 Index = 0; Index < Data.Positions.Num(); ++Index)
 		{
-			Data.Positions[Index] = OwnerTransform.InverseTransformPositionNoScale(Data.Positions[Index] + LocalSpaceLocation);  // Move into world space first
-			Data.Normals[Index] = OwnerTransform.InverseTransformVectorNoScale(-Data.Normals[Index]);  // Normals are inverted due to how barycentric coordinates are calculated (see GetPointBaryAndDist in ClothingMeshUtils.cpp)
+			Data.Positions[Index] = OwnerTransform.InverseTransformPosition(Data.Positions[Index] + LocalSpaceLocation);  // Move into world space first
+			Data.Normals[Index] = OwnerTransform.InverseTransformVector(-Data.Normals[Index]);  // Normals are inverted due to how barycentric coordinates are calculated (see GetPointBaryAndDist in ClothingMeshUtils.cpp)
 		}
     }
 }
@@ -444,6 +458,7 @@ void FClothingSimulation::RefreshClothConfig(const IClothingSimulationContext* I
 			ClothConfig->LinearVelocityScale,
 			ClothConfig->AngularVelocityScale,
 			ClothConfig->DragCoefficient,
+			ClothConfig->LiftCoefficient,
 			ClothConfig->DampingCoefficient,
 			ClothConfig->CollisionThickness,
 			ClothConfig->FrictionCoefficient,
@@ -1394,32 +1409,38 @@ void FClothingSimulation::DebugDrawLongRangeConstraint(FPrimitiveDrawInterface*
 	}
 }
 
-void FClothingSimulation::DebugDrawWindDragForces(FPrimitiveDrawInterface* PDI) const
+void FClothingSimulation::DebugDrawWindForces(FPrimitiveDrawInterface* PDI) const
 {
-	// TODO: Add lift and re-enable debug draw code
-	//const TPBDParticles<float, 3>& Particles = Evolution->Particles();
-
-	//for (int32 Index = 0; Index < Assets.Num(); ++Index)
-	//{
-	//	const UClothingAssetCommon* const Asset = Assets[Index];
-	//	if (!Asset) { continue; }
-
-	//	const FVelocityField& VelocityField = Evolution->GetVelocityField(Index);
-
-	//	const TArrayView<TVector<int32, 3>>& Elements = VelocityField.GetElements();
-	//	const TArray<TVector<float, 3>>& Forces = VelocityField.GetForces();
-
-	//	for (int32 ElementIndex = 0; ElementIndex < Elements.Num(); ++ElementIndex)
-	//	{
-	//		const TVector<int32, 3>& Element = Elements[ElementIndex];
-	//		const TVector<float, 3> Position = LocalSpaceLocation + (
-	//			Particles.X(Element[0]) +
-	//			Particles.X(Element[1]) +
-	//			Particles.X(Element[2])) / 3.f;
-	//		const TVector<float, 3>& Force = Forces[ElementIndex];
-	//		DrawLine(PDI, Position, Position + Force, FColor::Green);
-	//	}
-	//}
+	check(Solver);
+
+	const TVector<float, 3>& LocalSpaceLocation = Solver->GetLocalSpaceLocation();
+
+	for (const FClothingSimulationCloth* const Cloth : Solver->GetCloths())
+	{
+		const int32 Offset = Cloth->GetOffset(Solver.Get());
+		if (Offset == INDEX_NONE)
+		{
+			continue;
+		}
+
+		const TVelocityField<float, 3>& VelocityField = Solver->GetWindVelocityField(Cloth->GetGroupId());
+
+		const TConstArrayView<TVector<int32, 3>>& Elements = VelocityField.GetElements();
+		const TConstArrayView<TVector<float, 3>> Forces = VelocityField.GetForces();
+
+		const TConstArrayView<TVector<float, 3>> Positions = Cloth->GetParticlePositions(Solver.Get());
+
+		for (int32 ElementIndex = 0; ElementIndex < Elements.Num(); ++ElementIndex)
+		{
+			const TVector<int32, 3>& Element = Elements[ElementIndex];
+			const TVector<float, 3> Position = LocalSpaceLocation + (
+				Positions[Element[0]] +
+				Positions[Element[1]] +
+				Positions[Element[2]]) / 3.f;
+			const TVector<float, 3>& Force = Forces[ElementIndex] * 10.f;
+			DrawLine(PDI, Position, Position + Force, FColor::Green);
+		}
+	}
 }
 
 void FClothingSimulation::DebugDrawLocalSpace(FPrimitiveDrawInterface* PDI) const
diff --git a/Engine/Plugins/Experimental/ChaosCloth/Source/Chaos/Private/ChaosCloth/ChaosClothingSimulationCloth.cpp b/Engine/Plugins/Experimental/ChaosCloth/Source/Chaos/Private/ChaosCloth/ChaosClothingSimulationCloth.cpp
--- a/Engine/Plugins/Experimental/ChaosCloth/Source/Chaos/Private/ChaosCloth/ChaosClothingSimulationCloth.cpp
+++ b/Engine/Plugins/Experimental/ChaosCloth/Source/Chaos/Private/ChaosCloth/ChaosClothingSimulationCloth.cpp
@@ -339,6 +339,7 @@ FClothingSimulationCloth::FClothingSimulationCloth(
 	const TVector<float, 3>& InLinearVelocityScale,
 	float InAngularVelocityScale,
 	float InDragCoefficient,
+	float InLiftCoefficient,
 	float InDampingCoefficient,
 	float InCollisionThickness,
 	float InFrictionCoefficient,
@@ -371,6 +372,7 @@ FClothingSimulationCloth::FClothingSimulationCloth(
 	, LinearVelocityScale(InLinearVelocityScale)
 	, AngularVelocityScale(InAngularVelocityScale)
 	, DragCoefficient(InDragCoefficient)
+	, LiftCoefficient(InLiftCoefficient)
 	, DampingCoefficient(InDampingCoefficient)
 	, CollisionThickness(InCollisionThickness)
 	, FrictionCoefficient(InFrictionCoefficient)
@@ -673,7 +675,7 @@ void FClothingSimulationCloth::Update(FClothingSimulationSolver* Solver)
 		Solver->SetGravity(GroupId, GetGravity(Solver));
 
 		// Update wind
-		Solver->SetWindVelocityField(GroupId, DragCoefficient, &GetTriangleMesh(Solver));
+		Solver->SetWindVelocityField(GroupId, DragCoefficient, LiftCoefficient, &GetTriangleMesh(Solver));
 
 		// Update general solver properties
 		Solver->SetProperties(GroupId, DampingCoefficient, CollisionThickness, FrictionCoefficient);
diff --git a/Engine/Plugins/Experimental/ChaosCloth/Source/Chaos/Private/ChaosCloth/ChaosClothingSimulationSolver.cpp b/Engine/Plugins/Experimental/ChaosCloth/Source/Chaos/Private/ChaosCloth/ChaosClothingSimulationSolver.cpp
--- a/Engine/Plugins/Experimental/ChaosCloth/Source/Chaos/Private/ChaosCloth/ChaosClothingSimulationSolver.cpp
+++ b/Engine/Plugins/Experimental/ChaosCloth/Source/Chaos/Private/ChaosCloth/ChaosClothingSimulationSolver.cpp
@@ -558,11 +558,16 @@ void FClothingSimulationSolver::SetWindVelocity(const TVector<float, 3>& InWindV
 	WindVelocity = InWindVelocity * ChaosClothingSimulationSolverConstant::WorldScale;
 }
 
-void FClothingSimulationSolver::SetWindVelocityField(uint32 GroupId, float DragCoefficient, const TTriangleMesh<float>* TriangleMesh)
+void FClothingSimulationSolver::SetWindVelocityField(uint32 GroupId, float DragCoefficient, float LiftCoefficient, const TTriangleMesh<float>* TriangleMesh)
 {
 	TVelocityField<float, 3>& VelocityField = Evolution->GetVelocityField(GroupId);
 	VelocityField.SetGeometry(TriangleMesh);
-	VelocityField.SetDragCoefficient(DragCoefficient);
+	VelocityField.SetCoefficients(DragCoefficient, LiftCoefficient);
+}
+
+const TVelocityField<float, 3>& FClothingSimulationSolver::GetWindVelocityField(uint32 GroupId)
+{
+	return Evolution->GetVelocityField(GroupId);
 }
 
 void FClothingSimulationSolver::SetSelfCollisions(uint32 GroupId, float SelfCollisionThickness, const TTriangleMesh<float>* TriangleMesh)
diff --git a/Engine/Plugins/Experimental/ChaosCloth/Source/Chaos/Public/ChaosCloth/ChaosClothConfig.h b/Engine/Plugins/Experimental/ChaosCloth/Source/Chaos/Public/ChaosCloth/ChaosClothConfig.h
--- a/Engine/Plugins/Experimental/ChaosCloth/Source/Chaos/Public/ChaosCloth/ChaosClothConfig.h
+++ b/Engine/Plugins/Experimental/ChaosCloth/Source/Chaos/Public/ChaosCloth/ChaosClothConfig.h
@@ -125,9 +125,13 @@ public:
 	UPROPERTY(EditAnywhere, Category = "Environmental Properties", meta = (UIMin = "0", UIMax = "1", ClampMin = "0", ClampMax = "1"))
 	float DampingCoefficient = 0.01f;
 
-	// The drag coefficient applying on each particle.
+	// The aerodynamic coefficient of drag applying on each particle.
 	UPROPERTY(EditAnywhere, Category = "Environmental Properties", meta = (UIMin = "0", UIMax = "1", ClampMin = "0", ClampMax = "10"))
-	float DragCoefficient = 0.1f;
+	float DragCoefficient = 0.07f;
+
+	// The aerodynamic coefficient of lift applying on each particle.
+	UPROPERTY(EditAnywhere, Category = "Environmental Properties", meta = (UIMin = "0", UIMax = "1", ClampMin = "0", ClampMax = "10"))
+	float LiftCoefficient = 0.035f;
 
 	// Use the config gravity value instead of world gravity.
 	UPROPERTY(EditAnywhere, Category = "Environmental Properties", meta = (InlineEditConditionToggle))
diff --git a/Engine/Plugins/Experimental/ChaosCloth/Source/Chaos/Public/ChaosCloth/ChaosClothingSimulation.h b/Engine/Plugins/Experimental/ChaosCloth/Source/Chaos/Public/ChaosCloth/ChaosClothingSimulation.h
--- a/Engine/Plugins/Experimental/ChaosCloth/Source/Chaos/Public/ChaosCloth/ChaosClothingSimulation.h
+++ b/Engine/Plugins/Experimental/ChaosCloth/Source/Chaos/Public/ChaosCloth/ChaosClothingSimulation.h
@@ -47,7 +47,7 @@ namespace Chaos
 		virtual void CreateActor(USkeletalMeshComponent* InOwnerComponent, UClothingAssetBase* InAsset, int32 SimDataIndex) override;
 		virtual void DestroyActors() override;
 
-		virtual bool ShouldSimulate() const override { return true; }
+		virtual bool ShouldSimulate() const override;
 		virtual void Simulate(IClothingSimulationContext* InContext) override;
 		virtual void GetSimulationData(TMap<int32, FClothSimulData>& OutData, USkeletalMeshComponent* InOwnerComponent, USkinnedMeshComponent* InOverrideComponent) const override;
 
@@ -104,7 +104,7 @@ namespace Chaos
 		CHAOSCLOTH_API void DebugDrawMaxDistances(FPrimitiveDrawInterface* PDI = nullptr) const;
 		CHAOSCLOTH_API void DebugDrawAnimDrive(FPrimitiveDrawInterface* PDI = nullptr) const;
 		CHAOSCLOTH_API void DebugDrawLongRangeConstraint(FPrimitiveDrawInterface* PDI = nullptr) const;
-		CHAOSCLOTH_API void DebugDrawWindDragForces(FPrimitiveDrawInterface* PDI = nullptr) const;
+		CHAOSCLOTH_API void DebugDrawWindForces(FPrimitiveDrawInterface* PDI = nullptr) const;
 		CHAOSCLOTH_API void DebugDrawLocalSpace(FPrimitiveDrawInterface* PDI = nullptr) const;
 #endif  // #if WITH_EDITOR || CHAOS_DEBUG_DRAW
 
@@ -154,4 +154,4 @@ namespace Chaos
 		UMaterial* DebugClothMaterialVertex;
 #endif  // #if WITH_EDITOR
 	};
-} // namespace Chaos
+} // namespace Chaos
\ No newline at end of file
diff --git a/Engine/Plugins/Experimental/ChaosCloth/Source/Chaos/Public/ChaosCloth/ChaosClothingSimulationCloth.h b/Engine/Plugins/Experimental/ChaosCloth/Source/Chaos/Public/ChaosCloth/ChaosClothingSimulationCloth.h
--- a/Engine/Plugins/Experimental/ChaosCloth/Source/Chaos/Public/ChaosCloth/ChaosClothingSimulationCloth.h
+++ b/Engine/Plugins/Experimental/ChaosCloth/Source/Chaos/Public/ChaosCloth/ChaosClothingSimulationCloth.h
@@ -49,6 +49,7 @@ namespace Chaos
 			const TVector<float, 3>& InLinearVelocityScale,
 			float InAngularVelocityScale,
 			float InDragCoefficient,
+			float InLiftCoefficient,
 			float InDampingCoefficient,
 			float InCollisionThickness,
 			float InFrictionCoefficient,
@@ -178,6 +179,7 @@ namespace Chaos
 		TVector<float, 3> LinearVelocityScale;  // Linear ratio applied to the reference bone transforms
 		float AngularVelocityScale;  // Angular ratio factor applied to the reference bone transforms
 		float DragCoefficient;
+		float LiftCoefficient;
 		float DampingCoefficient;
 		float CollisionThickness;
 		float FrictionCoefficient;
diff --git a/Engine/Plugins/Experimental/ChaosCloth/Source/Chaos/Public/ChaosCloth/ChaosClothingSimulationSolver.h b/Engine/Plugins/Experimental/ChaosCloth/Source/Chaos/Public/ChaosCloth/ChaosClothingSimulationSolver.h
--- a/Engine/Plugins/Experimental/ChaosCloth/Source/Chaos/Public/ChaosCloth/ChaosClothingSimulationSolver.h
+++ b/Engine/Plugins/Experimental/ChaosCloth/Source/Chaos/Public/ChaosCloth/ChaosClothingSimulationSolver.h
@@ -13,6 +13,7 @@ namespace Chaos
 {
 	template<typename T, int d> class TPBDEvolution;
 	template<typename T> class TTriangleMesh;
+	template<typename T, int d> class TVelocityField;
 
 	class FClothingSimulationCloth;
 	class FClothingSimulationCollider;
@@ -83,7 +84,8 @@ namespace Chaos
 		void SetGravity(uint32 GroupId, const TVector<float, 3>& Gravity);
 
 		// Set the geometry affected by wind, or disable if TriangleMesh is null.
-		void SetWindVelocityField(uint32 GroupId, float DragCoefficient, const TTriangleMesh<float>* TriangleMesh = nullptr);
+		void SetWindVelocityField(uint32 GroupId, float DragCoefficient, float LiftCoefficient, const TTriangleMesh<float>* TriangleMesh = nullptr);
+		const TVelocityField<float, 3>&  GetWindVelocityField(uint32 GroupId);
 
 		// Enable self collision, or disable if TriangleMesh is null.
 		void SetSelfCollisions(uint32 GroupId, float SelfCollisionThickness, const TTriangleMesh<float>* TriangleMesh = nullptr);
diff --git a/Engine/Plugins/Experimental/ChaosClothEditor/Source/Chaos/Private/ChaosClothEditor/ChaosSimulationEditorExtender.cpp b/Engine/Plugins/Experimental/ChaosClothEditor/Source/Chaos/Private/ChaosClothEditor/ChaosSimulationEditorExtender.cpp
--- a/Engine/Plugins/Experimental/ChaosClothEditor/Source/Chaos/Private/ChaosClothEditor/ChaosSimulationEditorExtender.cpp
+++ b/Engine/Plugins/Experimental/ChaosClothEditor/Source/Chaos/Private/ChaosClothEditor/ChaosSimulationEditorExtender.cpp
@@ -67,8 +67,8 @@ const FVisualizationOption FVisualizationOption::OptionData[] =
 	FVisualizationOption(&Chaos::FClothingSimulation::DebugDrawMaxDistanceValues   , LOCTEXT("ChaosVisName_MaxDistanceValue"    , "Max Distances As Numbers"   ), LOCTEXT("ChaosVisName_MaxDistanceValue_ToolTip"    , "Draws the current max distances as numbers")),
 	FVisualizationOption(&Chaos::FClothingSimulation::DebugDrawAnimDrive           , LOCTEXT("ChaosVisName_AnimDrive"           , "Anim Drive"                 ), LOCTEXT("ChaosVisName_AnimDrive_Tooltip"           , "Draws the current skinned reference mesh for the simulation which anim drive will attempt to reach if enabled")),
 	FVisualizationOption(&Chaos::FClothingSimulation::DebugDrawLongRangeConstraint , LOCTEXT("ChaosVisName_LongRangeConstraint" , "Long Range Constraint"      ), LOCTEXT("ChaosVisName_LongRangeConstraint_Tooltip" , "Draws the long range attachment constraint distances")),
-	FVisualizationOption(&Chaos::FClothingSimulation::DebugDrawWindDragForces      , LOCTEXT("ChaosVisName_WindDragForces"      , "Wind drag forces"           ), LOCTEXT("ChaosVisName_Wind_Tooltip"                , "Draws the Wind drag forces")),
-	FVisualizationOption(&Chaos::FClothingSimulation::DebugDrawLocalSpace          , LOCTEXT("ChaosVisName_LocalSpace"          , "Local space reference bone" ), LOCTEXT("ChaosVisName_LocalSpace_Tooltip"          , "Draws the local space reference bone")),
+	FVisualizationOption(&Chaos::FClothingSimulation::DebugDrawWindForces          , LOCTEXT("ChaosVisName_WindForces"          , "Wind Aerodynamic Forces"    ), LOCTEXT("ChaosVisName_Wind_Tooltip"                , "Draws the Wind drag and lift forces")),
+	FVisualizationOption(&Chaos::FClothingSimulation::DebugDrawLocalSpace          , LOCTEXT("ChaosVisName_LocalSpace"          , "Local Space Reference Bone" ), LOCTEXT("ChaosVisName_LocalSpace_Tooltip"          , "Draws the local space reference bone")),
 };
 const uint32 FVisualizationOption::Count = sizeof(OptionData) / sizeof(FVisualizationOption);
 
diff --git a/Engine/Source/Editor/Persona/Private/SAnimViewportToolBar.cpp b/Engine/Source/Editor/Persona/Private/SAnimViewportToolBar.cpp
--- a/Engine/Source/Editor/Persona/Private/SAnimViewportToolBar.cpp
+++ b/Engine/Source/Editor/Persona/Private/SAnimViewportToolBar.cpp
@@ -80,7 +80,7 @@ public:
 					.MinValue(0)
 					.AllowSpin(true)
 					.MinSliderValue(0)
-					.MaxSliderValue(2)
+					.MaxSliderValue(10)
 					.Value(AnimViewportPtr.Pin().ToSharedRef(), &SAnimationEditorViewportTabBody::GetWindStrengthSliderValue)
 					.OnValueChanged(SSpinBox<float>::FOnValueChanged::CreateSP(AnimViewportPtr.Pin().ToSharedRef(), &SAnimationEditorViewportTabBody::SetWindStrength))
 				]
diff --git a/Engine/Source/Runtime/ClothingSystemRuntimeCommon/Private/ClothingAsset.cpp b/Engine/Source/Runtime/ClothingSystemRuntimeCommon/Private/ClothingAsset.cpp
--- a/Engine/Source/Runtime/ClothingSystemRuntimeCommon/Private/ClothingAsset.cpp
+++ b/Engine/Source/Runtime/ClothingSystemRuntimeCommon/Private/ClothingAsset.cpp
@@ -830,14 +830,17 @@ void UClothingAssetCommon::PostLoad()
 	if (ClothSimConfig_DEPRECATED)
 	{
 		ClothSimConfig_DEPRECATED->ConditionalPostLoad();  // PostLoad old configs before replacing them
+		ClothSimConfig_DEPRECATED->Rename(nullptr, nullptr, REN_DoNotDirty | REN_DontCreateRedirectors | REN_ForceNoResetLoaders | REN_NonTransactional);  // Rename the config so that the name doesn't collide with the new config map name
 	}
 	if (ChaosClothSimConfig_DEPRECATED)
 	{
 		ChaosClothSimConfig_DEPRECATED->ConditionalPostLoad();  // PostLoad old configs before replacing them
+		ChaosClothSimConfig_DEPRECATED->Rename(nullptr, nullptr, REN_DoNotDirty | REN_DontCreateRedirectors | REN_ForceNoResetLoaders | REN_NonTransactional);  // Rename the config so that the name doesn't collide with the new config map name
 	}
 	if (ClothSharedSimConfig_DEPRECATED)
 	{
 		ClothSharedSimConfig_DEPRECATED->ConditionalPostLoad();  // PostLoad old configs before replacing them
+		ClothSharedSimConfig_DEPRECATED->Rename(nullptr, nullptr, REN_DoNotDirty | REN_DontCreateRedirectors | REN_ForceNoResetLoaders | REN_NonTransactional);  // Rename the config so that the name doesn't collide with the new config map name
 	}
 	AddClothConfigs();
 
@@ -931,6 +934,7 @@ void UClothingAssetCommon::AddClothConfigs()
 				if (!ClothConfigs.Find(ClothConfigName))
 				{
 					// Create new config object
+					check(!StaticFindObject(ClothConfigClass, this, *ClothConfigClass->GetName(), true));
 					UClothConfigBase* const ClothConfig = NewObject<UClothConfigBase>(this, ClothConfigClass, ClothConfigClass->GetFName(), RF_Transactional);
 
 					// Use the legacy config struct to try find a common config as an acceptable migration source
diff --git a/Engine/Source/Runtime/Experimental/Chaos/Public/Chaos/VelocityField.h b/Engine/Source/Runtime/Experimental/Chaos/Public/Chaos/VelocityField.h
--- a/Engine/Source/Runtime/Experimental/Chaos/Public/Chaos/VelocityField.h
+++ b/Engine/Source/Runtime/Experimental/Chaos/Public/Chaos/VelocityField.h
@@ -22,6 +22,7 @@ namespace Chaos
 			const TTriangleMesh<T>& TriangleMesh,
 			const TVector<T, d>& InVelocity,
 			const T InDragCoefficient = (T)0.5,
+			const T InLiftCoefficient = (T)0.1,
 			const T InFluidDensity = (T)1.225e-6)
 			: PointToTriangleMap(TriangleMesh.GetPointToTriangleMap())
 			, Elements(TriangleMesh.GetElements())
@@ -29,7 +30,8 @@ namespace Chaos
 			, Range(TriangleMesh.GetVertexRange())
 		{
 			Forces.SetNumUninitialized(Elements.Num());
-			SetProperties(InDragCoefficient, InFluidDensity);
+			SetCoefficients(InDragCoefficient, InLiftCoefficient);
+			SetFluidDensity(InFluidDensity);
 		}
 
 		// Construct a vector field.
@@ -37,6 +39,7 @@ namespace Chaos
 			const TTriangleMesh<T>& TriangleMesh,
 			TFunction<TVector<T, d>(const TVector<T, d>&)> InGetVelocity,
 			const T InDragCoefficient = (T)0.5,
+			const T InLiftCoefficient = (T)0.1,
 			const T InFluidDensity = (T)1.225e-6)
 			: PointToTriangleMap(TriangleMesh.GetPointToTriangleMap())
 			, Elements(TriangleMesh.GetElements())
@@ -45,7 +48,8 @@ namespace Chaos
 			, Range(TriangleMesh.GetVertexRange())
 		{
 			Forces.SetNumUninitialized(Elements.Num());
-			SetProperties(InDragCoefficient, InFluidDensity);
+			SetCoefficients(InDragCoefficient, InLiftCoefficient);
+			SetFluidDensity(InFluidDensity);
 		}
 
 		virtual ~TVelocityField() {}
@@ -87,15 +91,16 @@ namespace Chaos
 			}
 		}
 
-		void SetProperties(const T InDragCoefficient, const T InFluidDensity)
+		void SetFluidDensity(const T InFluidDensity)
 		{
-			DragCoefficient = InDragCoefficient;
-			FluidDensity = InFluidDensity;
-			QuarterRhoDragCoefficient = (T)0.25 * FluidDensity * DragCoefficient;
+			QuarterRho = (T)0.25 * InFluidDensity;
 		}
 
-		void SetDragCoefficient(const T InDragCoefficient) { SetProperties(InDragCoefficient, FluidDensity); }
-		void SetFluidDensity(const T InFluidDensity) { SetProperties(DragCoefficient, InFluidDensity); }
+		void SetCoefficients(const T InDragCoefficient, const T InLiftCoefficient)
+		{
+			Cd = InDragCoefficient;
+			Cl = InLiftCoefficient;
+		}
 
 		void SetGeometry(const TTriangleMesh<T>* TriangleMesh)
 		{
@@ -127,31 +132,33 @@ namespace Chaos
 		}
 
 		const TConstArrayView<TVector<int32, 3>>& GetElements() const { return Elements; }
-		const TArray<TVector<T, d>>& GetForces() const { return Forces; }
+		TConstArrayView<TVector<T, d>> GetForces() const { return TConstArrayView<TVector<T, d>>(Forces); }
 
 	private:
 		inline void UpdateField(const TPBDParticles<T, d>& InParticles, int32 ElementIndex, const TVector<T, d>& InVelocity) const
 		{
 			const TVector<int32, 3>& Element = Elements[ElementIndex];
 
-			// Calculate direction and the relative velocity of the triangle to the flow
+			// Calculate the normal and the area of the surface exposed to the flow
+			TVector<T, d> N = TVector<T, d>::CrossProduct(
+				InParticles.X(Element[1]) - InParticles.X(Element[0]),
+				InParticles.X(Element[2]) - InParticles.X(Element[0]));
+			const T DoubleArea = N.SafeNormalize();
+
+			// Calculate the direction and the relative velocity of the triangle to the flow
 			const TVector<T, d>& SurfaceVelocity = (T)(1. / 3.) * (
 				InParticles.V(Element[0]) +
 				InParticles.V(Element[1]) +
 				InParticles.V(Element[2]));
-			TVector<T, d> Direction = InVelocity - SurfaceVelocity;
-			const T RelVelocity = Direction.SafeNormalize();
+			const TVector<T, d> V = InVelocity - SurfaceVelocity;
 
-			// Calculate the cross sectional area of the surface exposed to the flow
-			TVector<T, d> Normal = TVector<T, d>::CrossProduct(
-				InParticles.X(Element[1]) - InParticles.X(Element[0]),
-				InParticles.X(Element[2]) - InParticles.X(Element[0]));
-			const T DoubleArea = Normal.SafeNormalize();
-			const T DoubleCrossSectionalArea = DoubleArea * FMath::Abs(TVector<T, d>::DotProduct(Direction, Normal));
+			// Set the aerodynamic forces
+			const T VDotN = TVector<T, d>::DotProduct(V, N);
+			const T VSquare = TVector<T, d>::DotProduct(V, V);
 
-			// Calculate the drag force
-			const T Drag = QuarterRhoDragCoefficient * FMath::Square(RelVelocity) * DoubleCrossSectionalArea;
-			Forces[ElementIndex] = Direction * Drag;
+			Forces[ElementIndex] = QuarterRho * DoubleArea * (VDotN >= (T)0. ?  // The flow can hit either side of the triangle, so the normal might need to be reversed
+				(Cd - Cl) * VDotN * V + Cl * VSquare * N :
+				(Cl - Cd) * VDotN * V - Cl * VSquare * N);
 		}
 
 	private:
@@ -160,9 +167,9 @@ namespace Chaos
 		TVector<T, d> Velocity;
 		TFunction<TVector<T, d>(const TVector<T, d>&)> GetVelocity;
 		mutable TArray<TVector<T, 3>> Forces;
-		float QuarterRhoDragCoefficient;
-		float DragCoefficient;
-		float FluidDensity;
+		float QuarterRho;
+		float Cd;
+		float Cl;
 		TVector<int32, 2> Range;  // TODO: Remove? It is used by the check only
 	};
 }
