AutoTest - Fix Chaos Cloth NaN autotest issue. [Merging CL 14252779]

diff --git a/Engine/Plugins/Experimental/ChaosCloth/Source/Chaos/Private/ChaosCloth/ChaosClothingSimulation.cpp b/Engine/Plugins/Experimental/ChaosCloth/Source/Chaos/Private/ChaosCloth/ChaosClothingSimulation.cpp
--- a/Engine/Plugins/Experimental/ChaosCloth/Source/Chaos/Private/ChaosCloth/ChaosClothingSimulation.cpp
+++ b/Engine/Plugins/Experimental/ChaosCloth/Source/Chaos/Private/ChaosCloth/ChaosClothingSimulation.cpp
@@ -330,6 +330,12 @@ void FClothingSimulation::GetSimulationData(
 {
 	SCOPE_CYCLE_COUNTER(STAT_ChaosClothGetSimulationData);
 
+	if (!Cloths.Num() || !InOwnerComponent)
+	{
+		OutData.Reset();
+		return;
+	}
+
 	// Reset map when new cloths have appeared
 	if (OutData.Num() != Cloths.Num())
 	{
@@ -345,6 +351,12 @@ void FClothingSimulation::GetSimulationData(
 		const int32 AssetIndex = Cloth->GetGroupId();
 		FClothSimulData& Data = OutData.FindOrAdd(AssetIndex);
 
+		if (Cloth->GetLODIndex(Solver.Get()) == INDEX_NONE || Cloth->GetOffset(Solver.Get()) == INDEX_NONE)
+		{
+			Data.Reset();
+			continue;
+		}
+
 		// Output data in component space
 		Data.ComponentRelativeTransform = FTransform::Identity;
 		Data.Transform = OwnerTransform;
