Chaos Cloth - Initialize the normals when the cloth LOD is added, in case the sim data is queried before the simulation steps. This should fix the NaN issue in AutotTest. [Merging CL 14261352]

diff --git a/Engine/Plugins/Experimental/ChaosCloth/Source/Chaos/Private/ChaosCloth/ChaosClothingSimulationCloth.cpp b/Engine/Plugins/Experimental/ChaosCloth/Source/Chaos/Private/ChaosCloth/ChaosClothingSimulationCloth.cpp
--- a/Engine/Plugins/Experimental/ChaosCloth/Source/Chaos/Private/ChaosCloth/ChaosClothingSimulationCloth.cpp
+++ b/Engine/Plugins/Experimental/ChaosCloth/Source/Chaos/Private/ChaosCloth/ChaosClothingSimulationCloth.cpp
@@ -59,6 +59,9 @@ void FClothingSimulationCloth::FLODData::Add(FClothingSimulationSolver* Solver,
 	TriangleMesh.Init(MoveTemp(Elements));
 	TriangleMesh.GetPointToTriangleMap(); // Builds map for later use by GetPointNormals(), and the velocity fields
 
+	// Initialize the normals, in case the sim data is queried before the simulation steps
+	UpdateNormals(Solver);
+
 	// Set the particle masses
 	const TConstArrayView<float>& MaxDistances = WeightMaps[(int32)EChaosWeightMapTarget::MaxDistance];
 	static const float KinematicDistanceThreshold = 0.1f;  // TODO: This is not the same value as set in the painting UI but we might want to expose this value as parameter
