Chaos Cloth - Fix component scale issue on colliders. Fix local space simulation impulse at startup.

diff --git a/Engine/Plugins/Experimental/ChaosCloth/Source/Chaos/Private/ChaosCloth/ChaosClothingSimulation.cpp b/Engine/Plugins/Experimental/ChaosCloth/Source/Chaos/Private/ChaosCloth/ChaosClothingSimulation.cpp
--- a/Engine/Plugins/Experimental/ChaosCloth/Source/Chaos/Private/ChaosCloth/ChaosClothingSimulation.cpp
+++ b/Engine/Plugins/Experimental/ChaosCloth/Source/Chaos/Private/ChaosCloth/ChaosClothingSimulation.cpp
@@ -138,7 +138,8 @@ void FClothingSimulation::CreateActor(USkeletalMeshComponent* InOwnerComponent,
 		// Must set the local space location prior to adding any mesh/cloth, as otherwise the start poses would be in the wrong local space
 		const FClothingSimulationContext* const Context = static_cast<const FClothingSimulationContext*>(InOwnerComponent->GetClothingSimulationContext());
 		check(Context);
-		Solver->SetLocalSpaceLocation(bUseLocalSpaceSimulation ? Context->ComponentToWorld.GetLocation() : TVector<float, 3>(0.f));
+		static const bool bReset = true;
+		Solver->SetLocalSpaceLocation(bUseLocalSpaceSimulation ? Context->ComponentToWorld.GetLocation() : TVector<float, 3>(0.f), bReset);
 	}
 	else
 	{
@@ -429,7 +430,8 @@ void FClothingSimulation::RefreshClothConfig(const IClothingSimulationContext* I
 
 	// Update new space location
 	const FClothingSimulationContext* const Context = static_cast<const FClothingSimulationContext*>(InContext);
-	Solver->SetLocalSpaceLocation(bUseLocalSpaceSimulation ? Context->ComponentToWorld.GetLocation() : TVector<float, 3>(0.f));
+	static const bool bReset = true;
+	Solver->SetLocalSpaceLocation(bUseLocalSpaceSimulation ? Context->ComponentToWorld.GetLocation() : TVector<float, 3>(0.f), bReset);
 
 	// Reset stats
 	ResetStats();
@@ -728,7 +730,8 @@ static void DrawBox(FPrimitiveDrawInterface* PDI, const TBox<float, 3>& Box, con
 #if CHAOS_DEBUG_DRAW
 	if (!PDI)
 	{
-		FDebugDrawQueue::GetInstance().DrawDebugBox(Position, Box.Extents() * 0.5f, Rotation, Color.ToFColor(true), false, KINDA_SMALL_NUMBER, SDPG_Foreground, 0.f);
+		const TVector<float, 3> Center = Position + Rotation.RotateVector(Box.GetCenter());
+		FDebugDrawQueue::GetInstance().DrawDebugBox(Center, Box.Extents() * 0.5f, Rotation, Color.ToFColor(true), false, KINDA_SMALL_NUMBER, SDPG_Foreground, 0.f);
 		return;
 	}
 #endif
@@ -742,18 +745,23 @@ static void DrawCapsule(FPrimitiveDrawInterface* PDI, const TCapsule<float>& Cap
 {
 	const float Radius = Capsule.GetRadius();
 	const float HalfHeight = Capsule.GetHeight() * 0.5f + Radius;
+	const TVector<float, 3> Center = Position + Rotation.RotateVector(Capsule.GetCenter());
 #if CHAOS_DEBUG_DRAW
 	if (!PDI)
 	{
-		FDebugDrawQueue::GetInstance().DrawDebugCapsule(Position, HalfHeight, Radius, Rotation, Color.ToFColor(true), false, KINDA_SMALL_NUMBER, SDPG_Foreground, 0.f);
+		const FQuat Orientation = FQuat::FindBetweenNormals(TVector<float, 3>::UpVector, Capsule.GetAxis());
+		FDebugDrawQueue::GetInstance().DrawDebugCapsule(Center, HalfHeight, Radius, Rotation * Orientation, Color.ToFColor(true), false, KINDA_SMALL_NUMBER, SDPG_Foreground, 0.f);
 		return;
 	}
 #endif
 #if WITH_EDITOR
-	const FVector X = Rotation.RotateVector(FVector::ForwardVector);
-	const FVector Y = Rotation.RotateVector(FVector::RightVector);
-	const FVector Z = Rotation.RotateVector(FVector::UpVector);
-	DrawWireCapsule(PDI, Position, X, Y, Z, Color, Radius, HalfHeight, 12, SDPG_World, 0.0f, 0.001f, false);
+	const TVector<float, 3> Up = Capsule.GetAxis();
+	TVector<float, 3> Forward, Right;
+	Up.FindBestAxisVectors(Forward, Right);
+	const FVector X = Rotation.RotateVector(Forward);
+	const FVector Y = Rotation.RotateVector(Right);
+	const FVector Z = Rotation.RotateVector(Up);
+	DrawWireCapsule(PDI, Center, X, Y, Z, Color, Radius, HalfHeight, 12, SDPG_World, 0.0f, 0.001f, false);
 #endif
 }
 
diff --git a/Engine/Plugins/Experimental/ChaosCloth/Source/Chaos/Private/ChaosCloth/ChaosClothingSimulationCollider.cpp b/Engine/Plugins/Experimental/ChaosCloth/Source/Chaos/Private/ChaosCloth/ChaosClothingSimulationCollider.cpp
--- a/Engine/Plugins/Experimental/ChaosCloth/Source/Chaos/Private/ChaosCloth/ChaosClothingSimulationCollider.cpp
+++ b/Engine/Plugins/Experimental/ChaosCloth/Source/Chaos/Private/ChaosCloth/ChaosClothingSimulationCollider.cpp
@@ -25,6 +25,7 @@ void FClothingSimulationCollider::FLODData::Add(
 	FClothingSimulationSolver* Solver,
 	FClothingSimulationCloth* Cloth,
 	const FClothCollisionData& InClothCollisionData,
+	const float Scale,
 	const TArray<int32>& UsedBoneIndices)
 {
 	check(Solver);
@@ -89,57 +90,44 @@ void FClothingSimulationCollider::FLODData::Add(
 				LogChaosCloth, Warning, TEXT("Found a legacy cloth asset with a collision capsule spanning across two bones. This is not supported with the current system."));
 			UE_LOG(LogChaosCloth, VeryVerbose, TEXT("Found collision capsule on bone index %d."), BoneIndices[Index]);
 
-			const TVector<float, 3> X0 = Sphere0.LocalPosition;
-			const TVector<float, 3> X1 = Sphere1.LocalPosition;
-			const TVector<float, 3> Axis = X1 - X0;
-			const float AxisSize = Axis.Size();
+			const TVector<float, 3> X0 = Sphere0.LocalPosition * Scale;
+			const TVector<float, 3> X1 = Sphere1.LocalPosition * Scale;
+			const TVector<float, 3> Center = (X0 + X1) * 0.5f;
+			const TVector<float, 3> Axis = (X1 - X0) * 0.5f;
+			const TVector<float, 3> P0 = Center - Axis;
+			const TVector<float, 3> P1 = Center + Axis;
 
-			const float Radius0 = Sphere0.Radius;
-			const float Radius1 = Sphere1.Radius;
+			const float Radius0 = Sphere0.Radius * Scale;
+			const float Radius1 = Sphere1.Radius * Scale;
 			float MinRadius, MaxRadius;
 			if (Radius0 <= Radius1) { MinRadius = Radius0; MaxRadius = Radius1; }
 			else { MinRadius = Radius1; MaxRadius = Radius0; }
 
-			if (AxisSize < KINDA_SMALL_NUMBER)
+			BaseTransforms[Index] = TRigidTransform<float, 3>::Identity;
+
+			if (Axis.SizeSquared() < SMALL_NUMBER)
 			{
 				// Sphere
-				BaseTransforms[Index] = TRigidTransform<float, 3>::Identity;
-
 				Solver->SetCollisionGeometry(CapsuleOffset, Index,
-					MakeUnique<TSphere<float, 3>>(
-						X0,
-						MaxRadius));
+					MakeUnique<TSphere<float, 3>>(Center, MaxRadius));
 			}
 			else if (MaxRadius - MinRadius < KINDA_SMALL_NUMBER)
 			{
 				// Capsule
-				const TVector<float, 3> Center = (X0 + X1) * 0.5f;  // Construct a capsule centered at the origin along the Z axis
-				const TRotation<float, 3> Rotation = TRotation<float, 3>::FromRotatedVector(
-					TVector<float, 3>::AxisVector(2),
-					Axis.GetSafeNormal());
-
-				BaseTransforms[Index] = TRigidTransform<float, 3>(Center, Rotation);
-
-				const float HalfHeight = AxisSize * 0.5f;
-				Solver->SetCollisionGeometry(CapsuleOffset, Index,
-					MakeUnique<TCapsule<float>>(
-						TVector<float, 3>(0.f, 0.f, -HalfHeight), // Min
-						TVector<float, 3>(0.f, 0.f, HalfHeight), // Max
-						MaxRadius));
+				Solver->SetCollisionGeometry(CapsuleOffset, Index, 
+					MakeUnique<TCapsule<float>>(P0, P1, MaxRadius));
 			}
 			else
 			{
 				// Tapered capsule
-				BaseTransforms[Index] = TRigidTransform<float, 3>::Identity;
-
 				TArray<TUniquePtr<FImplicitObject>> Objects;
 				Objects.Reserve(3);
 				Objects.Add(TUniquePtr<FImplicitObject>(
-					new TTaperedCylinder<float>(X0, X1, Radius0, Radius1)));
+					new TTaperedCylinder<float>(P0, P1, Radius0, Radius1)));
 				Objects.Add(TUniquePtr<FImplicitObject>(
-					new TSphere<float, 3>(X0, Radius0)));
+					new TSphere<float, 3>(P0, Radius0)));
 				Objects.Add(TUniquePtr<FImplicitObject>(
-					new TSphere<float, 3>(X1, Radius1)));
+					new TSphere<float, 3>(P1, Radius1)));
 				Solver->SetCollisionGeometry(CapsuleOffset, Index,
 					MakeUnique<FImplicitObjectUnion>(MoveTemp(Objects)));  // TODO(Kriss.Gossart): Replace this once a TTaperedCapsule implicit type is implemented (note: this tapered cylinder with spheres is an approximation of a real tapered capsule)
 			}
@@ -171,8 +159,8 @@ void FClothingSimulationCollider::FLODData::Add(
 
 			Solver->SetCollisionGeometry(SphereOffset, Index,
 				MakeUnique<TSphere<float, 3>>(
-					Sphere.LocalPosition,
-					Sphere.Radius));
+					Sphere.LocalPosition * Scale,
+					Sphere.Radius * Scale));
 
 			++Index;
 		}
@@ -218,7 +206,7 @@ void FClothingSimulationCollider::FLODData::Add(
 					if (NormalizedPlane.Normalize())
 					{
 						const TVector<float, 3> Normal(static_cast<FVector>(NormalizedPlane));
-						const TVector<float, 3> Base = Normal * NormalizedPlane.W;
+						const TVector<float, 3> Base = Normal * NormalizedPlane.W * Scale;
 
 						Planes.Add(TPlaneConcrete<float, 3>(Base, Normal));
 					}
@@ -267,7 +255,8 @@ void FClothingSimulationCollider::FLODData::Add(
 			BoneIndices[Index] = GetMappedBoneIndex(UsedBoneIndices, Box.BoneIndex);
 			UE_LOG(LogChaosCloth, VeryVerbose, TEXT("Found collision box on bone index %d."), BoneIndices[Index]);
 
-			Solver->SetCollisionGeometry(BoxOffset, Index, MakeUnique<TBox<float, 3>>(-Box.HalfExtents, Box.HalfExtents));
+			const TVector<float, 3> HalfExtents = Box.HalfExtents * Scale;
+			Solver->SetCollisionGeometry(BoxOffset, Index, MakeUnique<TBox<float, 3>>(-HalfExtents, HalfExtents));
 		}
 	}
 
@@ -348,6 +337,7 @@ FClothingSimulationCollider::FClothingSimulationCollider(
 	, CollisionData(nullptr)
 	, bUseLODIndexOverride(bInUseLODIndexOverride)
 	, LODIndexOverride(InLODIndexOverride)
+	, Scale(1.f)
 {
 	// Prepare LOD array
 	const int32 NumLODs = Asset ? Asset->LodData.Num() : 0;
@@ -573,12 +563,24 @@ void FClothingSimulationCollider::Add(FClothingSimulationSolver* Solver, FClothi
 	int32& LODIndex = LODIndices.Add(FSolverClothPair(Solver, Cloth));
 	LODIndex = INDEX_NONE;
 
+	// Initialize scale
+	const FClothingSimulationContextCommon* const Context = SkeletalMeshComponent ? static_cast<const FClothingSimulationContextCommon*>(SkeletalMeshComponent->GetClothingSimulationContext()) : nullptr;
+	const TVector<float, 3> Scale3D = Context ? Context->ComponentToWorld.GetScale3D() : TVector<float, 3>(1.f);
+	UE_CLOG(FMath::Abs(Scale3D.X - Scale3D.Y) > KINDA_SMALL_NUMBER || FMath::Abs(Scale3D.X - Scale3D.Z) > KINDA_SMALL_NUMBER,
+		LogChaosCloth, Warning, TEXT(
+			"Actor '%s' component '%s' has a non uniform scale, and has a cloth simulation attached. "
+			"The collision volumes might no longer correctly match the shape of the mesh. "
+			"Please update this component transform scale with the same value for all scale axis."),
+		SkeletalMeshComponent->GetOwner() ? *SkeletalMeshComponent->GetOwner()->GetName() : TEXT("None"),
+		*SkeletalMeshComponent->GetName());
+	Scale = Scale3D.X;
+
 	// Create physics asset collisions, this will affect all LODs, so store it at index 0
 	FClothCollisionData PhysicsAssetCollisionData;
 	TArray<int32> UsedBoneIndices;
 	ExtractPhysicsAssetCollision(PhysicsAssetCollisionData, UsedBoneIndices);
 
-	LODData[(int32)ECollisionDataType::LODless].Add(Solver, Cloth, PhysicsAssetCollisionData, UsedBoneIndices);
+	LODData[(int32)ECollisionDataType::LODless].Add(Solver, Cloth, PhysicsAssetCollisionData, Scale, UsedBoneIndices);
 
 	// Create legacy asset LOD collisions
 	const int32 NumLODs = Asset ? Asset->LodData.Num() : 0;
@@ -588,8 +590,8 @@ void FClothingSimulationCollider::Add(FClothingSimulationSolver* Solver, FClothi
 		const FClothCollisionData& AssetCollisionData = Asset->LodData[Index].CollisionData;
 		UE_CLOG(AssetCollisionData.Spheres.Num() > 0 || AssetCollisionData.SphereConnections.Num() > 0 || AssetCollisionData.Convexes.Num() > 0,
 			LogChaosCloth, Warning, TEXT(
-				"Actor '%s' component '%s' has %d sphere, %d capsule, and %d convex collision objects for"
-				"physics authored as part of a LOD construct, probably by the Apex cloth authoring system."
+				"Actor '%s' component '%s' has %d sphere, %d capsule, and %d convex collision objects for "
+				"physics authored as part of a LOD construct, probably by the Apex cloth authoring system. "
 				"This is deprecated. Please update your asset!"),
 			SkeletalMeshComponent->GetOwner() ? *SkeletalMeshComponent->GetOwner()->GetName() : TEXT("None"),
 			*SkeletalMeshComponent->GetName(),
@@ -598,7 +600,7 @@ void FClothingSimulationCollider::Add(FClothingSimulationSolver* Solver, FClothi
 			AssetCollisionData.Convexes.Num());
 
 		// Add legacy collision
-		LODData[(int32)ECollisionDataType::LODs + Index].Add(Solver, Cloth, AssetCollisionData, Asset->UsedBoneIndices);
+		LODData[(int32)ECollisionDataType::LODs + Index].Add(Solver, Cloth, AssetCollisionData, Scale, Asset->UsedBoneIndices);
 	}
 }
 
diff --git a/Engine/Plugins/Experimental/ChaosCloth/Source/Chaos/Private/ChaosCloth/ChaosClothingSimulationSolver.cpp b/Engine/Plugins/Experimental/ChaosCloth/Source/Chaos/Private/ChaosCloth/ChaosClothingSimulationSolver.cpp
--- a/Engine/Plugins/Experimental/ChaosCloth/Source/Chaos/Private/ChaosCloth/ChaosClothingSimulationSolver.cpp
+++ b/Engine/Plugins/Experimental/ChaosCloth/Source/Chaos/Private/ChaosCloth/ChaosClothingSimulationSolver.cpp
@@ -110,6 +110,15 @@ FClothingSimulationSolver::~FClothingSimulationSolver()
 {
 }
 
+void FClothingSimulationSolver::SetLocalSpaceLocation(const TVector<float, 3>& InLocalSpaceLocation, bool bReset)
+{
+	LocalSpaceLocation = InLocalSpaceLocation;
+	if (bReset)
+	{
+		OldLocalSpaceLocation = InLocalSpaceLocation;
+	}
+}
+
 void FClothingSimulationSolver::SetCloths(TArray<FClothingSimulationCloth*>&& InCloths)
 {
 	// Remove old cloths
diff --git a/Engine/Plugins/Experimental/ChaosCloth/Source/Chaos/Public/ChaosCloth/ChaosClothingSimulationCollider.h b/Engine/Plugins/Experimental/ChaosCloth/Source/Chaos/Public/ChaosCloth/ChaosClothingSimulationCollider.h
--- a/Engine/Plugins/Experimental/ChaosCloth/Source/Chaos/Public/ChaosCloth/ChaosClothingSimulationCollider.h
+++ b/Engine/Plugins/Experimental/ChaosCloth/Source/Chaos/Public/ChaosCloth/ChaosClothingSimulationCollider.h
@@ -83,6 +83,7 @@ namespace Chaos
 				FClothingSimulationSolver* Solver,
 				FClothingSimulationCloth* Cloth,
 				const FClothCollisionData& InClothCollisionData,
+				const float Scale = 1.f,
 				const TArray<int32>& UsedBoneIndices = TArray<int32>());
 			void Remove(FClothingSimulationSolver* Solver, FClothingSimulationCloth* Cloth);
 
@@ -107,5 +108,8 @@ namespace Chaos
 		// Collision primitives
 		TArray<FLODData> LODData;  // Actual LODs start at LODStart
 		TMap<FSolverClothPair, int32> LODIndices;
+
+		// Initial scale
+		float Scale;
 	};
 } // namespace Chaos
\ No newline at end of file
diff --git a/Engine/Plugins/Experimental/ChaosCloth/Source/Chaos/Public/ChaosCloth/ChaosClothingSimulationSolver.h b/Engine/Plugins/Experimental/ChaosCloth/Source/Chaos/Public/ChaosCloth/ChaosClothingSimulationSolver.h
--- a/Engine/Plugins/Experimental/ChaosCloth/Source/Chaos/Public/ChaosCloth/ChaosClothingSimulationSolver.h
+++ b/Engine/Plugins/Experimental/ChaosCloth/Source/Chaos/Public/ChaosCloth/ChaosClothingSimulationSolver.h
@@ -27,7 +27,7 @@ namespace Chaos
 		~FClothingSimulationSolver();
 
 		// ---- Animatable property setters ----
-		void SetLocalSpaceLocation(const TVector<float, 3>& InLocalSpaceLocation) { LocalSpaceLocation = InLocalSpaceLocation; }
+		void SetLocalSpaceLocation(const TVector<float, 3>& InLocalSpaceLocation, bool bReset = false);
 		const TVector<float, 3>& GetLocalSpaceLocation() const { return LocalSpaceLocation; }
 
 		// Disables all Cloths gravity override mechanism
