CIS - Fix ChaosClothingSimulationCollider::Scale member variable shadowing.

diff --git a/Engine/Plugins/Experimental/ChaosCloth/Source/Chaos/Private/ChaosCloth/ChaosClothingSimulationCollider.cpp b/Engine/Plugins/Experimental/ChaosCloth/Source/Chaos/Private/ChaosCloth/ChaosClothingSimulationCollider.cpp
--- a/Engine/Plugins/Experimental/ChaosCloth/Source/Chaos/Private/ChaosCloth/ChaosClothingSimulationCollider.cpp
+++ b/Engine/Plugins/Experimental/ChaosCloth/Source/Chaos/Private/ChaosCloth/ChaosClothingSimulationCollider.cpp
@@ -25,7 +25,7 @@ void FClothingSimulationCollider::FLODData::Add(
 	FClothingSimulationSolver* Solver,
 	FClothingSimulationCloth* Cloth,
 	const FClothCollisionData& InClothCollisionData,
-	const float Scale,
+	const float InScale,
 	const TArray<int32>& UsedBoneIndices)
 {
 	check(Solver);
@@ -90,15 +90,15 @@ void FClothingSimulationCollider::FLODData::Add(
 				LogChaosCloth, Warning, TEXT("Found a legacy cloth asset with a collision capsule spanning across two bones. This is not supported with the current system."));
 			UE_LOG(LogChaosCloth, VeryVerbose, TEXT("Found collision capsule on bone index %d."), BoneIndices[Index]);
 
-			const TVector<float, 3> X0 = Sphere0.LocalPosition * Scale;
-			const TVector<float, 3> X1 = Sphere1.LocalPosition * Scale;
+			const TVector<float, 3> X0 = Sphere0.LocalPosition * InScale;
+			const TVector<float, 3> X1 = Sphere1.LocalPosition * InScale;
 			const TVector<float, 3> Center = (X0 + X1) * 0.5f;
 			const TVector<float, 3> Axis = (X1 - X0) * 0.5f;
 			const TVector<float, 3> P0 = Center - Axis;
 			const TVector<float, 3> P1 = Center + Axis;
 
-			const float Radius0 = Sphere0.Radius * Scale;
-			const float Radius1 = Sphere1.Radius * Scale;
+			const float Radius0 = Sphere0.Radius * InScale;
+			const float Radius1 = Sphere1.Radius * InScale;
 			float MinRadius, MaxRadius;
 			if (Radius0 <= Radius1) { MinRadius = Radius0; MaxRadius = Radius1; }
 			else { MinRadius = Radius1; MaxRadius = Radius0; }
@@ -159,8 +159,8 @@ void FClothingSimulationCollider::FLODData::Add(
 
 			Solver->SetCollisionGeometry(SphereOffset, Index,
 				MakeUnique<TSphere<float, 3>>(
-					Sphere.LocalPosition * Scale,
-					Sphere.Radius * Scale));
+					Sphere.LocalPosition * InScale,
+					Sphere.Radius * InScale));
 
 			++Index;
 		}
@@ -206,7 +206,7 @@ void FClothingSimulationCollider::FLODData::Add(
 					if (NormalizedPlane.Normalize())
 					{
 						const TVector<float, 3> Normal(static_cast<FVector>(NormalizedPlane));
-						const TVector<float, 3> Base = Normal * NormalizedPlane.W * Scale;
+						const TVector<float, 3> Base = Normal * NormalizedPlane.W * InScale;
 
 						Planes.Add(TPlaneConcrete<float, 3>(Base, Normal));
 					}
@@ -255,7 +255,7 @@ void FClothingSimulationCollider::FLODData::Add(
 			BoneIndices[Index] = GetMappedBoneIndex(UsedBoneIndices, Box.BoneIndex);
 			UE_LOG(LogChaosCloth, VeryVerbose, TEXT("Found collision box on bone index %d."), BoneIndices[Index]);
 
-			const TVector<float, 3> HalfExtents = Box.HalfExtents * Scale;
+			const TVector<float, 3> HalfExtents = Box.HalfExtents * InScale;
 			Solver->SetCollisionGeometry(BoxOffset, Index, MakeUnique<TBox<float, 3>>(-HalfExtents, HalfExtents));
 		}
 	}
diff --git a/Engine/Plugins/Experimental/ChaosCloth/Source/Chaos/Public/ChaosCloth/ChaosClothingSimulationCollider.h b/Engine/Plugins/Experimental/ChaosCloth/Source/Chaos/Public/ChaosCloth/ChaosClothingSimulationCollider.h
--- a/Engine/Plugins/Experimental/ChaosCloth/Source/Chaos/Public/ChaosCloth/ChaosClothingSimulationCollider.h
+++ b/Engine/Plugins/Experimental/ChaosCloth/Source/Chaos/Public/ChaosCloth/ChaosClothingSimulationCollider.h
@@ -83,7 +83,7 @@ namespace Chaos
 				FClothingSimulationSolver* Solver,
 				FClothingSimulationCloth* Cloth,
 				const FClothCollisionData& InClothCollisionData,
-				const float Scale = 1.f,
+				const float InScale = 1.f,
 				const TArray<int32>& UsedBoneIndices = TArray<int32>());
 			void Remove(FClothingSimulationSolver* Solver, FClothingSimulationCloth* Cloth);
 
