Chaos Cloth - Fix collision range deactivation triggered by wrong LOD index preventing external collision shape updates.
Cloth - Fix collision source transform reference on temporary object.
[Merging 14324074]

diff --git a/Engine/Plugins/Experimental/ChaosCloth/Source/Chaos/Private/ChaosCloth/ChaosClothingSimulationCollider.cpp b/Engine/Plugins/Experimental/ChaosCloth/Source/Chaos/Private/ChaosCloth/ChaosClothingSimulationCollider.cpp
--- a/Engine/Plugins/Experimental/ChaosCloth/Source/Chaos/Private/ChaosCloth/ChaosClothingSimulationCollider.cpp
+++ b/Engine/Plugins/Experimental/ChaosCloth/Source/Chaos/Private/ChaosCloth/ChaosClothingSimulationCollider.cpp
@@ -665,9 +665,9 @@ void FClothingSimulationCollider::Update(FClothingSimulationSolver* Solver, FClo
 			// Update initial state for collisions
 			LODData[(int32)ECollisionDataType::LODless].ResetStartPose(Solver, Cloth);
 		}
-		else if (PrevLODIndex > 0)
+		else if (PrevLODIndex >= (int32)ECollisionDataType::LODs)
 		{
-			// Disable previous LOD (can't be index 0, since it is not a LOD and therefore is always left enabled)
+			// Disable previous LOD
 			LODData[PrevLODIndex].Enable(Solver, Cloth, false);
 		}
 		if (LODIndex >= (int32)ECollisionDataType::LODs)
diff --git a/Engine/Source/Runtime/Engine/Private/SkeletalMeshComponentPhysics.cpp b/Engine/Source/Runtime/Engine/Private/SkeletalMeshComponentPhysics.cpp
--- a/Engine/Source/Runtime/Engine/Private/SkeletalMeshComponentPhysics.cpp
+++ b/Engine/Source/Runtime/Engine/Private/SkeletalMeshComponentPhysics.cpp
@@ -2634,7 +2634,7 @@ void USkeletalMeshComponent::ExtractCollisionsForCloth(
 		{
 			FClothCollisionPrim_Sphere& OutSphere = OutCollisions.Spheres.Add_GetRef(CachedSphere);
 
-			const FTransform& BoneTransform = SourceComponent->GetBoneTransform(OutSphere.BoneIndex, FTransform::Identity) * ComponentToComponentTransform;
+			const FTransform BoneTransform = SourceComponent->GetBoneTransform(OutSphere.BoneIndex, FTransform::Identity) * ComponentToComponentTransform;
 			OutSphere.LocalPosition = BoneTransform.TransformPosition(OutSphere.LocalPosition);
 			OutSphere.BoneIndex = INDEX_NONE;
 		}
