Chaos Cloth - Fix velocity scale calculations when the object scale isn't 1.0.

diff --git a/Engine/Plugins/Experimental/ChaosCloth/Source/Chaos/Private/ChaosCloth/ChaosClothingSimulation.cpp b/Engine/Plugins/Experimental/ChaosCloth/Source/Chaos/Private/ChaosCloth/ChaosClothingSimulation.cpp
--- a/Engine/Plugins/Experimental/ChaosCloth/Source/Chaos/Private/ChaosCloth/ChaosClothingSimulation.cpp
+++ b/Engine/Plugins/Experimental/ChaosCloth/Source/Chaos/Private/ChaosCloth/ChaosClothingSimulation.cpp
@@ -271,16 +271,16 @@ void FClothingSimulation::Simulate(IClothingSimulationContext* InContext)
 
 	const double StartTime = FPlatformTime::Seconds();
 
+	const bool bNeedsReset = (Context->TeleportMode == EClothingTeleportMode::TeleportAndReset);
+	const bool bNeedsTeleport = (Context->TeleportMode > EClothingTeleportMode::None);
+
 	// Update Solver animatable parameters
-	Solver->SetLocalSpaceLocation(bUseLocalSpaceSimulation ? Context->ComponentToWorld.GetLocation() : TVector<float, 3>(0.f));
+	Solver->SetLocalSpaceLocation(bUseLocalSpaceSimulation ? Context->ComponentToWorld.GetLocation() : TVector<float, 3>(0.f), bNeedsReset);
 	Solver->SetWindVelocity(Context->WindVelocity);
 	Solver->SetGravity(bUseGravityOverride ? GravityOverride : Context->WorldGravity);
 	Solver->EnableClothGravityOverride(!bUseGravityOverride);  // Disable all cloth gravity overrides when the interactor takes over
 
 	// Check teleport modes
-	const bool bNeedsReset = (Context->TeleportMode == EClothingTeleportMode::TeleportAndReset);
-	const bool bNeedsTeleport = (Context->TeleportMode > EClothingTeleportMode::None);
-
 	for (const TUniquePtr<FClothingSimulationCloth>& Cloth : Cloths)
 	{
 		// Update Cloth animatable parameters
diff --git a/Engine/Plugins/Experimental/ChaosCloth/Source/Chaos/Private/ChaosCloth/ChaosClothingSimulationCloth.cpp b/Engine/Plugins/Experimental/ChaosCloth/Source/Chaos/Private/ChaosCloth/ChaosClothingSimulationCloth.cpp
--- a/Engine/Plugins/Experimental/ChaosCloth/Source/Chaos/Private/ChaosCloth/ChaosClothingSimulationCloth.cpp
+++ b/Engine/Plugins/Experimental/ChaosCloth/Source/Chaos/Private/ChaosCloth/ChaosClothingSimulationCloth.cpp
@@ -622,6 +622,7 @@ void FClothingSimulationCloth::Update(FClothingSimulationSolver* Solver)
 	// Update reference space transform from the mesh's reference bone transform  TODO: Add override in the style of LODIndexOverride
 	const TRigidTransform<float, 3> OldReferenceSpaceTransform = ReferenceSpaceTransform;
 	ReferenceSpaceTransform = Mesh->GetReferenceBoneTransform();
+	ReferenceSpaceTransform.SetScale3D(TVector<float, 3>(1.f));
 
 	// Update Cloth Colliders
 	for (FClothingSimulationCollider* Collider : Colliders)
@@ -697,12 +698,14 @@ void FClothingSimulationCloth::Update(FClothingSimulationSolver* Solver)
 
 			// Reset to start pose
 			LODData[LODIndex].ResetStartPose(Solver);
+			UE_LOG(LogChaosCloth, VeryVerbose, TEXT("Cloth in group Id %d Needs reset."), GroupId);
 		}
 		else if (bNeedsTeleport)
 		{
 			// Remove all impulse velocity from the last frame
 			OutLinearVelocityScale = TVector<float, 3>(0.f);
 			OutAngularVelocityScale = 0.f;
+			UE_LOG(LogChaosCloth, VeryVerbose, TEXT("Cloth in group Id %d Needs teleport."), GroupId);
 		}
 		else
 		{
