Chaos Cloth - Fix crash when skeletal mesh with more than 1 LOD gets loaded with only 1 cloth LODs. [Merging 14339913]

diff --git a/Engine/Plugins/Experimental/ChaosCloth/Source/Chaos/Private/ChaosCloth/ChaosClothingSimulation.cpp b/Engine/Plugins/Experimental/ChaosCloth/Source/Chaos/Private/ChaosCloth/ChaosClothingSimulation.cpp
--- a/Engine/Plugins/Experimental/ChaosCloth/Source/Chaos/Private/ChaosCloth/ChaosClothingSimulation.cpp
+++ b/Engine/Plugins/Experimental/ChaosCloth/Source/Chaos/Private/ChaosCloth/ChaosClothingSimulation.cpp
@@ -355,7 +355,6 @@ void FClothingSimulation::GetSimulationData(
 
 		if (Cloth->GetLODIndex(Solver.Get()) == INDEX_NONE || Cloth->GetOffset(Solver.Get()) == INDEX_NONE)
 		{
-			Data.Reset();
 			continue;
 		}
 
