Chaos Cloth - Add point based wind model to ensure PhyX cloth asset backward compatibility (AKA "Legacy Wind"). [Merging CL 14397275]

diff --git a/Engine/Plugins/Experimental/ChaosCloth/Source/Chaos/Private/ChaosCloth/ChaosClothConfig.cpp b/Engine/Plugins/Experimental/ChaosCloth/Source/Chaos/Private/ChaosCloth/ChaosClothConfig.cpp
--- a/Engine/Plugins/Experimental/ChaosCloth/Source/Chaos/Private/ChaosCloth/ChaosClothConfig.cpp
+++ b/Engine/Plugins/Experimental/ChaosCloth/Source/Chaos/Private/ChaosCloth/ChaosClothConfig.cpp
@@ -61,8 +61,9 @@ void UChaosClothConfig::MigrateFrom(const FClothConfig_Legacy& ClothConfig)
 	LimitScale = FMath::Clamp(ClothConfig.TetherLimit, 0.01f, 10.f);
 	ShapeTargetStiffness = 0.f;
 
-	DragCoefficient = (ClothConfig.WindMethod == EClothingWindMethod_Legacy::Accurate) ? ClothConfig.WindDragCoefficient: 0.07f;  // Only Accurate wind uses the WindDragCoefficient
-	LiftCoefficient = (ClothConfig.WindMethod == EClothingWindMethod_Legacy::Accurate) ? ClothConfig.WindLiftCoefficient: 0.035f;  // Only Accurate wind uses the WindLiftCoefficient
+	bUsePointBasedWindModel = (ClothConfig.WindMethod == EClothingWindMethod_Legacy::Legacy);
+	DragCoefficient = bUsePointBasedWindModel ? 0.07f  : ClothConfig.WindDragCoefficient;  // Only Accurate wind uses the WindDragCoefficient
+	LiftCoefficient = bUsePointBasedWindModel ? 0.035f : ClothConfig.WindLiftCoefficient;  // Only Accurate wind uses the WindLiftCoefficient
 
 	const float Damping = (ClothConfig.Damping.X + ClothConfig.Damping.Y + ClothConfig.Damping.Z) / 3.f;
 	DampingCoefficient = FMath::Clamp(Damping * Damping * 0.7f, 0.f, 1.f);  // Nv Cloth seems to have a different damping formulation.
diff --git a/Engine/Plugins/Experimental/ChaosCloth/Source/Chaos/Private/ChaosCloth/ChaosClothingSimulation.cpp b/Engine/Plugins/Experimental/ChaosCloth/Source/Chaos/Private/ChaosCloth/ChaosClothingSimulation.cpp
--- a/Engine/Plugins/Experimental/ChaosCloth/Source/Chaos/Private/ChaosCloth/ChaosClothingSimulation.cpp
+++ b/Engine/Plugins/Experimental/ChaosCloth/Source/Chaos/Private/ChaosCloth/ChaosClothingSimulation.cpp
@@ -198,6 +198,7 @@ void FClothingSimulation::CreateActor(USkeletalMeshComponent* InOwnerComponent,
 		ClothConfig->AngularVelocityScale,
 		ClothConfig->DragCoefficient,
 		ClothConfig->LiftCoefficient,
+		ClothConfig->bUsePointBasedWindModel,
 		ClothConfig->DampingCoefficient,
 		ClothConfig->CollisionThickness,
 		ClothConfig->FrictionCoefficient,
@@ -276,7 +277,7 @@ void FClothingSimulation::Simulate(IClothingSimulationContext* InContext)
 
 	// Update Solver animatable parameters
 	Solver->SetLocalSpaceLocation(bUseLocalSpaceSimulation ? Context->ComponentToWorld.GetLocation() : TVector<float, 3>(0.f), bNeedsReset);
-	Solver->SetWindVelocity(Context->WindVelocity);
+	Solver->SetWindVelocity(Context->WindVelocity, Context->WindAdaption);
 	Solver->SetGravity(bUseGravityOverride ? GravityOverride : Context->WorldGravity);
 	Solver->EnableClothGravityOverride(!bUseGravityOverride);  // Disable all cloth gravity overrides when the interactor takes over
 
@@ -473,6 +474,7 @@ void FClothingSimulation::RefreshClothConfig(const IClothingSimulationContext* I
 			ClothConfig->AngularVelocityScale,
 			ClothConfig->DragCoefficient,
 			ClothConfig->LiftCoefficient,
+			ClothConfig->bUsePointBasedWindModel,
 			ClothConfig->DampingCoefficient,
 			ClothConfig->CollisionThickness,
 			ClothConfig->FrictionCoefficient,
diff --git a/Engine/Plugins/Experimental/ChaosCloth/Source/Chaos/Private/ChaosCloth/ChaosClothingSimulationCloth.cpp b/Engine/Plugins/Experimental/ChaosCloth/Source/Chaos/Private/ChaosCloth/ChaosClothingSimulationCloth.cpp
--- a/Engine/Plugins/Experimental/ChaosCloth/Source/Chaos/Private/ChaosCloth/ChaosClothingSimulationCloth.cpp
+++ b/Engine/Plugins/Experimental/ChaosCloth/Source/Chaos/Private/ChaosCloth/ChaosClothingSimulationCloth.cpp
@@ -10,6 +10,11 @@
 
 using namespace Chaos;
 
+namespace ChaosClothingSimulationClothConsoleVariables
+{
+	TAutoConsoleVariable<bool> CVarLegacyDisablesAccurateWind(TEXT("p.ChaosCloth.LegacyDisablesAccurateWind"), true, TEXT("Whether using the Legacy wind model switches off the accurate wind model, or adds up to it"));
+}
+
 FClothingSimulationCloth::FLODData::FLODData(int32 InNumParticles, const TConstArrayView<uint32>& InIndices, const TArray<TConstArrayView<float>>& InWeightMaps)
 	: NumParticles(InNumParticles)
 	, Indices(InIndices)
@@ -343,6 +348,7 @@ FClothingSimulationCloth::FClothingSimulationCloth(
 	float InAngularVelocityScale,
 	float InDragCoefficient,
 	float InLiftCoefficient,
+	bool bInUseLegacyWind,
 	float InDampingCoefficient,
 	float InCollisionThickness,
 	float InFrictionCoefficient,
@@ -377,6 +383,7 @@ FClothingSimulationCloth::FClothingSimulationCloth(
 	, AngularVelocityScale(InAngularVelocityScale)
 	, DragCoefficient(InDragCoefficient)
 	, LiftCoefficient(InLiftCoefficient)
+	, bUseLegacyWind(bInUseLegacyWind)
 	, DampingCoefficient(InDampingCoefficient)
 	, CollisionThickness(InCollisionThickness)
 	, FrictionCoefficient(InFrictionCoefficient)
@@ -681,7 +688,16 @@ void FClothingSimulationCloth::Update(FClothingSimulationSolver* Solver)
 		Solver->SetGravity(GroupId, GetGravity(Solver));
 
 		// Update wind
-		Solver->SetWindVelocityField(GroupId, DragCoefficient, LiftCoefficient, &GetTriangleMesh(Solver));
+		Solver->SetLegacyWind(GroupId, bUseLegacyWind);
+
+		if (bUseLegacyWind && ChaosClothingSimulationClothConsoleVariables::CVarLegacyDisablesAccurateWind.GetValueOnAnyThread())
+		{
+			Solver->SetWindVelocityField(GroupId, 0.f, 0.f, &GetTriangleMesh(Solver));
+		}
+		else
+		{
+			Solver->SetWindVelocityField(GroupId, DragCoefficient, LiftCoefficient, &GetTriangleMesh(Solver));
+		}
 
 		// Update general solver properties
 		Solver->SetProperties(GroupId, DampingCoefficient, CollisionThickness, FrictionCoefficient);
diff --git a/Engine/Plugins/Experimental/ChaosCloth/Source/Chaos/Private/ChaosCloth/ChaosClothingSimulationSolver.cpp b/Engine/Plugins/Experimental/ChaosCloth/Source/Chaos/Private/ChaosCloth/ChaosClothingSimulationSolver.cpp
--- a/Engine/Plugins/Experimental/ChaosCloth/Source/Chaos/Private/ChaosCloth/ChaosClothingSimulationSolver.cpp
+++ b/Engine/Plugins/Experimental/ChaosCloth/Source/Chaos/Private/ChaosCloth/ChaosClothingSimulationSolver.cpp
@@ -47,6 +47,7 @@ FClothingSimulationSolver::FClothingSimulationSolver()
 	, CollisionParticlesSize(0)
 	, Gravity(ChaosClothingSimulationSolverDefault::Gravity)
 	, WindVelocity(ChaosClothingSimulationSolverDefault::WindVelocity)
+	, LegacyWindAdaption(0.f)
 	, WindFluidDensity(ChaosClothingSimulationSolverDefault::WindFluidDensity)
 	, bIsClothGravityOverrideEnabled(false)
 {
@@ -562,9 +563,10 @@ void FClothingSimulationSolver::SetGravity(uint32 GroupId, const TVector<float,
 	Evolution->GetGravityForces(GroupId).SetAcceleration(InGravity);
 }
 
-void FClothingSimulationSolver::SetWindVelocity(const TVector<float, 3>& InWindVelocity)
+void FClothingSimulationSolver::SetWindVelocity(const TVector<float, 3>& InWindVelocity, float InLegacyWindAdaption)
 {
 	WindVelocity = InWindVelocity * ChaosClothingSimulationSolverConstant::WorldScale;
+	LegacyWindAdaption = InLegacyWindAdaption;
 }
 
 void FClothingSimulationSolver::SetWindVelocityField(uint32 GroupId, float DragCoefficient, float LiftCoefficient, const TTriangleMesh<float>* TriangleMesh)
@@ -579,6 +581,36 @@ const TVelocityField<float, 3>& FClothingSimulationSolver::GetWindVelocityField(
 	return Evolution->GetVelocityField(GroupId);
 }
 
+void FClothingSimulationSolver::SetLegacyWind(uint32 GroupId, bool bUseLegacyWind)
+{
+	if (!bUseLegacyWind)
+	{
+		// Clear force function
+		// NOTE: This assumes that the force function is only used for the legacy wind effect
+		Evolution->GetForceFunction(GroupId) = TFunction<void(TPBDParticles<float, 3>&, const float, const int32)>();
+	}
+	else
+	{
+		// Add legacy wind function
+		Evolution->GetForceFunction(GroupId) = 
+			[this](TPBDParticles<float, 3>& Particles, const float /*Dt*/, const int32 Index)
+			{
+				if (Particles.InvM(Index) != 0.f)
+				{
+					// Calculate wind velocity delta
+					static const float LegacyWindMultiplier = 25.f;
+					const TVector<float, 3> VelocityDelta = WindVelocity * LegacyWindMultiplier - Particles.V(Index);
+
+					// Scale by angle
+					const float DirectionDot = TVector<float, 3>::DotProduct(VelocityDelta.GetUnsafeNormal(), Normals[Index]);
+					const float ScaleFactor = FMath::Min(1.f, FMath::Abs(DirectionDot) * LegacyWindAdaption);
+					
+					Particles.F(Index) += VelocityDelta * ScaleFactor * Particles.M(Index);
+				}
+			};
+	}
+}
+
 void FClothingSimulationSolver::SetSelfCollisions(uint32 GroupId, float SelfCollisionThickness, const TTriangleMesh<float>* TriangleMesh)
 {
 #ifdef CHAOS_CLOTH_MUST_IMPROVE_SELF_COLLISION  // TODO: Improve self-collision until it can run in engine tests without crashing the simulation.
diff --git a/Engine/Plugins/Experimental/ChaosCloth/Source/Chaos/Public/ChaosCloth/ChaosClothConfig.h b/Engine/Plugins/Experimental/ChaosCloth/Source/Chaos/Public/ChaosCloth/ChaosClothConfig.h
--- a/Engine/Plugins/Experimental/ChaosCloth/Source/Chaos/Public/ChaosCloth/ChaosClothConfig.h
+++ b/Engine/Plugins/Experimental/ChaosCloth/Source/Chaos/Public/ChaosCloth/ChaosClothConfig.h
@@ -117,7 +117,7 @@ public:
 	UPROPERTY(EditAnywhere, Category = "Collision Properties", meta = (InlineEditConditionToggle))
 	bool bUseSelfCollisions = false;
 
-	// This parameter is automatically set by the migration code. It can be overriden here to use the old way of authoring the backstop distances.
+	// This parameter is automatically set by the migration code. It can be overridden here to use the old way of authoring the backstop distances.
 	// The legacy backstop requires the sphere radius to be included within the painted distance mask, making it difficult to author correctly. In this case the backstop distance is the distance from the animated mesh to the center of the corresponding backstop collision sphere.
 	// The non legacy backstop automatically adds the matching sphere's radius to the distance calculations at runtime to make for a simpler authoring of the backstop distances. In this case the backstop distance is the distance from the animated mesh to the surface of the backstop collision sphere.
 	// In both cases, a positive backstop distance goes against the corresponding animated mesh's normal, and a negative backstop distance goes along the corresponding animated mesh's normal.
@@ -132,12 +132,18 @@ public:
 	UPROPERTY(EditAnywhere, Category = "Environmental Properties", meta = (UIMin = "0", UIMax = "1", ClampMin = "0", ClampMax = "1"))
 	float DampingCoefficient = 0.01f;
 
+	// This parameter is automatically set by the migration code. It can be overridden here to use the old deprecated "Legacy" wind model in order to preserve behavior with previous versions of the engine.
+	// The old wind model is not an accurate aerodynamic model and as such should be avoided. Being point based, it doesn't take into account the surface area that gets hit by the wind.
+	// Using this model makes the simulation slightly slower, disables the aerodynamically accurate wind model, and prevents the cloth to interact with the surrounding environment (air, water, ...etc.) even when there is no wind.
+	UPROPERTY(EditAnywhere, Category = "Environmental Properties")
+	bool bUsePointBasedWindModel = false;
+
 	// The aerodynamic coefficient of drag applying on each particle.
-	UPROPERTY(EditAnywhere, Category = "Environmental Properties", meta = (UIMin = "0", UIMax = "1", ClampMin = "0", ClampMax = "10"))
+	UPROPERTY(EditAnywhere, Category = "Environmental Properties", meta = (UIMin = "0", UIMax = "1", ClampMin = "0", ClampMax = "10", EditCondition = "!bUsePointBasedWindModel"))
 	float DragCoefficient = 0.07f;
 
 	// The aerodynamic coefficient of lift applying on each particle.
-	UPROPERTY(EditAnywhere, Category = "Environmental Properties", meta = (UIMin = "0", UIMax = "1", ClampMin = "0", ClampMax = "10"))
+	UPROPERTY(EditAnywhere, Category = "Environmental Properties", meta = (UIMin = "0", UIMax = "1", ClampMin = "0", ClampMax = "10", EditCondition = "!bUsePointBasedWindModel"))
 	float LiftCoefficient = 0.035f;
 
 	// Use the config gravity value instead of world gravity.
diff --git a/Engine/Plugins/Experimental/ChaosCloth/Source/Chaos/Public/ChaosCloth/ChaosClothingSimulationCloth.h b/Engine/Plugins/Experimental/ChaosCloth/Source/Chaos/Public/ChaosCloth/ChaosClothingSimulationCloth.h
--- a/Engine/Plugins/Experimental/ChaosCloth/Source/Chaos/Public/ChaosCloth/ChaosClothingSimulationCloth.h
+++ b/Engine/Plugins/Experimental/ChaosCloth/Source/Chaos/Public/ChaosCloth/ChaosClothingSimulationCloth.h
@@ -50,6 +50,7 @@ namespace Chaos
 			float InAngularVelocityScale,
 			float InDragCoefficient,
 			float InLiftCoefficient,
+			bool bInUseLegacyWind,
 			float InDampingCoefficient,
 			float InCollisionThickness,
 			float InFrictionCoefficient,
@@ -181,6 +182,7 @@ namespace Chaos
 		float AngularVelocityScale;  // Angular ratio factor applied to the reference bone transforms
 		float DragCoefficient;
 		float LiftCoefficient;
+		bool bUseLegacyWind;
 		float DampingCoefficient;
 		float CollisionThickness;
 		float FrictionCoefficient;
diff --git a/Engine/Plugins/Experimental/ChaosCloth/Source/Chaos/Public/ChaosCloth/ChaosClothingSimulationSolver.h b/Engine/Plugins/Experimental/ChaosCloth/Source/Chaos/Public/ChaosCloth/ChaosClothingSimulationSolver.h
--- a/Engine/Plugins/Experimental/ChaosCloth/Source/Chaos/Public/ChaosCloth/ChaosClothingSimulationSolver.h
+++ b/Engine/Plugins/Experimental/ChaosCloth/Source/Chaos/Public/ChaosCloth/ChaosClothingSimulationSolver.h
@@ -36,7 +36,7 @@ namespace Chaos
 		void SetGravity(const TVector<float, 3>& InGravity) { Gravity = InGravity; }
 		const TVector<float, 3>& GetGravity() const { return Gravity; }
 
-		void SetWindVelocity(const TVector<float, 3>& InWindVelocity);
+		void SetWindVelocity(const TVector<float, 3>& InWindVelocity, float InLegacyWindAdaption = 0.f);
 		void SetWindFluidDensity(float InWindFluidDensity) { WindFluidDensity = InWindFluidDensity; }
 
 		void SetNumIterations(int32 InNumIterations) { NumIterations = InNumIterations; }
@@ -87,6 +87,9 @@ namespace Chaos
 		void SetWindVelocityField(uint32 GroupId, float DragCoefficient, float LiftCoefficient, const TTriangleMesh<float>* TriangleMesh = nullptr);
 		const TVelocityField<float, 3>&  GetWindVelocityField(uint32 GroupId);
 
+		// Set legacy noise wind.
+		void SetLegacyWind(uint32 GroupId, bool bUseLegacyWind);
+
 		// Enable self collision, or disable if TriangleMesh is null.
 		void SetSelfCollisions(uint32 GroupId, float SelfCollisionThickness, const TTriangleMesh<float>* TriangleMesh = nullptr);
 
@@ -178,6 +181,7 @@ namespace Chaos
 		// Solver parameters
 		TVector<float, 3> Gravity;
 		TVector<float, 3> WindVelocity;
+		float LegacyWindAdaption;
 		float WindFluidDensity;
 		bool bIsClothGravityOverrideEnabled;
 	};
diff --git a/Engine/Source/Runtime/ClothingSystemRuntimeCommon/Private/ClothingSimulation.cpp b/Engine/Source/Runtime/ClothingSystemRuntimeCommon/Private/ClothingSimulation.cpp
--- a/Engine/Source/Runtime/ClothingSystemRuntimeCommon/Private/ClothingSimulation.cpp
+++ b/Engine/Source/Runtime/ClothingSystemRuntimeCommon/Private/ClothingSimulation.cpp
@@ -119,7 +119,7 @@ void FClothingSimulationContextCommon::FillWorldGravity(const USkeletalMeshCompo
 
 void FClothingSimulationContextCommon::FillWindVelocity(const USkeletalMeshComponent* InComponent)
 {
-	SetWindFromComponent(InComponent);
+	InComponent->GetWindForCloth_GameThread(WindVelocity, WindAdaption);
 }
 
 void FClothingSimulationContextCommon::FillDeltaSeconds(float InDeltaSeconds, float InMaxPhysicsDelta)
@@ -134,16 +134,9 @@ void FClothingSimulationContextCommon::FillTeleportMode(const USkeletalMeshCompo
 		InComponent->ClothTeleportMode;
 }
 
-float FClothingSimulationContextCommon::SetWindFromComponent(const USkeletalMeshComponent* Component)
+void FClothingSimulationContextCommon::FillMaxDistanceScale(const USkeletalMeshComponent* InComponent)
 {
-	float WindAdaption;
-	Component->GetWindForCloth_GameThread(WindVelocity, WindAdaption);
-	return WindAdaption;
-}
-
-void FClothingSimulationContextCommon::FillMaxDistanceScale(const USkeletalMeshComponent* Component)
-{
-	MaxDistanceScale = Component->GetClothMaxDistanceScale();
+	MaxDistanceScale = InComponent->GetClothMaxDistanceScale();
 }
 
 //==============================================================================
diff --git a/Engine/Source/Runtime/ClothingSystemRuntimeCommon/Public/ClothingSimulation.h b/Engine/Source/Runtime/ClothingSystemRuntimeCommon/Public/ClothingSimulation.h
--- a/Engine/Source/Runtime/ClothingSystemRuntimeCommon/Public/ClothingSimulation.h
+++ b/Engine/Source/Runtime/ClothingSystemRuntimeCommon/Public/ClothingSimulation.h
@@ -27,10 +27,7 @@ protected:
 	virtual void FillWindVelocity(const USkeletalMeshComponent* InComponent);
 	virtual void FillDeltaSeconds(float InDeltaSeconds, float InMaxPhysicsDelta);
 	virtual void FillTeleportMode(const USkeletalMeshComponent* InComponent, float InDeltaSeconds, float InMaxPhysicsDelta);
-	virtual void FillMaxDistanceScale(const USkeletalMeshComponent* Component);
-
-	// Set the wind velocity and return the wind adaptation if required
-	float SetWindFromComponent(const USkeletalMeshComponent* Component);
+	virtual void FillMaxDistanceScale(const USkeletalMeshComponent* InComponent);
 
 public:
 	// Component space bone transforms of the owning component
@@ -48,6 +45,10 @@ public:
 	// Wind velocity at the component location
 	FVector WindVelocity;
 
+	// Wind adaption, a measure of how quickly to adapt to the wind speed
+	// when using the legacy wind calculation mode
+	float WindAdaption;
+
 	// Delta for this tick
 	float DeltaSeconds;
 
diff --git a/Engine/Source/Runtime/ClothingSystemRuntimeNv/Private/ClothingSimulationNv.cpp b/Engine/Source/Runtime/ClothingSystemRuntimeNv/Private/ClothingSimulationNv.cpp
--- a/Engine/Source/Runtime/ClothingSystemRuntimeNv/Private/ClothingSimulationNv.cpp
+++ b/Engine/Source/Runtime/ClothingSystemRuntimeNv/Private/ClothingSimulationNv.cpp
@@ -65,11 +65,6 @@ void FClothingSimulationContextNv::FillWorldGravity(const USkeletalMeshComponent
 	FClothingSimulationContextCommon::WorldGravity = FVector(0.f, 0.f, GravityStrength);
 }
 
-void FClothingSimulationContextNv::FillWindVelocity(const USkeletalMeshComponent* InComponent)
-{
-	WindAdaption = FClothingSimulationContextCommon::SetWindFromComponent(InComponent);  // Also update wind adaptation
-}
-
 //=============================================================================
 // FClothingSimulationNv
 //=============================================================================
diff --git a/Engine/Source/Runtime/ClothingSystemRuntimeNv/Public/ClothingSimulationNv.h b/Engine/Source/Runtime/ClothingSystemRuntimeNv/Public/ClothingSimulationNv.h
--- a/Engine/Source/Runtime/ClothingSystemRuntimeNv/Public/ClothingSimulationNv.h
+++ b/Engine/Source/Runtime/ClothingSystemRuntimeNv/Public/ClothingSimulationNv.h
@@ -40,14 +40,7 @@ public:
 	virtual void Fill(const USkeletalMeshComponent* InComponent, float InDeltaSeconds, float InMaxPhysicsDelta) override;
 
 	// Set the world gravity in the parent class while preserving the Nv legacy code behavior
-	void virtual FillWorldGravity(const USkeletalMeshComponent* InComponent) override;
-
-	// Set WindVelocity in the parent class and Nv specific WindAdaption
-	virtual void FillWindVelocity(const USkeletalMeshComponent* InComponent) override;
-
-	// Wind adaption, a measure of how quickly to adapt to the wind speed
-	// when using the legacy wind calculation mode
-	float WindAdaption;
+	virtual void FillWorldGravity(const USkeletalMeshComponent* InComponent) override;
 };
 
 // Scratch data for simulation to avoid allocations while processing, per actor data
diff --git a/Engine/Source/Runtime/Experimental/Chaos/Private/Chaos/PBDEvolution.cpp b/Engine/Source/Runtime/Experimental/Chaos/Private/Chaos/PBDEvolution.cpp
--- a/Engine/Source/Runtime/Experimental/Chaos/Private/Chaos/PBDEvolution.cpp
+++ b/Engine/Source/Runtime/Experimental/Chaos/Private/Chaos/PBDEvolution.cpp
@@ -72,6 +72,7 @@ TPBDEvolution<T, d>::TPBDEvolution(TPBDParticles<T, d>&& InParticles, TKinematic
 {
 	// Add group arrays
 	TArrayCollection::AddArray(&MGroupGravityForces);
+	TArrayCollection::AddArray(&MGroupForceRules);
 	TArrayCollection::AddArray(&MGroupCollisionThicknesses);
 	TArrayCollection::AddArray(&MGroupSelfCollisionThicknesses);
 	TArrayCollection::AddArray(&MGroupCoefficientOfFrictions);
@@ -228,11 +229,12 @@ void TPBDEvolution<T, d>::AdvanceOneTimeStep(const T Dt)
 
 				InitForceRule.Apply(Particles, Dt, Index); // F = TV(0)
 				MGroupGravityForces[ParticleGroupId].Apply(Particles, Dt, Index); // F += M * G
-				for (TFunction<void(TPBDParticles<T, d>&, const T, const int32)>& ForceRule : MForceRules)
+
+				if (MGroupForceRules[ParticleGroupId])
 				{
-					ForceRule(Particles, Dt, Index); // F += M * A
+					MGroupForceRules[ParticleGroupId](Particles, Dt, Index); // F += M * A
 				}
-			
+
 				MGroupVelocityFields[ParticleGroupId].Apply(Particles, Dt, Index);
 
 				if (MKinematicUpdate)
diff --git a/Engine/Source/Runtime/Experimental/Chaos/Public/Chaos/PBDEvolution.h b/Engine/Source/Runtime/Experimental/Chaos/Public/Chaos/PBDEvolution.h
--- a/Engine/Source/Runtime/Experimental/Chaos/Public/Chaos/PBDEvolution.h
+++ b/Engine/Source/Runtime/Experimental/Chaos/Public/Chaos/PBDEvolution.h
@@ -85,8 +85,8 @@ class CHAOS_API TPBDEvolution : public TArrayCollection
 	void SetKinematicUpdateFunction(TFunction<void(TPBDParticles<T, d>&, const T, const T, const int32)> KinematicUpdate) { MKinematicUpdate = KinematicUpdate; }
 	void SetCollisionKinematicUpdateFunction(TFunction<void(TKinematicGeometryClothParticles<T, d>&, const T, const T, const int32)> KinematicUpdate) { MCollisionKinematicUpdate = KinematicUpdate; }
 
-	void ResetForceFunctions() { MForceRules.Reset(); }
-	void AddForceFunction(TFunction<void(TPBDParticles<T, d>&, const T, const int32)> ForceFunction) { MForceRules.Add(ForceFunction); }
+	TFunction<void(TPBDParticles<T, d>&, const T, const int32)>& GetForceFunction(const uint32 GroupId = 0) { return MGroupForceRules[GroupId]; }
+	const TFunction<void(TPBDParticles<T, d>&, const T, const int32)>& GetForceFunction(const uint32 GroupId = 0) const { return MGroupForceRules[GroupId]; }
 
 	FGravityForces& GetGravityForces(const uint32 GroupId = 0) { check(GroupId < TArrayCollection::Size()); return MGroupGravityForces[GroupId]; }
 	const FGravityForces& GetGravityForces(const uint32 GroupId = 0) const { check(GroupId < TArrayCollection::Size()); return MGroupGravityForces[GroupId]; }
@@ -137,6 +137,7 @@ private:
 	TArrayCollectionArray<uint32> MParticleGroupIds;  // Used for per group parameters for particles
 
 	TArrayCollectionArray<FGravityForces> MGroupGravityForces;
+	TArrayCollectionArray<TFunction<void(TPBDParticles<T, d>&, const T, const int32)>> MGroupForceRules;
 	TArrayCollectionArray<T> MGroupCollisionThicknesses;
 	TArrayCollectionArray<T> MGroupSelfCollisionThicknesses;
 	TArrayCollectionArray<T> MGroupCoefficientOfFrictions;
@@ -148,7 +149,6 @@ private:
 	TArray<TFunction<void(TPBDParticles<T, d>&, const T)>> MConstraintRules;
 	TPBDActiveView<TArray<TFunction<void(TPBDParticles<T, d>&, const T)>>> MConstraintRulesActiveView;
 
-	TArray<TFunction<void(TPBDParticles<T, d>&, const T, const int32)>> MForceRules;
 	TFunction<void(TPBDActiveView<TPBDParticles<T, d>>&, const T)> MParticleUpdate;
 	TFunction<void(TPBDParticles<T, d>&, const T, const T, const int32)> MKinematicUpdate;
 	TFunction<void(TKinematicGeometryClothParticles<T, d>&, const T, const T, const int32)> MCollisionKinematicUpdate;
