Chaos Cloth - Fix missing Animdrive stiffness and MaxDistancesMultiplier updates. [Merging CL 14439723]

diff --git a/Engine/Plugins/Experimental/ChaosCloth/Source/Chaos/Private/ChaosCloth/ChaosClothingSimulationCloth.cpp b/Engine/Plugins/Experimental/ChaosCloth/Source/Chaos/Private/ChaosCloth/ChaosClothingSimulationCloth.cpp
--- a/Engine/Plugins/Experimental/ChaosCloth/Source/Chaos/Private/ChaosCloth/ChaosClothingSimulationCloth.cpp
+++ b/Engine/Plugins/Experimental/ChaosCloth/Source/Chaos/Private/ChaosCloth/ChaosClothingSimulationCloth.cpp
@@ -274,6 +274,18 @@ void FClothingSimulationCloth::FLODData::Remove(FClothingSimulationSolver* Solve
 	SolverData.Remove(Solver);
 }
 
+void FClothingSimulationCloth::FLODData::Update(FClothingSimulationSolver* Solver, FClothingSimulationCloth* Cloth)
+{
+	check(Solver);
+	const int32 Offset = SolverData.FindChecked(Solver).Offset;
+	check(Offset != INDEX_NONE);
+
+	// Update the animatable constraint parameters
+	FClothConstraints& ClothConstraints = Solver->GetClothConstraints(Offset);
+	ClothConstraints.SetMaxDistancesMultiplier(Cloth->MaxDistancesMultiplier);
+	ClothConstraints.SetAnimDriveSpringStiffness(Cloth->AnimDriveSpringStiffness);
+}
+
 void FClothingSimulationCloth::FLODData::Enable(FClothingSimulationSolver* Solver, bool bEnable) const
 {
 	check(Solver);
@@ -680,6 +692,7 @@ void FClothingSimulationCloth::Update(FClothingSimulationSolver* Solver)
 	if (LODIndex != INDEX_NONE)
 	{
 		// Update Cloth group parameters  TODO: Cloth groups should exist as their own node object so that they can be used by several cloth objects
+		LODData[LODIndex].Update(Solver, this);
 
 		// Update gravity
 		// This code relies on the solver gravity property being already set.
diff --git a/Engine/Plugins/Experimental/ChaosCloth/Source/Chaos/Public/ChaosCloth/ChaosClothingSimulationCloth.h b/Engine/Plugins/Experimental/ChaosCloth/Source/Chaos/Public/ChaosCloth/ChaosClothingSimulationCloth.h
--- a/Engine/Plugins/Experimental/ChaosCloth/Source/Chaos/Public/ChaosCloth/ChaosClothingSimulationCloth.h
+++ b/Engine/Plugins/Experimental/ChaosCloth/Source/Chaos/Public/ChaosCloth/ChaosClothingSimulationCloth.h
@@ -148,6 +148,8 @@ namespace Chaos
 			void Add(FClothingSimulationSolver* Solver, FClothingSimulationCloth* Cloth, int32 LODIndex);
 			void Remove(FClothingSimulationSolver* Solver);
 
+			void Update(FClothingSimulationSolver* Solver, FClothingSimulationCloth* Cloth);
+
 			void Enable(FClothingSimulationSolver* Solver, bool bEnable) const;
 
 			void ResetStartPose(FClothingSimulationSolver* Solver) const;
