Chaos Cloth - Fix the AnimDriveSpringStiffness parameter refresh when changing its value in the Clothing UI. [Merging CL 14458997]

diff --git a/Engine/Plugins/Experimental/ChaosCloth/Source/Chaos/Private/ChaosCloth/ChaosClothingSimulation.cpp b/Engine/Plugins/Experimental/ChaosCloth/Source/Chaos/Private/ChaosCloth/ChaosClothingSimulation.cpp
--- a/Engine/Plugins/Experimental/ChaosCloth/Source/Chaos/Private/ChaosCloth/ChaosClothingSimulation.cpp
+++ b/Engine/Plugins/Experimental/ChaosCloth/Source/Chaos/Private/ChaosCloth/ChaosClothingSimulation.cpp
@@ -447,6 +447,7 @@ void FClothingSimulation::RefreshClothConfig(const IClothingSimulationContext* I
 		const uint32 GroupId = Cloth->GetGroupId();
 		const UChaosClothConfig* const ClothConfig = Mesh->GetAsset()->GetClothConfig<UChaosClothConfig>();
 
+		AnimDriveSpringStiffness = ClothConfig->AnimDriveSpringStiffness;
 		Cloth = MakeUnique<FClothingSimulationCloth>(
 			Mesh,
 			MoveTemp(ClothColliders),
