Chaos Cloth - Add a warning message in the Mesh Editor viewport for when the simulation has lost all impulses because of running slow and being teleported. [Merging CL 14479886]

diff --git a/Engine/Plugins/Experimental/ChaosCloth/Source/Chaos/Private/ChaosCloth/ChaosClothingSimulation.cpp b/Engine/Plugins/Experimental/ChaosCloth/Source/Chaos/Private/ChaosCloth/ChaosClothingSimulation.cpp
--- a/Engine/Plugins/Experimental/ChaosCloth/Source/Chaos/Private/ChaosCloth/ChaosClothingSimulation.cpp
+++ b/Engine/Plugins/Experimental/ChaosCloth/Source/Chaos/Private/ChaosCloth/ChaosClothingSimulation.cpp
@@ -274,6 +274,7 @@ void FClothingSimulation::Simulate(IClothingSimulationContext* InContext)
 
 	const bool bNeedsReset = (Context->TeleportMode == EClothingTeleportMode::TeleportAndReset);
 	const bool bNeedsTeleport = (Context->TeleportMode > EClothingTeleportMode::None);
+	bIsTeleported = bNeedsTeleport;
 
 	// Update Solver animatable parameters
 	Solver->SetLocalSpaceLocation(bUseLocalSpaceSimulation ? Context->ComponentToWorld.GetLocation() : TVector<float, 3>(0.f), bNeedsReset);
diff --git a/Engine/Plugins/Experimental/ChaosCloth/Source/Chaos/Public/ChaosCloth/ChaosClothingSimulation.h b/Engine/Plugins/Experimental/ChaosCloth/Source/Chaos/Public/ChaosCloth/ChaosClothingSimulation.h
--- a/Engine/Plugins/Experimental/ChaosCloth/Source/Chaos/Public/ChaosCloth/ChaosClothingSimulation.h
+++ b/Engine/Plugins/Experimental/ChaosCloth/Source/Chaos/Public/ChaosCloth/ChaosClothingSimulation.h
@@ -73,12 +73,13 @@ namespace Chaos
 		void RefreshPhysicsAsset();
 
 		// IClothingSimulation interface
-		virtual int32 GetNumCloths() const { return NumCloths; }
-		virtual int32 GetNumKinematicParticles() const { return NumKinemamicParticles; }
-		virtual int32 GetNumDynamicParticles() const { return NumDynamicParticles; }
-		virtual int32 GetNumIterations() const { return NumIterations; }
-		virtual int32 GetNumSubsteps() const { return NumSubsteps; }
-		virtual float GetSimulationTime() const { return SimulationTime; }
+		virtual int32 GetNumCloths() const override { return NumCloths; }
+		virtual int32 GetNumKinematicParticles() const override { return NumKinemamicParticles; }
+		virtual int32 GetNumDynamicParticles() const override { return NumDynamicParticles; }
+		virtual int32 GetNumIterations() const override { return NumIterations; }
+		virtual int32 GetNumSubsteps() const override { return NumSubsteps; }
+		virtual float GetSimulationTime() const override { return SimulationTime; }
+		virtual bool IsTeleported() const override { return bIsTeleported; }
 		// End of IClothingSimulation interface
 
 #if WITH_EDITOR
@@ -140,6 +141,7 @@ namespace Chaos
 		TAtomic<int32> NumIterations;
 		TAtomic<int32> NumSubsteps;
 		TAtomic<float> SimulationTime;
+		TAtomic<bool> bIsTeleported;
 
 		// Overrides
 		bool bUseLocalSpaceSimulation;
diff --git a/Engine/Source/Editor/Persona/Private/AnimationEditorViewportClient.cpp b/Engine/Source/Editor/Persona/Private/AnimationEditorViewportClient.cpp
--- a/Engine/Source/Editor/Persona/Private/AnimationEditorViewportClient.cpp
+++ b/Engine/Source/Editor/Persona/Private/AnimationEditorViewportClient.cpp
@@ -981,6 +981,10 @@ FText FAnimationViewportClient::GetDisplayInfo(bool bDisplayAllInfo) const
 			NumberFormatOptions.MaximumFractionalDigits = 2;
 			TextValue = ConcatenateLine(TextValue, FText::Format(LOCTEXT("SimulationTime", "Simulation Time: {0}ms"), FText::AsNumber(SimulationTime, &NumberFormatOptions)));
 		}
+		if (ClothingSimulation->IsTeleported())
+		{
+			TextValue = ConcatenateLine(TextValue, LOCTEXT("IsTeleported", "*** Warning ***: Max Delta Time Teleport!"));
+		}
 	}
 
 	if (PreviewMeshComponent->GetSectionPreview() != INDEX_NONE)
diff --git a/Engine/Source/Runtime/ClothingSystemRuntimeInterface/Public/ClothingSimulationInterface.h b/Engine/Source/Runtime/ClothingSystemRuntimeInterface/Public/ClothingSimulationInterface.h
--- a/Engine/Source/Runtime/ClothingSystemRuntimeInterface/Public/ClothingSimulationInterface.h
+++ b/Engine/Source/Runtime/ClothingSystemRuntimeInterface/Public/ClothingSimulationInterface.h
@@ -162,4 +162,6 @@ public:
 	virtual int32 GetNumSubsteps() const { return 0; }
 	/** Return the simulation time in ms. Implementation must be thread safe. */
 	virtual float GetSimulationTime() const { return 0.f; }
+	/** Return whether the simulation is teleported. Implementation must be thread safe. */
+	virtual bool IsTeleported() const { return false; }
 };
