LWC: Compile fixes.

diff --git a/Engine/Plugins/Experimental/ChaosCloth/Source/Chaos/Private/ChaosCloth/ChaosClothingSimulationInteractor.cpp b/Engine/Plugins/Experimental/ChaosCloth/Source/Chaos/Private/ChaosCloth/ChaosClothingSimulationInteractor.cpp
--- a/Engine/Plugins/Experimental/ChaosCloth/Source/Chaos/Private/ChaosCloth/ChaosClothingSimulationInteractor.cpp
+++ b/Engine/Plugins/Experimental/ChaosCloth/Source/Chaos/Private/ChaosCloth/ChaosClothingSimulationInteractor.cpp
@@ -45,7 +45,7 @@ void UChaosClothingInteractor::SetLongRangeAttachment(FVector2D TetherStiffness)
 {
 	Commands.Add(FChaosClothingInteractorCommand::CreateLambda([TetherStiffness](FClothingSimulationCloth* Cloth)
 	{
-		Cloth->SetLongRangeAttachmentProperties(TetherStiffness);
+		Cloth->SetLongRangeAttachmentProperties((FVec2)TetherStiffness);
 	}));
 }
 
diff --git a/Engine/Plugins/FX/Niagara/Source/Niagara/Private/NiagaraDataInterfaceCubeTexture.cpp b/Engine/Plugins/FX/Niagara/Source/Niagara/Private/NiagaraDataInterfaceCubeTexture.cpp
--- a/Engine/Plugins/FX/Niagara/Source/Niagara/Private/NiagaraDataInterfaceCubeTexture.cpp
+++ b/Engine/Plugins/FX/Niagara/Source/Niagara/Private/NiagaraDataInterfaceCubeTexture.cpp
@@ -258,7 +258,7 @@ public:
 				GBlackTextureCube->SamplerStateRHI,
 				GBlackTextureCube->TextureRHI
 			);
-			SetShaderValue(RHICmdList, ComputeShaderRHI, Dimensions, FVector::ZeroVector);
+			SetShaderValue(RHICmdList, ComputeShaderRHI, Dimensions, FVector3f::ZeroVector);
 		}
 	}
 private:
