Better fix compile issue with LWC on and cloth code

diff --git a/Engine/Plugins/Experimental/ChaosCloth/Source/Chaos/Private/ChaosCloth/ChaosClothingSimulationInteractor.cpp b/Engine/Plugins/Experimental/ChaosCloth/Source/Chaos/Private/ChaosCloth/ChaosClothingSimulationInteractor.cpp
--- a/Engine/Plugins/Experimental/ChaosCloth/Source/Chaos/Private/ChaosCloth/ChaosClothingSimulationInteractor.cpp
+++ b/Engine/Plugins/Experimental/ChaosCloth/Source/Chaos/Private/ChaosCloth/ChaosClothingSimulationInteractor.cpp
@@ -45,7 +45,7 @@ void UChaosClothingInteractor::SetLongRangeAttachment(FVector2D TetherStiffness)
 {
 	Commands.Add(FChaosClothingInteractorCommand::CreateLambda([TetherStiffness](FClothingSimulationCloth* Cloth)
 	{
-		Cloth->SetLongRangeAttachmentProperties((FVec2)TetherStiffness);
+		Cloth->SetLongRangeAttachmentProperties(Chaos::TVec2<Chaos::FRealSingle>(TetherStiffness.X, TetherStiffness.Y));
 	}));
 }
 
