namespace usage consistency for recent fix

diff --git a/Engine/Plugins/Experimental/ChaosCloth/Source/Chaos/Private/ChaosCloth/ChaosClothingSimulationInteractor.cpp b/Engine/Plugins/Experimental/ChaosCloth/Source/Chaos/Private/ChaosCloth/ChaosClothingSimulationInteractor.cpp
--- a/Engine/Plugins/Experimental/ChaosCloth/Source/Chaos/Private/ChaosCloth/ChaosClothingSimulationInteractor.cpp
+++ b/Engine/Plugins/Experimental/ChaosCloth/Source/Chaos/Private/ChaosCloth/ChaosClothingSimulationInteractor.cpp
@@ -45,7 +45,7 @@ void UChaosClothingInteractor::SetLongRangeAttachment(FVector2D TetherStiffness)
 {
 	Commands.Add(FChaosClothingInteractorCommand::CreateLambda([TetherStiffness](FClothingSimulationCloth* Cloth)
 	{
-		Cloth->SetLongRangeAttachmentProperties(Chaos::TVec2<Chaos::FRealSingle>(TetherStiffness.X, TetherStiffness.Y));
+		Cloth->SetLongRangeAttachmentProperties(TVec2<FRealSingle>(TetherStiffness.X, TetherStiffness.Y));
 	}));
 }
 
