Fix error in Chaos Cloth when compiling using LWC

diff --git a/Engine/Plugins/Experimental/ChaosCloth/Source/Chaos/Private/ChaosCloth/ChaosClothingSimulationInteractor.cpp b/Engine/Plugins/Experimental/ChaosCloth/Source/Chaos/Private/ChaosCloth/ChaosClothingSimulationInteractor.cpp
--- a/Engine/Plugins/Experimental/ChaosCloth/Source/Chaos/Private/ChaosCloth/ChaosClothingSimulationInteractor.cpp
+++ b/Engine/Plugins/Experimental/ChaosCloth/Source/Chaos/Private/ChaosCloth/ChaosClothingSimulationInteractor.cpp
@@ -44,7 +44,7 @@ void UChaosClothingInteractor::SetMaterial(FVector2D EdgeStiffness, FVector2D Be
 {
 	Commands.Add(FChaosClothingInteractorCommand::CreateLambda([EdgeStiffness, BendingStiffness, AreaStiffness](FClothingSimulationCloth* Cloth)
 	{
-		Cloth->SetMaterialProperties(EdgeStiffness, BendingStiffness, AreaStiffness);
+		Cloth->SetMaterialProperties(TVec2<FRealSingle>(EdgeStiffness[0], EdgeStiffness[1]), TVec2<FRealSingle>(BendingStiffness[0], BendingStiffness[1]), TVec2<FRealSingle>(AreaStiffness[0], AreaStiffness[1]));
 	}));
 }
 
