[Cloth] - add trace scopes

diff --git a/Engine/Plugins/Experimental/ChaosCloth/Source/Chaos/Private/ChaosCloth/ChaosClothingSimulation.cpp b/Engine/Plugins/Experimental/ChaosCloth/Source/Chaos/Private/ChaosCloth/ChaosClothingSimulation.cpp
--- a/Engine/Plugins/Experimental/ChaosCloth/Source/Chaos/Private/ChaosCloth/ChaosClothingSimulation.cpp
+++ b/Engine/Plugins/Experimental/ChaosCloth/Source/Chaos/Private/ChaosCloth/ChaosClothingSimulation.cpp
@@ -319,6 +319,7 @@ IClothingSimulationContext* FClothingSimulation::CreateContext()
 
 void FClothingSimulation::CreateActor(USkeletalMeshComponent* InOwnerComponent, UClothingAssetBase* InAsset, int32 InSimDataIndex)
 {
+	TRACE_CPUPROFILER_EVENT_SCOPE(FClothingSimulation_CreateActor);
 	SCOPE_CYCLE_COUNTER(STAT_ChaosClothCreateActor);
 
 	check(InOwnerComponent);
@@ -492,6 +493,7 @@ bool FClothingSimulation::ShouldSimulate() const
 
 void FClothingSimulation::Simulate(IClothingSimulationContext* InContext)
 {
+	TRACE_CPUPROFILER_EVENT_SCOPE(FClothingSimulation_Simulate);
 #if !(UE_BUILD_SHIPPING || UE_BUILD_TEST)
 	if (ChaosClothingSimulationConsole::Command && ChaosClothingSimulationConsole::Command->MustStep(StepCount))
 #endif
@@ -578,6 +580,7 @@ void FClothingSimulation::GetSimulationData(
 	USkeletalMeshComponent* InOwnerComponent,
 	USkinnedMeshComponent* InOverrideComponent) const
 {
+	TRACE_CPUPROFILER_EVENT_SCOPE(FClothingSimulation_GetSimulationData);
 	SCOPE_CYCLE_COUNTER(STAT_ChaosClothGetSimulationData);
 
 	if (!Cloths.Num() || !InOwnerComponent)
diff --git a/Engine/Plugins/Experimental/ChaosCloth/Source/Chaos/Private/ChaosCloth/ChaosClothingSimulationCollider.cpp b/Engine/Plugins/Experimental/ChaosCloth/Source/Chaos/Private/ChaosCloth/ChaosClothingSimulationCollider.cpp
--- a/Engine/Plugins/Experimental/ChaosCloth/Source/Chaos/Private/ChaosCloth/ChaosClothingSimulationCollider.cpp
+++ b/Engine/Plugins/Experimental/ChaosCloth/Source/Chaos/Private/ChaosCloth/ChaosClothingSimulationCollider.cpp
@@ -658,6 +658,7 @@ void FClothingSimulationCollider::PreUpdate(FClothingSimulationSolver* Solver, F
 	check(Solver);
 	check(Cloth);
 
+	TRACE_CPUPROFILER_EVENT_SCOPE(FClothingSimulationCollider_PreUpdate);
 	SCOPE_CYCLE_COUNTER(STAT_ChaosClothingSimulationColliderUpdate);
 
 	// Add or re-add the external collision particles
@@ -676,7 +677,7 @@ void FClothingSimulationCollider::Update(FClothingSimulationSolver* Solver, FClo
 {
 	check(Solver);
 	check(Cloth);
-
+	TRACE_CPUPROFILER_EVENT_SCOPE(FClothingSimulationCollider_Update);
 	SCOPE_CYCLE_COUNTER(STAT_ChaosClothingSimulationColliderUpdate);
 
 	// Update the collision transforms
diff --git a/Engine/Plugins/Experimental/ChaosCloth/Source/Chaos/Private/ChaosCloth/ChaosClothingSimulationMesh.cpp b/Engine/Plugins/Experimental/ChaosCloth/Source/Chaos/Private/ChaosCloth/ChaosClothingSimulationMesh.cpp
--- a/Engine/Plugins/Experimental/ChaosCloth/Source/Chaos/Private/ChaosCloth/ChaosClothingSimulationMesh.cpp
+++ b/Engine/Plugins/Experimental/ChaosCloth/Source/Chaos/Private/ChaosCloth/ChaosClothingSimulationMesh.cpp
@@ -127,6 +127,7 @@ bool FClothingSimulationMesh::WrapDeformLOD(
 	const FVec3* Positions,
 	FVec3* OutPositions) const
 {
+	TRACE_CPUPROFILER_EVENT_SCOPE(FClothingSimulationMesh_WrapDeformLOD);
 	SCOPE_CYCLE_COUNTER(STAT_ChaosClothWrapDeformMesh);
 
 	const int32 NumLODsPassed = FMath::Abs(LODIndex - PrevLODIndex);
@@ -168,6 +169,7 @@ bool FClothingSimulationMesh::WrapDeformLOD(
 	FVec3* OutPositions1,
 	FVec3* OutVelocities) const
 {
+	TRACE_CPUPROFILER_EVENT_SCOPE(FClothingSimulationMesh_WrapDeformLOD);
 	SCOPE_CYCLE_COUNTER(STAT_ChaosClothWrapDeformClothLOD);
 
 	const int32 NumLODsPassed = FMath::Abs(LODIndex - PrevLODIndex);
@@ -213,6 +215,7 @@ FORCEINLINE static void AddInfluence(FVector& OutPosition, FVector& OutNormal, c
 
 void FClothingSimulationMesh::SkinPhysicsMesh(int32 LODIndex, const FVec3& LocalSpaceLocation, FVec3* OutPositions, FVec3* OutNormals) const
 {
+	TRACE_CPUPROFILER_EVENT_SCOPE(FClothingSimulationMesh_SkinPhysicsMesh);
 	SCOPE_CYCLE_COUNTER(STAT_ChaosClothSkinPhysicsMesh);
 	SCOPE_CYCLE_COUNTER(STAT_ClothSkinPhysMesh);
 
diff --git a/Engine/Plugins/Experimental/ChaosCloth/Source/Chaos/Private/ChaosCloth/ChaosClothingSimulationSolver.cpp b/Engine/Plugins/Experimental/ChaosCloth/Source/Chaos/Private/ChaosCloth/ChaosClothingSimulationSolver.cpp
--- a/Engine/Plugins/Experimental/ChaosCloth/Source/Chaos/Private/ChaosCloth/ChaosClothingSimulationSolver.cpp
+++ b/Engine/Plugins/Experimental/ChaosCloth/Source/Chaos/Private/ChaosCloth/ChaosClothingSimulationSolver.cpp
@@ -727,6 +727,7 @@ void FClothingSimulationSolver::AddExternalForces(uint32 GroupId, bool bUseLegac
 
 void FClothingSimulationSolver::ApplyPreSimulationTransforms()
 {
+	TRACE_CPUPROFILER_EVENT_SCOPE(FClothingSimulationSolver_ApplyPreSimulationTransforms);
 	const FVec3 DeltaLocalSpaceLocation = LocalSpaceLocation - OldLocalSpaceLocation;
 
 	const TPBDActiveView<FPBDParticles>& ParticlesActiveView = Evolution->ParticlesActiveView();
@@ -735,6 +736,7 @@ void FClothingSimulationSolver::ApplyPreSimulationTransforms()
 	ParticlesActiveView.RangeFor(
 		[this, &ParticleGroupIds, &DeltaLocalSpaceLocation](FPBDParticles& Particles, int32 Offset, int32 Range)
 		{
+			TRACE_CPUPROFILER_EVENT_SCOPE(FClothingSimulationSolver_ParticlePreSimulationTransforms);
 			SCOPE_CYCLE_COUNTER(STAT_ChaosClothParticlePreSimulationTransforms);
 
 			const int32 RangeSize = Range - Offset;
@@ -760,6 +762,7 @@ void FClothingSimulationSolver::ApplyPreSimulationTransforms()
 #endif
 
 	{
+		TRACE_CPUPROFILER_EVENT_SCOPE(FClothingSimulationSolver_CollisionPreSimulationTransforms);
 		SCOPE_CYCLE_COUNTER(STAT_ChaosClothCollisionPreSimulationTransforms);
 
 		const TPBDActiveView<FKinematicGeometryClothParticles>& CollisionParticlesActiveView = Evolution->CollisionParticlesActiveView();
@@ -781,6 +784,7 @@ void FClothingSimulationSolver::ApplyPreSimulationTransforms()
 
 void FClothingSimulationSolver::Update(FReal InDeltaTime)
 {
+	TRACE_CPUPROFILER_EVENT_SCOPE(FClothingSimulationSolver_Update);
 	SCOPE_CYCLE_COUNTER(STAT_ChaosClothSolverUpdate);
 
 	// Filter delta time to smoothen time variations and prevent unwanted vibrations
@@ -790,6 +794,7 @@ void FClothingSimulationSolver::Update(FReal InDeltaTime)
 
 	// Update Cloths and cloth colliders
 	{
+		TRACE_CPUPROFILER_EVENT_SCOPE(FClothingSimulationSolver_UpdateCloths);
 		SCOPE_CYCLE_COUNTER(STAT_ChaosClothSolverUpdateCloths);
 
 		Swap(OldCollisionTransforms, CollisionTransforms);
@@ -820,6 +825,7 @@ void FClothingSimulationSolver::Update(FReal InDeltaTime)
 
 	// Pre solver step, apply group space transforms for teleport and linear/delta ratios, ...etc
 	{
+		TRACE_CPUPROFILER_EVENT_SCOPE(FClothingSimulationSolver_UpdatePreSolverStep);
 		SCOPE_CYCLE_COUNTER(STAT_ChaosClothSolverUpdatePreSolverStep);
 
 		ApplyPreSimulationTransforms();
@@ -827,6 +833,7 @@ void FClothingSimulationSolver::Update(FReal InDeltaTime)
 
 	// Advance Sim
 	{
+		TRACE_CPUPROFILER_EVENT_SCOPE(FClothingSimulationSolver_UpdateSolverStep);
 		SCOPE_CYCLE_COUNTER(STAT_ChaosClothSolverUpdateSolverStep);
 		SCOPE_CYCLE_COUNTER(STAT_ClothInternalSolve);
 
@@ -845,6 +852,7 @@ void FClothingSimulationSolver::Update(FReal InDeltaTime)
 
 	// Post solver step, update normals, ...etc
 	{
+		TRACE_CPUPROFILER_EVENT_SCOPE(FClothingSimulationSolver_UpdatePostSolverStep);
 		SCOPE_CYCLE_COUNTER(STAT_ChaosClothSolverUpdatePostSolverStep);
 		SCOPE_CYCLE_COUNTER(STAT_ClothComputeNormals);
 
@@ -861,6 +869,7 @@ void FClothingSimulationSolver::Update(FReal InDeltaTime)
 
 FBoxSphereBounds FClothingSimulationSolver::CalculateBounds() const
 {
+	TRACE_CPUPROFILER_EVENT_SCOPE(FClothingSimulationSolver_CalculateBounds);
 	SCOPE_CYCLE_COUNTER(STAT_ChaosClothSolverCalculateBounds);
 
 	const TPBDActiveView<FPBDParticles>& ParticlesActiveView = Evolution->ParticlesActiveView();
diff --git a/Engine/Source/Runtime/Experimental/Chaos/Private/Chaos/Framework/Parallel.cpp b/Engine/Source/Runtime/Experimental/Chaos/Private/Chaos/Framework/Parallel.cpp
--- a/Engine/Source/Runtime/Experimental/Chaos/Private/Chaos/Framework/Parallel.cpp
+++ b/Engine/Source/Runtime/Experimental/Chaos/Private/Chaos/Framework/Parallel.cpp
@@ -21,6 +21,7 @@ namespace Chaos
 
 void Chaos::PhysicsParallelFor(int32 InNum, TFunctionRef<void(int32)> InCallable, bool bForceSingleThreaded)
 {
+	TRACE_CPUPROFILER_EVENT_SCOPE(Chaos_PhysicsParallelFor);
 	using namespace Chaos;
 	// Passthrough for now, except with global flag to disable parallel
 	
diff --git a/Engine/Source/Runtime/Experimental/Chaos/Private/Chaos/PBDAxialSpringConstraints.cpp b/Engine/Source/Runtime/Experimental/Chaos/Private/Chaos/PBDAxialSpringConstraints.cpp
--- a/Engine/Source/Runtime/Experimental/Chaos/Private/Chaos/PBDAxialSpringConstraints.cpp
+++ b/Engine/Source/Runtime/Experimental/Chaos/Private/Chaos/PBDAxialSpringConstraints.cpp
@@ -55,6 +55,7 @@ void FPBDAxialSpringConstraints::ApplyHelper(FPBDParticles& Particles, const FRe
 
 void FPBDAxialSpringConstraints::Apply(FPBDParticles& Particles, const FReal Dt) const
 {
+	TRACE_CPUPROFILER_EVENT_SCOPE(FPBDAxialSpringConstraints_Apply);
 	SCOPE_CYCLE_COUNTER(STAT_PBD_AxialSpring);
 	if (ConstraintsPerColor.Num() > 0 && Constraints.Num() > Chaos_AxialSpring_ParallelConstraintCount)
 	{
diff --git a/Engine/Source/Runtime/Experimental/Chaos/Private/Chaos/PBDEvolution.cpp b/Engine/Source/Runtime/Experimental/Chaos/Private/Chaos/PBDEvolution.cpp
--- a/Engine/Source/Runtime/Experimental/Chaos/Private/Chaos/PBDEvolution.cpp
+++ b/Engine/Source/Runtime/Experimental/Chaos/Private/Chaos/PBDEvolution.cpp
@@ -249,6 +249,7 @@ void FPBDEvolution::PreIterationUpdate(
 
 void FPBDEvolution::AdvanceOneTimeStep(const FReal Dt)
 {
+	TRACE_CPUPROFILER_EVENT_SCOPE(FPBDEvolution_AdvanceOneTimeStep);
 	SCOPE_CYCLE_COUNTER(STAT_ChaosPBDVAdvanceTime);
 
 	// Advance time
@@ -260,6 +261,7 @@ void FPBDEvolution::AdvanceOneTimeStep(const FReal Dt)
 	const bool bWriteCCDContacts = CVarChaosPBDEvolutionWriteCCDContacts.GetValueOnAnyThread();
 
 	{
+		TRACE_CPUPROFILER_EVENT_SCOPE(ChaosPBDPreIterationUpdates);
 		SCOPE_CYCLE_COUNTER(STAT_ChaosPBDPreIterationUpdates);
 
 		MParticlesActiveView.RangeFor(
@@ -324,6 +326,7 @@ void FPBDEvolution::AdvanceOneTimeStep(const FReal Dt)
 	{
 		if (MCollisionKinematicUpdate)
 		{
+			TRACE_CPUPROFILER_EVENT_SCOPE(ChaosPBDCollisionKinematicUpdate);
 			SCOPE_CYCLE_COUNTER(STAT_ChaosPBDCollisionKinematicUpdate);
 
 			MCollisionParticlesActiveView.SequentialFor(
@@ -338,6 +341,7 @@ void FPBDEvolution::AdvanceOneTimeStep(const FReal Dt)
 		}
 
 		{
+			TRACE_CPUPROFILER_EVENT_SCOPE(ChaosPBDClearCollidedArray);
 			SCOPE_CYCLE_COUNTER(STAT_ChaosPBDClearCollidedArray);
 			memset(MCollided.GetData(), 0, MCollided.Num() * sizeof(bool));
 		}
@@ -345,6 +349,7 @@ void FPBDEvolution::AdvanceOneTimeStep(const FReal Dt)
 
 	// Constraint init (clear XPBD's Lambdas, init self collisions)
 	{
+		TRACE_CPUPROFILER_EVENT_SCOPE(ChaosXPBDConstraintsInit);
 		SCOPE_CYCLE_COUNTER(STAT_ChaosXPBDConstraintsInit);
 		MConstraintInitsActiveView.SequentialFor(
 			[this, Dt](TArray<TFunction<void(const FPBDParticles&, const FReal)>>& ConstraintInits, int32 Index)
@@ -379,6 +384,7 @@ void FPBDEvolution::AdvanceOneTimeStep(const FReal Dt)
 
 	// Iteration loop
 	{
+		TRACE_CPUPROFILER_EVENT_SCOPE(ChaosPBDIterationLoop);
 		SCOPE_CYCLE_COUNTER(STAT_ChaosPBDIterationLoop);
 
 		for (int32 i = 0; i < MNumIterations; ++i)
@@ -386,6 +392,7 @@ void FPBDEvolution::AdvanceOneTimeStep(const FReal Dt)
 			MConstraintRulesActiveView.RangeFor(
 				[this, Dt](TArray<TFunction<void(FPBDParticles&, const FReal)>>& ConstraintRules, int32 Offset, int32 Range)
 				{
+					TRACE_CPUPROFILER_EVENT_SCOPE(ChaosPBDConstraintRule);
 					SCOPE_CYCLE_COUNTER(STAT_ChaosPBDConstraintRule);
 					for (int32 ConstraintIndex = Offset; ConstraintIndex < Range; ++ConstraintIndex)
 					{
@@ -394,10 +401,12 @@ void FPBDEvolution::AdvanceOneTimeStep(const FReal Dt)
 				}, bUseSingleThreadedRange);
 
 			{
+				TRACE_CPUPROFILER_EVENT_SCOPE(ChaosPBDCollisionRule);
 				SCOPE_CYCLE_COUNTER(STAT_ChaosPBDCollisionRule);
 				MParticlesActiveView.RangeFor(
 					[this, &CollisionRule, &CCDCollisionRule, Dt](FPBDParticles& Particles, int32 Offset, int32 Range)
 					{
+						TRACE_CPUPROFILER_EVENT_SCOPE(ChaosPBDCollisionRuleP);
 						const uint32 DynamicGroupId = MParticleGroupIds[Offset];  // Particle group Id, must be the same across the entire range
 						const bool bUseCCD = MGroupUseCCDs[DynamicGroupId];
 						if (!bUseCCD)
@@ -432,6 +441,7 @@ void FPBDEvolution::AdvanceOneTimeStep(const FReal Dt)
 		MParticlesActiveView.ParallelFor(
 			[&CollisionRule, Dt](FPBDParticles& Particles, int32 Index)
 			{
+				TRACE_CPUPROFILER_EVENT_SCOPE(ChaosPBDCollisionRuleFriction);
 				CollisionRule.ApplyFriction(Particles, Dt, Index);
 			}, bUseSingleThreadedRange, MinParallelBatchSize);
 	}
diff --git a/Engine/Source/Runtime/Experimental/Chaos/Private/Chaos/PBDLongRangeConstraints.cpp b/Engine/Source/Runtime/Experimental/Chaos/Private/Chaos/PBDLongRangeConstraints.cpp
--- a/Engine/Source/Runtime/Experimental/Chaos/Private/Chaos/PBDLongRangeConstraints.cpp
+++ b/Engine/Source/Runtime/Experimental/Chaos/Private/Chaos/PBDLongRangeConstraints.cpp
@@ -14,6 +14,7 @@ using namespace Chaos;
 
 void FPBDLongRangeConstraints::Apply(FPBDParticles& Particles, const FReal /*Dt*/, const TArray<int32>& ConstraintIndices) const
 {
+	TRACE_CPUPROFILER_EVENT_SCOPE(FPBDLongRangeConstraints_Apply);
 	SCOPE_CYCLE_COUNTER(STAT_PBD_LongRange);
 	if (!Stiffness.HasWeightMap())
 	{
@@ -37,6 +38,7 @@ void FPBDLongRangeConstraints::Apply(FPBDParticles& Particles, const FReal /*Dt*
 
 void FPBDLongRangeConstraints::Apply(FPBDParticles& Particles, const FReal /*Dt*/) const
 {
+	TRACE_CPUPROFILER_EVENT_SCOPE(FPBDLongRangeConstraints_Apply);
 	SCOPE_CYCLE_COUNTER(STAT_PBD_LongRange);
 
 	if (!Stiffness.HasWeightMap())
diff --git a/Engine/Source/Runtime/Experimental/Chaos/Private/Chaos/PBDLongRangeConstraintsBase.cpp b/Engine/Source/Runtime/Experimental/Chaos/Private/Chaos/PBDLongRangeConstraintsBase.cpp
--- a/Engine/Source/Runtime/Experimental/Chaos/Private/Chaos/PBDLongRangeConstraintsBase.cpp
+++ b/Engine/Source/Runtime/Experimental/Chaos/Private/Chaos/PBDLongRangeConstraintsBase.cpp
@@ -260,6 +260,7 @@ void FPBDLongRangeConstraintsBase::ComputeGeodesicConstraints(
     const TMap<int32, TSet<int32>>& PointToNeighbors,
     const int32 MaxNumTetherIslands)
 {
+	TRACE_CPUPROFILER_EVENT_SCOPE(FPBDLongRangeConstraintsBase_ComputeGeodesicConstraints);
 	SCOPE_CYCLE_COUNTER(STAT_ChaosClothComputeGeodesicConstraints);
 
 	// Fill up the list of all used indices
