Chaos Cloth - Fix broken cloth simulation on legacy wind using assets. [Edigrate CL 16420140]

diff --git a/Engine/Plugins/Experimental/ChaosCloth/Source/Chaos/Private/ChaosCloth/ChaosClothingSimulationSolver.cpp b/Engine/Plugins/Experimental/ChaosCloth/Source/Chaos/Private/ChaosCloth/ChaosClothingSimulationSolver.cpp
--- a/Engine/Plugins/Experimental/ChaosCloth/Source/Chaos/Private/ChaosCloth/ChaosClothingSimulationSolver.cpp
+++ b/Engine/Plugins/Experimental/ChaosCloth/Source/Chaos/Private/ChaosCloth/ChaosClothingSimulationSolver.cpp
@@ -674,7 +674,7 @@ void FClothingSimulationSolver::AddExternalForces(uint32 GroupId, bool bUseLegac
 		const FVec3 LegacyWindVelocity = WindVelocity * LegacyWindMultiplier;
 
 		Evolution->GetForceFunction(GroupId) =
-			[this, bHasVelocityField, bHasForceField, bHasFictitiousForces, bUseLegacyWind, &LegacyWindVelocity, &AngularDisplacement](FPBDParticles& Particles, const FReal Dt, const int32 Index)
+			[this, bHasVelocityField, bHasForceField, bHasFictitiousForces, bUseLegacyWind, LegacyWindVelocity, AngularDisplacement](FPBDParticles& Particles, const FReal Dt, const int32 Index)
 			{
 				FVec3 Forces((FReal)0.);
 
