Fix bad field evaluation when multiple cloths per solver

diff --git a/Engine/Plugins/Experimental/ChaosCloth/Source/Chaos/Private/ChaosCloth/ChaosClothingSimulationSolver.cpp b/Engine/Plugins/Experimental/ChaosCloth/Source/Chaos/Private/ChaosCloth/ChaosClothingSimulationSolver.cpp
--- a/Engine/Plugins/Experimental/ChaosCloth/Source/Chaos/Private/ChaosCloth/ChaosClothingSimulationSolver.cpp
+++ b/Engine/Plugins/Experimental/ChaosCloth/Source/Chaos/Private/ChaosCloth/ChaosClothingSimulationSolver.cpp
@@ -637,35 +637,8 @@ void FClothingSimulationSolver::AddExternalForces(uint32 GroupId, bool bUseLegac
 {
 	if (Evolution)
 	{
-		bool bHasVelocityField = false;
-		bool bHasForceField = false;
-
-		if (!PerSolverField.IsEmpty())
-		{
-			TArray<FVector>& SamplePositions = PerSolverField.GetSamplePositions();
-			TArray<FFieldContextIndex>& SampleIndices = PerSolverField.GetSampleIndices();
-
-			const uint32 NumParticles = Evolution->Particles().Size();
-
-			SamplePositions.SetNum(NumParticles,false);
-			SampleIndices.SetNum(NumParticles,false);
-
-			for (uint32 Index = 0; Index < NumParticles; ++Index)
-			{
-				SamplePositions[Index] = Evolution->Particles().X(Index) + LocalSpaceLocation;
-				SampleIndices[Index] = FFieldContextIndex(Index, Index);
-			}
-			PerSolverField.ComputeFieldLinearImpulse(GetTime());
-
-			if (PerSolverField.GetOutputResults(EFieldCommandOutputType::LinearVelocity).Num() > 0)
-			{
-				bHasVelocityField = true;
-			}
-			if (PerSolverField.GetOutputResults(EFieldCommandOutputType::LinearForce).Num() > 0)
-			{
-				bHasForceField = true;
-			}
-		}
+		const bool bHasVelocityField = !PerSolverField.GetOutputResults(EFieldCommandOutputType::LinearVelocity).IsEmpty();
+		const bool bHasForceField = !PerSolverField.GetOutputResults(EFieldCommandOutputType::LinearForce).IsEmpty();
 
 		const FVec3& AngularDisplacement = FictitiousAngularDisplacement[GroupId];
 		const bool bHasFictitiousForces = !AngularDisplacement.IsNearlyZero();
@@ -782,6 +755,27 @@ void FClothingSimulationSolver::ApplyPreSimulationTransforms()
 	}
 }
 
+void FClothingSimulationSolver::UpdateSolverField()
+{
+	if (Evolution && !PerSolverField.IsEmpty())
+	{
+		TArray<FVector>& SamplePositions = PerSolverField.GetSamplePositions();
+		TArray<FFieldContextIndex>& SampleIndices = PerSolverField.GetSampleIndices();
+
+		const uint32 NumParticles = Evolution->Particles().Size();
+
+		SamplePositions.SetNum(NumParticles, false);
+		SampleIndices.SetNum(NumParticles, false);
+
+		for (uint32 Index = 0; Index < NumParticles; ++Index)
+		{
+			SamplePositions[Index] = Evolution->Particles().X(Index) + LocalSpaceLocation;
+			SampleIndices[Index] = FFieldContextIndex(Index, Index);
+		}
+		PerSolverField.ComputeFieldLinearImpulse(GetTime());
+	}
+}
+
 void FClothingSimulationSolver::Update(FReal InDeltaTime)
 {
 	TRACE_CPUPROFILER_EVENT_SCOPE(FClothingSimulationSolver_Update);
@@ -803,6 +797,9 @@ void FClothingSimulationSolver::Update(FReal InDeltaTime)
 		// Clear external collisions so that they can be re-added
 		CollisionParticlesSize = 0;
 
+		// Compute the solver field forces/velocities for future use in the AddExternalForces
+		UpdateSolverField();
+
 		// Run sequential pre-updates first
 		for (FClothingSimulationCloth* const Cloth : Cloths)
 		{
diff --git a/Engine/Plugins/Experimental/ChaosCloth/Source/Chaos/Public/ChaosCloth/ChaosClothingSimulationSolver.h b/Engine/Plugins/Experimental/ChaosCloth/Source/Chaos/Public/ChaosCloth/ChaosClothingSimulationSolver.h
--- a/Engine/Plugins/Experimental/ChaosCloth/Source/Chaos/Public/ChaosCloth/ChaosClothingSimulationSolver.h
+++ b/Engine/Plugins/Experimental/ChaosCloth/Source/Chaos/Public/ChaosCloth/ChaosClothingSimulationSolver.h
@@ -177,6 +177,9 @@ namespace Chaos
 		void ParticleMassUpdateDensity(const FTriangleMesh& Mesh, FReal Density);
 		void ParticleMassClampAndEnslave(int32 Offset, int32 Size, FReal MinPerParticleMass, const TFunctionRef<bool(int32)>& KinematicPredicate);
 
+		// Update the solver field forces/velocities at the particles location
+		void UpdateSolverField();
+
 	private:
 		TUniquePtr<FPBDEvolution> Evolution;
 
