Chaos Cloth - Fix impulse on game hitch caused by the smoothed time step on velocities. Re-apply the smooth time step on forces only and inside the substepping loop.
[Edigrate CL 16972937]

Chaos Cloth - Add CVAR to re-enable old time step smoothing behavior if ever needed.
[Edigrate CL 16986079]

diff --git a/Engine/Plugins/Experimental/ChaosCloth/Source/Chaos/Private/ChaosCloth/ChaosClothingSimulationSolver.cpp b/Engine/Plugins/Experimental/ChaosCloth/Source/Chaos/Private/ChaosCloth/ChaosClothingSimulationSolver.cpp
--- a/Engine/Plugins/Experimental/ChaosCloth/Source/Chaos/Private/ChaosCloth/ChaosClothingSimulationSolver.cpp
+++ b/Engine/Plugins/Experimental/ChaosCloth/Source/Chaos/Private/ChaosCloth/ChaosClothingSimulationSolver.cpp
@@ -29,14 +29,24 @@ static int32 ChaosClothSolverMinParallelBatchSize = 1000;
 static bool bChaosClothSolverParallelClothPreUpdate = false;  // TODO: Doesn't seem to improve much here. Review this after the ISPC implementation.
 static bool bChaosClothSolverParallelClothUpdate = true;
 static bool bChaosClothSolverParallelClothPostUpdate = true;
+static bool bChaosClothSolverUseImprovedTimeStepSmoothing = true;
 
 #if !UE_BUILD_SHIPPING
+static int32 ChaosClothSolverDebugHitchLength = 0;
+static int32 ChaosClothSolverDebugHitchInterval = 0;
+static bool bChaosClothSolverDisableCollision = false;
+
 FAutoConsoleVariableRef CVarChaosClothSolverMinParallelBatchSize(TEXT("p.ChaosCloth.Solver.MinParallelBatchSize"), ChaosClothSolverMinParallelBatchSize, TEXT("The minimum number of particle to process in parallel batch by the solver."));
 FAutoConsoleVariableRef CVarChaosClothSolverParallelClothPreUpdate(TEXT("p.ChaosCloth.Solver.ParallelClothPreUpdate"), bChaosClothSolverParallelClothPreUpdate, TEXT("Pre-transform the cloth particles for each cloth in parallel."));
 FAutoConsoleVariableRef CVarChaosClothSolverParallelClothUpdate(TEXT("p.ChaosCloth.Solver.ParallelClothUpdate"), bChaosClothSolverParallelClothUpdate, TEXT("Skin the physics mesh and do the other cloth update for each cloth in parallel."));
 FAutoConsoleVariableRef CVarChaosClothSolverParallelClothPostUpdate(TEXT("p.ChaosCloth.Solver.ParallelClothPostUpdate"), bChaosClothSolverParallelClothPostUpdate, TEXT("Pre-transform the cloth particles for each cloth in parallel."));
+FAutoConsoleVariableRef CVarChaosClothSolverDebugHitchLength(TEXT("p.ChaosCloth.Solver.DebugHitchLength"), ChaosClothSolverDebugHitchLength, TEXT("Hitch length in ms. Create artificial hitches to debug simulation jitter. 0 to disable"));
+FAutoConsoleVariableRef CVarChaosClothSolverDebugHitchInterval(TEXT("p.ChaosCloth.Solver.DebugHitchInterval"), ChaosClothSolverDebugHitchInterval, TEXT("Hitch interval in frames. Create artificial hitches to debug simulation jitter. 0 to disable"));
+FAutoConsoleVariableRef CVarChaosClothSolverDisableCollision(TEXT("p.ChaosCloth.Solver.DisableCollision"), bChaosClothSolverDisableCollision, TEXT("Disable all collision particles. Needs reset of the simulation (p.ChaosCloth.Reset)."));
 #endif
 
+FAutoConsoleVariableRef CVarChaosClothSolverUseImprovedTimeStepSmoothing(TEXT("p.ChaosCloth.Solver.UseImprovedTimeStepSmoothing"), bChaosClothSolverUseImprovedTimeStepSmoothing, TEXT("Use the time step smoothing on input forces only rather than on the entire cloth solver, in order to avoid miscalculating velocities."));
+
 namespace ChaosClothingSimulationSolverDefault
 {
 	static const FVec3 Gravity(0.f, 0.f, -980.665f);  // cm/s^2
@@ -397,7 +407,16 @@ int32 FClothingSimulationSolver::AddCollisionParticles(int32 NumCollisionParticl
 
 void FClothingSimulationSolver::EnableCollisionParticles(int32 Offset, bool bEnable)
 {
-	Evolution->ActivateCollisionParticleRange(Offset, bEnable);
+#if !UE_BUILD_SHIPPING
+	if (bChaosClothSolverDisableCollision)
+	{
+		Evolution->ActivateCollisionParticleRange(Offset, false);
+	}
+	else
+#endif  // #if !UE_BUILD_SHIPPING
+	{
+		Evolution->ActivateCollisionParticleRange(Offset, bEnable);
+	}
 }
 
 const FVec3* FClothingSimulationSolver::GetCollisionParticleXs(int32 Offset) const
@@ -781,10 +800,33 @@ void FClothingSimulationSolver::Update(FReal InDeltaTime)
 	TRACE_CPUPROFILER_EVENT_SCOPE(FClothingSimulationSolver_Update);
 	SCOPE_CYCLE_COUNTER(STAT_ChaosClothSolverUpdate);
 
-	// Filter delta time to smoothen time variations and prevent unwanted vibrations
-	static const FReal DeltaTimeDecay = 0.1f;
-	const FReal PrevDeltaTime = DeltaTime;
-	DeltaTime = DeltaTime + (InDeltaTime - DeltaTime) * DeltaTimeDecay;
+	if (!bChaosClothSolverUseImprovedTimeStepSmoothing)
+	{
+		// Filter delta time to smoothen time variations and prevent unwanted vibrations
+		// Note: This is now deprecated and replaced by in solver input force timestep smoothing
+		static const FReal DeltaTimeDecay = 0.1f;
+		const FReal PrevDeltaTime = DeltaTime;
+		DeltaTime = DeltaTime + (InDeltaTime - DeltaTime) * DeltaTimeDecay;
+	}
+	else
+	{
+		// Update time step
+		DeltaTime = InDeltaTime;
+	}
+
+#if !UE_BUILD_SHIPPING
+	// Introduce artificial hitches for debugging any simulation jitter
+	if (ChaosClothSolverDebugHitchLength && ChaosClothSolverDebugHitchInterval)
+	{
+		static int32 HitchCounter = 0;
+		if (--HitchCounter < 0)
+		{
+			UE_LOG(LogChaosCloth, Warning, TEXT("Hitching for %dms"), ChaosClothSolverDebugHitchLength);
+			FPlatformProcess::Sleep((float)ChaosClothSolverDebugHitchLength * 0.001f);
+			HitchCounter = ChaosClothSolverDebugHitchInterval;
+		}
+	}
+#endif  // #if !UE_BUILD_SHIPPING
 
 	// Update Cloths and cloth colliders
 	{
@@ -840,11 +882,11 @@ void FClothingSimulationSolver::Update(FReal InDeltaTime)
 	
 		for (int32 i = 0; i < NumSubsteps; ++i)
 		{
-			Evolution->AdvanceOneTimeStep(SubstepDeltaTime);
+			Evolution->AdvanceOneTimeStep(SubstepDeltaTime, bChaosClothSolverUseImprovedTimeStepSmoothing);
 		}
 
 		Time = Evolution->GetTime();
-		UE_LOG(LogChaosCloth, VeryVerbose, TEXT("DeltaTime: %.6f, FilteredDeltaTime: %.6f, Time = %.6f"), InDeltaTime, DeltaTime, Time);
+		UE_LOG(LogChaosCloth, VeryVerbose, TEXT("DeltaTime: %.6f, Time = %.6f"), DeltaTime, Time);
 	}
 
 	// Post solver step, update normals, ...etc
diff --git a/Engine/Source/Runtime/Experimental/Chaos/Private/Chaos/PBDEvolution.cpp b/Engine/Source/Runtime/Experimental/Chaos/Private/Chaos/PBDEvolution.cpp
--- a/Engine/Source/Runtime/Experimental/Chaos/Private/Chaos/PBDEvolution.cpp
+++ b/Engine/Source/Runtime/Experimental/Chaos/Private/Chaos/PBDEvolution.cpp
@@ -32,6 +32,7 @@ DECLARE_CYCLE_STAT(TEXT("Chaos XPBD Constraints Init"), STAT_ChaosXPBDConstraint
 
 TAutoConsoleVariable<bool> CVarChaosPBDEvolutionUseNestedParallelFor(TEXT("p.Chaos.PBDEvolution.UseNestedParallelFor"), true, TEXT(""), ECVF_Cheat);
 TAutoConsoleVariable<bool> CVarChaosPBDEvolutionFastPositionBasedFriction(TEXT("p.Chaos.PBDEvolution.FastPositionBasedFriction"), true, TEXT(""), ECVF_Cheat);
+TAutoConsoleVariable<bool> CVarChaosPBDEvolutionUseSmoothTimeStep(TEXT("p.Chaos.PBDEvolution.UseSmoothTimeStep"), true, TEXT(""), ECVF_Cheat);
 TAutoConsoleVariable<int32> CVarChaosPBDEvolutionMinParallelBatchSize(TEXT("p.Chaos.PBDEvolution.MinParallelBatchSize"), 300, TEXT(""), ECVF_Cheat);
 TAutoConsoleVariable<bool> CVarChaosPBDEvolutionWriteCCDContacts(TEXT("p.Chaos.PBDEvolution.WriteCCDContacts"), false, TEXT("Write CCD collision contacts and normals potentially causing the CCD collision threads to lock, allowing for debugging of these contacts."), ECVF_Cheat);
 
@@ -77,6 +78,7 @@ FPBDEvolution::FPBDEvolution(FPBDParticles&& InParticles, FKinematicGeometryClot
 	, MCoefficientOfFriction(CoefficientOfFriction)
 	, MDamping(Damping)
 	, MTime(0)
+	, MSmoothDt(1.f / 30.f)  // Initialize filtered timestep at 30fps
 {
 	// Add group arrays
 	TArrayCollection::AddArray(&MGroupGravityForces);
@@ -229,7 +231,7 @@ void FPBDEvolution::PreIterationUpdate(
 				}
 
 				// Euler Step Velocity
-				MParticles.V(Index) += MParticles.F(Index) * MParticles.InvM(Index) * Dt;
+				MParticles.V(Index) += MParticles.F(Index) * MParticles.InvM(Index) * MSmoothDt;
 
 				// Damp Velocity Rule
 				if (bDampVelocityRule)
@@ -247,7 +249,7 @@ void FPBDEvolution::PreIterationUpdate(
 		}, RangeSize < MinParallelBatchSize);
 }
 
-void FPBDEvolution::AdvanceOneTimeStep(const FReal Dt)
+void FPBDEvolution::AdvanceOneTimeStep(const FReal Dt, const bool bSmoothDt)
 {
 	TRACE_CPUPROFILER_EVENT_SCOPE(FPBDEvolution_AdvanceOneTimeStep);
 	SCOPE_CYCLE_COUNTER(STAT_ChaosPBDVAdvanceTime);
@@ -255,6 +257,17 @@ void FPBDEvolution::AdvanceOneTimeStep(const FReal Dt)
 	// Advance time
 	MTime += Dt;
 
+	// Filter delta time to smoothen time variations and prevent unwanted vibrations, works best on Forces
+	if (bSmoothDt && CVarChaosPBDEvolutionUseSmoothTimeStep.GetValueOnAnyThread())
+	{
+		constexpr FReal DeltaTimeDecay = (FReal)0.1;
+		MSmoothDt += (Dt - MSmoothDt) * DeltaTimeDecay;
+	}
+	else
+	{
+		MSmoothDt = Dt;
+	}
+
 	// Don't bother with threaded execution if we don't have enough work to make it worth while.
 	const bool bUseSingleThreadedRange = !CVarChaosPBDEvolutionUseNestedParallelFor.GetValueOnAnyThread();
 	const int32 MinParallelBatchSize = CVarChaosPBDEvolutionMinParallelBatchSize.GetValueOnAnyThread(); // TODO: 1000 is a guess, tune this!
diff --git a/Engine/Source/Runtime/Experimental/Chaos/Public/Chaos/PBDEvolution.h b/Engine/Source/Runtime/Experimental/Chaos/Public/Chaos/PBDEvolution.h
--- a/Engine/Source/Runtime/Experimental/Chaos/Public/Chaos/PBDEvolution.h
+++ b/Engine/Source/Runtime/Experimental/Chaos/Public/Chaos/PBDEvolution.h
@@ -20,7 +20,8 @@ class CHAOS_API FPBDEvolution : public TArrayCollection
 	FPBDEvolution(FPBDParticles&& InParticles, FKinematicGeometryClothParticles&& InGeometryParticles, TArray<TVec3<int32>>&& CollisionTriangles, int32 NumIterations = 1, FReal CollisionThickness = 0, FReal SelfCollisionsThickness = 0, FReal CoefficientOfFriction = 0, FReal Damping = 0.04f);
 	~FPBDEvolution() {}
 
-	void AdvanceOneTimeStep(const FReal dt);
+	// Advance one time step. Filter the input time step if specified.
+	void AdvanceOneTimeStep(const FReal Dt, const bool bSmoothDt = true);
 
 	// Remove all particles, will also reset all rules
 	void ResetParticles();
@@ -173,5 +174,6 @@ private:
 	FReal MCoefficientOfFriction;
 	FReal MDamping;
 	FReal MTime;
+	FReal MSmoothDt;
 };
 }
