Chaos Cloth - Remove stiffness multiplier when migrating from the NvCloth config.

diff --git a/Engine/Plugins/Experimental/ChaosCloth/Source/Chaos/Private/ChaosCloth/ChaosClothConfig.cpp b/Engine/Plugins/Experimental/ChaosCloth/Source/Chaos/Private/ChaosCloth/ChaosClothConfig.cpp
--- a/Engine/Plugins/Experimental/ChaosCloth/Source/Chaos/Private/ChaosCloth/ChaosClothConfig.cpp
+++ b/Engine/Plugins/Experimental/ChaosCloth/Source/Chaos/Private/ChaosCloth/ChaosClothConfig.cpp
@@ -33,20 +33,16 @@ void UChaosClothConfig::MigrateFrom(const FClothConfig_Legacy& ClothConfig)
 {
 #if WITH_EDITORONLY_DATA
 	const float VerticalStiffness =
-		ClothConfig.VerticalConstraintConfig.Stiffness *
-		ClothConfig.VerticalConstraintConfig.StiffnessMultiplier;
+		ClothConfig.VerticalConstraintConfig.Stiffness;
 	const float HorizontalStiffness =
-		ClothConfig.HorizontalConstraintConfig.Stiffness *
-		ClothConfig.HorizontalConstraintConfig.StiffnessMultiplier;
+		ClothConfig.HorizontalConstraintConfig.Stiffness;
 	EdgeStiffnessWeighted.Low = EdgeStiffnessWeighted.High = FMath::Clamp((VerticalStiffness + HorizontalStiffness) * 0.5f, 0.f, 1.f);
 
 	BendingStiffnessWeighted.Low = BendingStiffnessWeighted.High = FMath::Clamp(
-		ClothConfig.BendConstraintConfig.Stiffness *
-		ClothConfig.BendConstraintConfig.StiffnessMultiplier, 0.f, 1.f);
+		ClothConfig.BendConstraintConfig.Stiffness, 0.f, 1.f);
 
 	AreaStiffnessWeighted.Low = AreaStiffnessWeighted.High = FMath::Clamp(
-		ClothConfig.ShearConstraintConfig.Stiffness *
-		ClothConfig.ShearConstraintConfig.StiffnessMultiplier, 0.f, 1.f);
+		ClothConfig.ShearConstraintConfig.Stiffness, 0.f, 1.f);
 
 	AnimDriveStiffness.Low = 0.f;
 	AnimDriveStiffness.High = FMath::Clamp(ClothConfig.AnimDriveSpringStiffness, 0.f, 1.f);
