Chaos Cloth - Fix LWC skinning code.

diff --git a/Engine/Plugins/Experimental/ChaosCloth/Source/Chaos/Private/ChaosCloth/ChaosClothingSimulationMesh.cpp b/Engine/Plugins/Experimental/ChaosCloth/Source/Chaos/Private/ChaosCloth/ChaosClothingSimulationMesh.cpp
--- a/Engine/Plugins/Experimental/ChaosCloth/Source/Chaos/Private/ChaosCloth/ChaosClothingSimulationMesh.cpp
+++ b/Engine/Plugins/Experimental/ChaosCloth/Source/Chaos/Private/ChaosCloth/ChaosClothingSimulationMesh.cpp
@@ -225,10 +225,10 @@ bool FClothingSimulationMesh::WrapDeformLOD(
 }
 
 // Inline function used to force the unrolling of the skinning loop
-FORCEINLINE static void AddInfluence(FVector& OutPosition, FVector& OutNormal, const FVector3f& RefParticle, const FVector3f& RefNormal, const FMatrix44f& BoneMatrix, const FRealSingle Weight)
+FORCEINLINE static void AddInfluence(FVector3f& OutPosition, FVector3f& OutNormal, const FVector3f& RefParticle, const FVector3f& RefNormal, const FMatrix44f& BoneMatrix, const float Weight)
 {
-	OutPosition += FVector4(BoneMatrix.TransformPosition(RefParticle) * Weight);
-	OutNormal += FVector4(BoneMatrix.TransformVector(RefNormal) * Weight);
+	OutPosition += BoneMatrix.TransformPosition(RefParticle) * Weight;
+	OutNormal += BoneMatrix.TransformVector(RefNormal) * Weight;
 }
 
 void FClothingSimulationMesh::SkinPhysicsMesh(int32 LODIndex, const FVec3& LocalSpaceLocation, FVec3* OutPositions, FVec3* OutNormals) const
@@ -245,12 +245,9 @@ void FClothingSimulationMesh::SkinPhysicsMesh(int32 LODIndex, const FVec3& Local
 
 	check(SkeletalMeshComponent && SkeletalMeshComponent->GetClothingSimulationContext());
 	const FClothingSimulationContextCommon* const Context = static_cast<const FClothingSimulationContextCommon*>(SkeletalMeshComponent->GetClothingSimulationContext());
-	FTransform ComponentToLocalSpace = Context->ComponentToWorld;
-	ComponentToLocalSpace.AddToTranslation(-LocalSpaceLocation);
-
-	// Zero out positions & normals
-	FMemory::Memzero((uint8*)OutPositions, NumPoints * sizeof(FVec3));  // It is faster on some platforms to zero the memory first
-	FMemory::Memzero((uint8*)OutNormals, NumPoints * sizeof(FVec3));    // instead of changing this function to work with uninitialized memory
+	FTransform ComponentToLocalSpaceReal = Context->ComponentToWorld;
+	ComponentToLocalSpaceReal.AddToTranslation(-LocalSpaceLocation);
+	const FTransform3f ComponentToLocalSpace(ComponentToLocalSpaceReal);
 
 	const int32* const RESTRICT BoneMap = Asset->UsedBoneIndices.GetData();
 	const FMatrix44f* const RESTRICT BoneMatrices = Context->RefToLocals.GetData();
@@ -272,43 +269,45 @@ void FClothingSimulationMesh::SkinPhysicsMesh(int32 LODIndex, const FVec3& Local
 	}
 	else
 	{
-	static const uint32 MinParallelVertices = 500;  // 500 seems to be the lowest threshold still giving gains even on profiled assets that are only using a small number of influences
+		static const uint32 MinParallelVertices = 500;  // 500 seems to be the lowest threshold still giving gains even on profiled assets that are only using a small number of influences
 
-	ParallelFor(NumPoints, [&PhysicalMeshData, &ComponentToLocalSpace, BoneMap, BoneMatrices, &OutPositions, &OutNormals](uint32 VertIndex)
-	{
-		const uint16* const RESTRICT BoneIndices = PhysicalMeshData.BoneData[VertIndex].BoneIndices;
-		const FRealSingle* const RESTRICT BoneWeights = PhysicalMeshData.BoneData[VertIndex].BoneWeights;
-
-		// WARNING - HORRIBLE UNROLLED LOOP + JUMP TABLE BELOW
-		// done this way because this is a pretty tight and performance critical loop. essentially
-		// rather than checking each influence we can just jump into this switch and fall through
-		// everything to compose the final skinned data
-		const FVec3& RefParticle = PhysicalMeshData.Vertices[VertIndex];
-		const FVec3& RefNormal = PhysicalMeshData.Normals[VertIndex];
-		FVec3& OutPosition = OutPositions[VertIndex];
-		FVec3& OutNormal = OutNormals[VertIndex];
-		switch (PhysicalMeshData.BoneData[VertIndex].NumInfluences)
+		ParallelFor(NumPoints, [&PhysicalMeshData, &ComponentToLocalSpace, BoneMap, BoneMatrices, &OutPositions, &OutNormals](uint32 VertIndex)
 		{
-		case 12: AddInfluence(OutPosition, OutNormal, RefParticle, RefNormal, BoneMatrices[BoneMap[BoneIndices[11]]], BoneWeights[11]);  // Intentional fall through
-		case 11: AddInfluence(OutPosition, OutNormal, RefParticle, RefNormal, BoneMatrices[BoneMap[BoneIndices[10]]], BoneWeights[10]);  // Intentional fall through
-		case 10: AddInfluence(OutPosition, OutNormal, RefParticle, RefNormal, BoneMatrices[BoneMap[BoneIndices[ 9]]], BoneWeights[ 9]);  // Intentional fall through
-		case  9: AddInfluence(OutPosition, OutNormal, RefParticle, RefNormal, BoneMatrices[BoneMap[BoneIndices[ 8]]], BoneWeights[ 8]);  // Intentional fall through
-		case  8: AddInfluence(OutPosition, OutNormal, RefParticle, RefNormal, BoneMatrices[BoneMap[BoneIndices[ 7]]], BoneWeights[ 7]);  // Intentional fall through
-		case  7: AddInfluence(OutPosition, OutNormal, RefParticle, RefNormal, BoneMatrices[BoneMap[BoneIndices[ 6]]], BoneWeights[ 6]);  // Intentional fall through
-		case  6: AddInfluence(OutPosition, OutNormal, RefParticle, RefNormal, BoneMatrices[BoneMap[BoneIndices[ 5]]], BoneWeights[ 5]);  // Intentional fall through
-		case  5: AddInfluence(OutPosition, OutNormal, RefParticle, RefNormal, BoneMatrices[BoneMap[BoneIndices[ 4]]], BoneWeights[ 4]);  // Intentional fall through
-		case  4: AddInfluence(OutPosition, OutNormal, RefParticle, RefNormal, BoneMatrices[BoneMap[BoneIndices[ 3]]], BoneWeights[ 3]);  // Intentional fall through
-		case  3: AddInfluence(OutPosition, OutNormal, RefParticle, RefNormal, BoneMatrices[BoneMap[BoneIndices[ 2]]], BoneWeights[ 2]);  // Intentional fall through
-		case  2: AddInfluence(OutPosition, OutNormal, RefParticle, RefNormal, BoneMatrices[BoneMap[BoneIndices[ 1]]], BoneWeights[ 1]);  // Intentional fall through
-		case  1: AddInfluence(OutPosition, OutNormal, RefParticle, RefNormal, BoneMatrices[BoneMap[BoneIndices[ 0]]], BoneWeights[ 0]);  // Intentional fall through
-		default: break;
-		}
+			const uint16* const RESTRICT BoneIndices = PhysicalMeshData.BoneData[VertIndex].BoneIndices;
+			const float* const RESTRICT BoneWeights = PhysicalMeshData.BoneData[VertIndex].BoneWeights;
+	
+			// WARNING - HORRIBLE UNROLLED LOOP + JUMP TABLE BELOW
+			// done this way because this is a pretty tight and performance critical loop. essentially
+			// rather than checking each influence we can just jump into this switch and fall through
+			// everything to compose the final skinned data
+			const FVector3f& RefParticle = PhysicalMeshData.Vertices[VertIndex];
+			const FVector3f& RefNormal = PhysicalMeshData.Normals[VertIndex];
 
-		OutPosition = ComponentToLocalSpace.TransformPosition(OutPosition);
-		OutNormal = ComponentToLocalSpace.TransformVector(OutNormal);
-		OutNormal.Normalize();
-	}, NumPoints > MinParallelVertices ? EParallelForFlags::None : EParallelForFlags::ForceSingleThread);
-}
+			FVector3f Position(ForceInitToZero);
+			FVector3f Normal(ForceInitToZero);
+
+			switch (PhysicalMeshData.BoneData[VertIndex].NumInfluences)
+			{
+			case 12: AddInfluence(Position, Normal, RefParticle, RefNormal, BoneMatrices[BoneMap[BoneIndices[11]]], BoneWeights[11]);  // Intentional fall through
+			case 11: AddInfluence(Position, Normal, RefParticle, RefNormal, BoneMatrices[BoneMap[BoneIndices[10]]], BoneWeights[10]);  // Intentional fall through
+			case 10: AddInfluence(Position, Normal, RefParticle, RefNormal, BoneMatrices[BoneMap[BoneIndices[ 9]]], BoneWeights[ 9]);  // Intentional fall through
+			case  9: AddInfluence(Position, Normal, RefParticle, RefNormal, BoneMatrices[BoneMap[BoneIndices[ 8]]], BoneWeights[ 8]);  // Intentional fall through
+			case  8: AddInfluence(Position, Normal, RefParticle, RefNormal, BoneMatrices[BoneMap[BoneIndices[ 7]]], BoneWeights[ 7]);  // Intentional fall through
+			case  7: AddInfluence(Position, Normal, RefParticle, RefNormal, BoneMatrices[BoneMap[BoneIndices[ 6]]], BoneWeights[ 6]);  // Intentional fall through
+			case  6: AddInfluence(Position, Normal, RefParticle, RefNormal, BoneMatrices[BoneMap[BoneIndices[ 5]]], BoneWeights[ 5]);  // Intentional fall through
+			case  5: AddInfluence(Position, Normal, RefParticle, RefNormal, BoneMatrices[BoneMap[BoneIndices[ 4]]], BoneWeights[ 4]);  // Intentional fall through
+			case  4: AddInfluence(Position, Normal, RefParticle, RefNormal, BoneMatrices[BoneMap[BoneIndices[ 3]]], BoneWeights[ 3]);  // Intentional fall through
+			case  3: AddInfluence(Position, Normal, RefParticle, RefNormal, BoneMatrices[BoneMap[BoneIndices[ 2]]], BoneWeights[ 2]);  // Intentional fall through
+			case  2: AddInfluence(Position, Normal, RefParticle, RefNormal, BoneMatrices[BoneMap[BoneIndices[ 1]]], BoneWeights[ 1]);  // Intentional fall through
+			case  1: AddInfluence(Position, Normal, RefParticle, RefNormal, BoneMatrices[BoneMap[BoneIndices[ 0]]], BoneWeights[ 0]);  // Intentional fall through
+			default: break;
+			}
+
+			OutPositions[VertIndex] = ComponentToLocalSpace.TransformPosition(Position);
+			OutNormals[VertIndex] = ComponentToLocalSpace.TransformVector(Normal).GetSafeNormal();
+
+		}, NumPoints > MinParallelVertices ? EParallelForFlags::None : EParallelForFlags::ForceSingleThread);
+	}
 }
 
 void FClothingSimulationMesh::Update(
