Chaos Cloth - Faster version of the cloth skinning ISPC code.
- Re-enable the ISPC bounds calculation that seems now roughly about 3 times faster.
- LWC code compatibility update.
Originally from Jeff.Rous shelf 17168225.

diff --git a/Engine/Plugins/Experimental/ChaosCloth/Source/Chaos/Private/ChaosCloth/ChaosClothingSimulationMesh.cpp b/Engine/Plugins/Experimental/ChaosCloth/Source/Chaos/Private/ChaosCloth/ChaosClothingSimulationMesh.cpp
--- a/Engine/Plugins/Experimental/ChaosCloth/Source/Chaos/Private/ChaosCloth/ChaosClothingSimulationMesh.cpp
+++ b/Engine/Plugins/Experimental/ChaosCloth/Source/Chaos/Private/ChaosCloth/ChaosClothingSimulationMesh.cpp
@@ -252,9 +252,9 @@ void FClothingSimulationMesh::SkinPhysicsMesh(int32 LODIndex, const FVec3& Local
 	const int32* const RESTRICT BoneMap = Asset->UsedBoneIndices.GetData();
 	const FMatrix44f* const RESTRICT BoneMatrices = Context->RefToLocals.GetData();
 		
-	if (bChaos_SkinPhysicsMesh_ISPC_Enabled)
-	{
 #if INTEL_ISPC
+	if (bChaos_SkinPhysicsMesh_ISPC_Enabled && bRealTypeCompatibleWithISPC)
+	{
 		ispc::SkinPhysicsMesh(
 			(ispc::FVector*)OutPositions,
 			(ispc::FVector*)OutNormals,
@@ -265,9 +265,9 @@ void FClothingSimulationMesh::SkinPhysicsMesh(int32 LODIndex, const FVec3& Local
 			(ispc::FMatrix*)BoneMatrices,
 			(ispc::FTransform&)ComponentToLocalSpace,
 			NumPoints);
-#endif
 	}
 	else
+#endif
 	{
 		static const uint32 MinParallelVertices = 500;  // 500 seems to be the lowest threshold still giving gains even on profiled assets that are only using a small number of influences
 
diff --git a/Engine/Plugins/Experimental/ChaosCloth/Source/Chaos/Private/ChaosCloth/ChaosClothingSimulationMesh.ispc b/Engine/Plugins/Experimental/ChaosCloth/Source/Chaos/Private/ChaosCloth/ChaosClothingSimulationMesh.ispc
--- a/Engine/Plugins/Experimental/ChaosCloth/Source/Chaos/Private/ChaosCloth/ChaosClothingSimulationMesh.ispc
+++ b/Engine/Plugins/Experimental/ChaosCloth/Source/Chaos/Private/ChaosCloth/ChaosClothingSimulationMesh.ispc
@@ -16,39 +16,28 @@ struct FClothVertBoneData
 	float BoneWeights[MaxTotalInfluences];
 };
 
-// Helper function since we only need 3x3
-static inline FMatrix LoadMatrixHelper(const uniform FMatrix *varying M)
+static inline void AddInfluence(uniform FVector& OutPosition, uniform FVector& OutNormal, const uniform FVector& RefParticle, const uniform FVector& RefNormal, const uniform FMatrix *uniform BoneMatrix, const uniform float Weight)
 {
-	#pragma ignore warning(perf)
-	const varying float M0 = M->M[0];
-	#pragma ignore warning(perf)
-	const varying float M1 = M->M[1];
-	#pragma ignore warning(perf)
-	const varying float M2 = M->M[2];
-	#pragma ignore warning(perf)
-	const varying float M4 = M->M[4];
-	#pragma ignore warning(perf)
-	const varying float M5 = M->M[5];
-	#pragma ignore warning(perf)
-	const varying float M6 = M->M[6];
-	#pragma ignore warning(perf)
-	const varying float M8 = M->M[8];
-	#pragma ignore warning(perf)
-	const varying float M9 = M->M[9];
-	#pragma ignore warning(perf)
-	const varying float M10 = M->M[10];
-	#pragma ignore warning(perf)
-	const varying float M12 = M->M[12];
-	#pragma ignore warning(perf)
-	const varying float M13 = M->M[13];
-	#pragma ignore warning(perf)
-	const varying float M14 = M->M[14];
-
-	return SetMatrix(
-		SetVector4(M0 , M1 , M2 , 0.f),
-		SetVector4(M4 , M5 , M6 , 0.f),
-		SetVector4(M8 , M9 , M10, 0.f),
-		SetVector4(M12, M13, M14, 1.f));
+	const uniform FVector4 M0M1M2 = *((uniform FVector4 *uniform)&BoneMatrix->M[0]);
+
+	const uniform FVector4 VTempXP = RefParticle.V[0] * M0M1M2;
+	const uniform FVector4 VTempXN = RefNormal.V[0] * M0M1M2;
+
+	const uniform FVector4 M4M5M6 = *((uniform FVector4 *uniform)&BoneMatrix->M[4]);
+
+	const uniform FVector4 VTempYP = RefParticle.V[1] * M4M5M6 + VTempXP;
+	const uniform FVector4 VTempYN = RefNormal.V[1] * M4M5M6 + VTempXN;
+
+	const uniform FVector4 M8M9M10 = *((uniform FVector4 *uniform)&BoneMatrix->M[8]);
+
+	const uniform FVector4 VTempZP = RefParticle.V[2] * M8M9M10 + VTempYP;
+	const uniform FVector4 Norm = RefNormal.V[2] * M8M9M10 + VTempYN;
+
+	const uniform FVector4 VTempWP = *((uniform FVector4 *uniform)&BoneMatrix->M[12]);
+	const uniform FVector4 Pos = VTempZP + VTempWP;
+
+	OutPosition = OutPosition + SetVector(Pos.V[0], Pos.V[1], Pos.V[2]) * Weight;
+	OutNormal = OutNormal + SetVector(Norm.V[0], Norm.V[1], Norm.V[2]) * Weight;
 }
 
 export void SkinPhysicsMesh(uniform FVector OutPositions[],
@@ -61,38 +50,29 @@ export void SkinPhysicsMesh(uniform FVector OutPositions[],
 							const uniform FTransform &ComponentToLocalSpace,
 							const uniform int32 NumPoints)
 {
-	foreach(i = 0 ... NumPoints)
+	for(uniform int i = 0; i < NumPoints; ++i)
 	{
-		const uniform int VertIndex = extract(i, 0);
-		const FVector RefParticle = VectorLoad(&Vertices[VertIndex]);
-		const FVector RefNormal = VectorLoad(&Normals[VertIndex]);
+		const uniform FVector RefParticle = Vertices[i];
+		const uniform FVector RefNormal = Normals[i];
+
+		uniform FVector OutPosition = OutPositions[i];
+		uniform FVector OutNormal = OutNormals[i];
 
-		FVector OutPosition = VectorLoad(&OutPositions[VertIndex]);
-		FVector OutNormal = VectorLoad(&OutNormals[VertIndex]);
+		OutPosition = SetVector(0.0f, 0.0f, 0.0f);
+		OutNormal = SetVector(0.0f, 0.0f, 0.0f);
 
-		#pragma ignore warning(perf)
-		const varying int NumInfluences = BoneData[i].NumInfluences;
+		const uniform int NumInfluences = BoneData[i].NumInfluences;
 
 		for(uniform int j = 0; j < NumInfluences; ++j)
 		{
-			#pragma ignore warning(perf)
-			const varying int BoneMatrixIndex = BoneMap[BoneData[i].BoneIndices[j]];
-			#pragma ignore warning(perf)
-			const varying float Weight = BoneData[i].BoneWeights[j];
-
-			const varying FMatrix BoneMatrix = LoadMatrixHelper(&BoneMatrices[BoneMatrixIndex]);
-			const FVector Pos = MatrixTransformPosition(RefParticle, BoneMatrix);
-			const FVector Norm = MatrixTransformVector(RefNormal, BoneMatrix);
-
-			// AddInfluence
-			OutPosition = OutPosition + Pos * Weight;
-			OutNormal = OutNormal + Norm * Weight;
-		}
+			const uniform int BoneMatrixIndex = BoneMap[BoneData[i].BoneIndices[j]];
+			const uniform float Weight = BoneData[i].BoneWeights[j];
 
-		OutPosition = TransformPosition(ComponentToLocalSpace, OutPosition);
-		OutNormal = VectorNormalize(TransformVector(ComponentToLocalSpace, OutNormal));
+			const uniform FMatrix *uniform BoneMatrix = &BoneMatrices[BoneMatrixIndex];
+			AddInfluence(OutPosition, OutNormal, RefParticle, RefNormal, BoneMatrix, Weight);
+		}
 
-		VectorStore(&OutPositions[VertIndex], OutPosition);
-		VectorStore(&OutNormals[VertIndex], OutNormal);
+		OutPositions[i] = TransformPosition(ComponentToLocalSpace, OutPosition);
+		OutNormals[i] = VectorNormalize(TransformVector(ComponentToLocalSpace, OutNormal));
 	}
 }
diff --git a/Engine/Plugins/Experimental/ChaosCloth/Source/Chaos/Private/ChaosCloth/ChaosClothingSimulationSolver.cpp b/Engine/Plugins/Experimental/ChaosCloth/Source/Chaos/Private/ChaosCloth/ChaosClothingSimulationSolver.cpp
--- a/Engine/Plugins/Experimental/ChaosCloth/Source/Chaos/Private/ChaosCloth/ChaosClothingSimulationSolver.cpp
+++ b/Engine/Plugins/Experimental/ChaosCloth/Source/Chaos/Private/ChaosCloth/ChaosClothingSimulationSolver.cpp
@@ -946,9 +946,9 @@ FBoxSphereBounds FClothingSimulationSolver::CalculateBounds() const
 		// Calculate bounding box
 		FAABB3 BoundingBox = FAABB3::EmptyAABB();
 
-		if (bChaos_CalculateBounds_ISPC_Enabled)
-		{
 #if INTEL_ISPC
+		if (bChaos_CalculateBounds_ISPC_Enabled && bRealTypeCompatibleWithISPC)
+		{
 			ParticlesActiveView.RangeFor(
 				[&BoundingBox](FPBDParticles& Particles, int32 Offset, int32 Range)
 				{
@@ -965,9 +965,9 @@ FBoxSphereBounds FClothingSimulationSolver::CalculateBounds() const
 					TAABB<float, 3> NewAABB(NewMin, NewMax);
 					BoundingBox = NewAABB;
 				});
-#endif
 		}
 		else
+#endif
 		{
 		ParticlesActiveView.SequentialFor(
 			[&BoundingBox](FPBDParticles& Particles, int32 Index)
diff --git a/Engine/Plugins/Experimental/ChaosCloth/Source/Chaos/Public/ChaosCloth/ChaosClothingSimulationMesh.h b/Engine/Plugins/Experimental/ChaosCloth/Source/Chaos/Public/ChaosCloth/ChaosClothingSimulationMesh.h
--- a/Engine/Plugins/Experimental/ChaosCloth/Source/Chaos/Public/ChaosCloth/ChaosClothingSimulationMesh.h
+++ b/Engine/Plugins/Experimental/ChaosCloth/Source/Chaos/Public/ChaosCloth/ChaosClothingSimulationMesh.h
@@ -74,7 +74,7 @@ namespace Chaos
 } // namespace Chaos
 
 // Support ISPC enable/disable in non-shipping builds
-constexpr bool bChaos_SkinPhysicsMesh_ISPC_Enable = false;  // Disabled by default
+constexpr bool bChaos_SkinPhysicsMesh_ISPC_Enable = true;
 #if !INTEL_ISPC
 const bool bChaos_SkinPhysicsMesh_ISPC_Enabled = false;
 #elif UE_BUILD_SHIPPING
diff --git a/Engine/Plugins/Experimental/ChaosCloth/Source/Chaos/Public/ChaosCloth/ChaosClothingSimulationSolver.h b/Engine/Plugins/Experimental/ChaosCloth/Source/Chaos/Public/ChaosCloth/ChaosClothingSimulationSolver.h
--- a/Engine/Plugins/Experimental/ChaosCloth/Source/Chaos/Public/ChaosCloth/ChaosClothingSimulationSolver.h
+++ b/Engine/Plugins/Experimental/ChaosCloth/Source/Chaos/Public/ChaosCloth/ChaosClothingSimulationSolver.h
@@ -232,7 +232,7 @@ namespace Chaos
 } // namespace Chaos
 
 // Support ISPC enable/disable in non-shipping builds
-constexpr bool bChaos_CalculateBounds_ISPC_Enable = false;  // Disabled by default
+constexpr bool bChaos_CalculateBounds_ISPC_Enable = true;
 #if !INTEL_ISPC
 const bool bChaos_PreSimulationTransforms_ISPC_Enabled = false;
 const bool bChaos_CalculateBounds_ISPC_Enabled = false;
diff --git a/Engine/Source/Runtime/Core/Public/Math/Vector.isph b/Engine/Source/Runtime/Core/Public/Math/Vector.isph
--- a/Engine/Source/Runtime/Core/Public/Math/Vector.isph
+++ b/Engine/Source/Runtime/Core/Public/Math/Vector.isph
@@ -33,7 +33,6 @@ struct FVector
 #endif
 };
 
-
 #ifndef EXPLICIT_VECTOR4
 #define EXPLICIT_VECTOR4 0
 #endif
@@ -2677,6 +2676,14 @@ inline FVector VectorNormalize(const varying FVector& Vector)
 	return VectorSelect(bNonZeroMask, Vector * Scale, Vector);
 }
 
+inline uniform FVector VectorNormalize(const uniform FVector& Vector)
+{
+	const uniform float SquareSum = VectorDot(Vector, Vector);
+	const uniform bool bNonZeroMask = SquareSum > SMALL_NUMBER;
+	const uniform float Scale = ReciprocalSqrtAccurate(SquareSum);
+	return VectorSelect(bNonZeroMask, Vector * Scale, Vector);
+}
+
 inline uniform FVector4 VectorTruncate(const uniform FVector4 &X)
 {
 	varying float S0, Result;
