Chaos Cloth - Fix per cloth debug draw bounds that uses world space, not local coordinates.

diff --git a/Engine/Plugins/Experimental/ChaosCloth/Source/Chaos/Private/ChaosCloth/ChaosClothingSimulation.cpp b/Engine/Plugins/Experimental/ChaosCloth/Source/Chaos/Private/ChaosCloth/ChaosClothingSimulation.cpp
--- a/Engine/Plugins/Experimental/ChaosCloth/Source/Chaos/Private/ChaosCloth/ChaosClothingSimulation.cpp
+++ b/Engine/Plugins/Experimental/ChaosCloth/Source/Chaos/Private/ChaosCloth/ChaosClothingSimulation.cpp
@@ -1250,7 +1250,7 @@ void FClothingSimulation::DebugDrawBounds() const
 		}
 
 		const FAABB3 BoundingBox = Cloth->CalculateBoundingBox(Solver.Get());
-		DrawBox(nullptr, BoundingBox, FQuat::Identity, Bounds.Origin, Color);
+		DrawBox(nullptr, BoundingBox, FQuat::Identity, FVector::ZeroVector, Color);  // TODO: Express bounds in local coordinates for LWC
 	}
 }
 
