Chaos Cloth - Re-enable convex collision code on mobile platforms.

diff --git a/Engine/Plugins/Experimental/ChaosCloth/Source/Chaos/Private/ChaosCloth/ChaosClothingSimulationCollider.cpp b/Engine/Plugins/Experimental/ChaosCloth/Source/Chaos/Private/ChaosCloth/ChaosClothingSimulationCollider.cpp
--- a/Engine/Plugins/Experimental/ChaosCloth/Source/Chaos/Private/ChaosCloth/ChaosClothingSimulationCollider.cpp
+++ b/Engine/Plugins/Experimental/ChaosCloth/Source/Chaos/Private/ChaosCloth/ChaosClothingSimulationCollider.cpp
@@ -396,7 +396,7 @@ void FClothingSimulationCollider::ExtractPhysicsAssetCollision(FClothCollisionDa
 	ClothCollisionData.Reset();
 	UsedBoneIndices.Reset();
 
-	if(!Asset)
+	if (!Asset)
 	{
 		return;
 	}
@@ -410,7 +410,9 @@ void FClothingSimulationCollider::ExtractPhysicsAssetCollision(FClothCollisionDa
 		for (const USkeletalBodySetup* BodySetup : PhysAsset->SkeletalBodySetups)
 		{
 			if (!BodySetup)
+			{
 				continue;
+			}
 
 			const int32 MeshBoneIndex = TargetMesh->GetRefSkeleton().FindBoneIndex(BodySetup->BoneName);
 			const int32 MappedBoneIndex = UsedBoneIndices.Add(MeshBoneIndex);
@@ -510,8 +512,7 @@ void FClothingSimulationCollider::ExtractPhysicsAssetCollision(FClothCollisionDa
 				}
 			}
 
-#if !PLATFORM_ANDROID  // TODO(Kriss.Gossart): Compile on Android and fix whatever errors the following code is causing
-#if WITH_CHAOS  // Only the chaos build has access to the ChaosConvex property
+#if WITH_CHAOS  // Only the Chaos build has access to the ChaosConvex property of the FKConvexElem structure
 			// Add convexes
 			for (const FKConvexElem& ConvexElem : AggGeom.ConvexElems)
 			{
@@ -553,8 +554,6 @@ void FClothingSimulationCollider::ExtractPhysicsAssetCollision(FClothCollisionDa
 				ClothCollisionData.Convexes.Add(Convex);
 			}
 #endif  // #if WITH_CHAOS
-#endif  // #if !PLATFORM_ANDROID
-
 		}  // End for PhysAsset->SkeletalBodySetups
 	}  // End if Asset->PhysicsAsset
 }
