Chaos Cloth - Fix the MaxDistance and Backstop constraints when the component has been scaled.

diff --git a/Engine/Plugins/Experimental/ChaosCloth/Source/Chaos/Private/ChaosCloth/ChaosClothConstraints.cpp b/Engine/Plugins/Experimental/ChaosCloth/Source/Chaos/Private/ChaosCloth/ChaosClothConstraints.cpp
--- a/Engine/Plugins/Experimental/ChaosCloth/Source/Chaos/Private/ChaosCloth/ChaosClothConstraints.cpp
+++ b/Engine/Plugins/Experimental/ChaosCloth/Source/Chaos/Private/ChaosCloth/ChaosClothConstraints.cpp
@@ -579,10 +579,11 @@ void FClothConstraints::SetSelfCollisionProperties(FReal SelfCollisionThickness)
 	}
 }
 
-void FClothConstraints::SetBackstopProperties(bool bEnabled)
+void FClothConstraints::SetBackstopProperties(bool bEnabled, FReal BackstopDistancesMultiplier)
 {
 	if (BackstopConstraints)
 	{
 		BackstopConstraints->SetEnabled(bEnabled);
+		BackstopConstraints->SetSphereRadiiMultiplier(BackstopDistancesMultiplier);
 	}
 }
diff --git a/Engine/Plugins/Experimental/ChaosCloth/Source/Chaos/Private/ChaosCloth/ChaosClothingSimulationCloth.cpp b/Engine/Plugins/Experimental/ChaosCloth/Source/Chaos/Private/ChaosCloth/ChaosClothingSimulationCloth.cpp
--- a/Engine/Plugins/Experimental/ChaosCloth/Source/Chaos/Private/ChaosCloth/ChaosClothingSimulationCloth.cpp
+++ b/Engine/Plugins/Experimental/ChaosCloth/Source/Chaos/Private/ChaosCloth/ChaosClothingSimulationCloth.cpp
@@ -278,7 +278,7 @@ void FClothingSimulationCloth::FLODData::Update(FClothingSimulationSolver* Solve
 
 	// Update the animatable constraint parameters
 	FClothConstraints& ClothConstraints = Solver->GetClothConstraints(Offset);
-	ClothConstraints.SetMaximumDistanceProperties((FReal)Cloth->MaxDistancesMultiplier);
+	ClothConstraints.SetMaximumDistanceProperties((FReal)Cloth->MaxDistancesMultiplier * MeshScale);
 	ClothConstraints.SetEdgeProperties(FVec2((FReal)Cloth->EdgeStiffness[0], (FReal)Cloth->EdgeStiffness[1]));
 	ClothConstraints.SetBendingProperties(FVec2((FReal)Cloth->BendingStiffness[0], (FReal)Cloth->BendingStiffness[1]));
 	ClothConstraints.SetAreaProperties(FVec2((FReal)Cloth->AreaStiffness[0], (FReal)Cloth->AreaStiffness[1]));
@@ -291,7 +291,7 @@ void FClothingSimulationCloth::FLODData::Update(FClothingSimulationSolver* Solve
 		FVec2((FReal)Cloth->AnimDriveDamping[0], (FReal)Cloth->AnimDriveDamping[1]));
 	ClothConstraints.SetThinShellVolumeProperties((FReal)Cloth->VolumeStiffness);
 	ClothConstraints.SetVolumeProperties((FReal)Cloth->VolumeStiffness);
-	ClothConstraints.SetBackstopProperties(Cloth->bEnableBackstop);
+	ClothConstraints.SetBackstopProperties(Cloth->bEnableBackstop, MeshScale);
 }
 
 void FClothingSimulationCloth::FLODData::Enable(FClothingSimulationSolver* Solver, bool bEnable) const
diff --git a/Engine/Plugins/Experimental/ChaosCloth/Source/Chaos/Public/ChaosCloth/ChaosClothConstraints.h b/Engine/Plugins/Experimental/ChaosCloth/Source/Chaos/Public/ChaosCloth/ChaosClothConstraints.h
--- a/Engine/Plugins/Experimental/ChaosCloth/Source/Chaos/Public/ChaosCloth/ChaosClothConstraints.h
+++ b/Engine/Plugins/Experimental/ChaosCloth/Source/Chaos/Public/ChaosCloth/ChaosClothConstraints.h
@@ -71,7 +71,9 @@ namespace Chaos
 		void SetMaximumDistanceProperties(FReal MaxDistancesMultiplier);
 		void SetAnimDriveProperties(const FVec2& AnimDriveStiffness, const FVec2& AnimDriveDamping);
 		void SetSelfCollisionProperties(FReal SelfCollisionThickness);
-		void SetBackstopProperties(bool bEnabled);
+		UE_DEPRECATED(5.0, "Use SetBackstopProperties(bool, FReal) instead.")
+		void SetBackstopProperties(bool bEnabled) { SetBackstopProperties(bEnabled, (FReal)1.); }
+		void SetBackstopProperties(bool bEnabled, FReal BackstopDistancesMultiplier);
 		// ---- End of Cloth interface ----
 
 		// ---- Debug functions ----
