Chaos Cloth - Slightly change the initialization fix to avoid re-reading the atomic.

diff --git a/Engine/Plugins/Experimental/ChaosCloth/Source/Chaos/Private/ChaosCloth/ChaosClothingSimulation.cpp b/Engine/Plugins/Experimental/ChaosCloth/Source/Chaos/Private/ChaosCloth/ChaosClothingSimulation.cpp
--- a/Engine/Plugins/Experimental/ChaosCloth/Source/Chaos/Private/ChaosCloth/ChaosClothingSimulation.cpp
+++ b/Engine/Plugins/Experimental/ChaosCloth/Source/Chaos/Private/ChaosCloth/ChaosClothingSimulation.cpp
@@ -513,6 +513,7 @@ void FClothingSimulation::Simulate(IClothingSimulationContext* InContext)
 		}
 
 		const double StartTime = FPlatformTime::Seconds();
+		const float PrevSimulationTime = SimulationTime;  // Copy the atomic to prevent a re-read
 
 #if !(UE_BUILD_SHIPPING || UE_BUILD_TEST)
 		const bool bNeedsReset = (ChaosClothingSimulationConsole::Command && ChaosClothingSimulationConsole::Command->MustReset(ResetCount)) ||
@@ -520,7 +521,7 @@ void FClothingSimulation::Simulate(IClothingSimulationContext* InContext)
 #else
 		const bool bNeedsReset = (Context->TeleportMode == EClothingTeleportMode::TeleportAndReset);
 #endif
-		const bool bNeedsTeleport = (Context->TeleportMode > EClothingTeleportMode::None) || !SimulationTime;
+		const bool bNeedsTeleport = (Context->TeleportMode > EClothingTeleportMode::None) || !PrevSimulationTime;
 		bIsTeleported = bNeedsTeleport;
 
 		// Update Solver animatable parameters
@@ -549,7 +550,6 @@ void FClothingSimulation::Simulate(IClothingSimulationContext* InContext)
 		Solver->Update(Context->DeltaSeconds);
 
 		// Update simulation time in ms (and provide an instant average instead of the value in real-time)
-		const float PrevSimulationTime = SimulationTime;  // Copy the atomic to prevent a re-read
 		const float CurrSimulationTime = (float)((FPlatformTime::Seconds() - StartTime) * 1000.);
 		static const float SimulationTimeDecay = 0.03f; // 0.03 seems to provide a good rate of update for the instant average
 		SimulationTime = PrevSimulationTime ? PrevSimulationTime + (CurrSimulationTime - PrevSimulationTime) * SimulationTimeDecay : CurrSimulationTime;
