Chaos Cloth - Move back the post LWC solver to float for performance reason, since cloth is being simulated in local space.
- Fixes one remaining issue with the cloth data output transforms in LWC.
- Remove using namespace Chaos from global scope to help with unity builds.

diff --git a/Engine/Plugins/Experimental/ChaosCloth/Source/Chaos/Private/ChaosCloth/ChaosClothConstraints.cpp b/Engine/Plugins/Experimental/ChaosCloth/Source/Chaos/Private/ChaosCloth/ChaosClothConstraints.cpp
--- a/Engine/Plugins/Experimental/ChaosCloth/Source/Chaos/Private/ChaosCloth/ChaosClothConstraints.cpp
+++ b/Engine/Plugins/Experimental/ChaosCloth/Source/Chaos/Private/ChaosCloth/ChaosClothConstraints.cpp
@@ -13,7 +13,7 @@
 #include "Chaos/PBDCollisionSpringConstraints.h"
 #include "Chaos/PBDEvolution.h"
 
-using namespace Chaos;
+namespace Chaos {
 
 FClothConstraints::FClothConstraints()
 	: Evolution(nullptr)
@@ -34,10 +34,10 @@ FClothConstraints::~FClothConstraints()
 }
 
 void FClothConstraints::Initialize(
-	FPBDEvolution* InEvolution,
-	const TArray<FVec3>& InAnimationPositions,
-	const TArray<FVec3>& InOldAnimationPositions,
-	const TArray<FVec3>& InAnimationNormals,
+	Softs::FPBDEvolution* InEvolution,
+	const TArray<Softs::FSolverVec3>& InAnimationPositions,
+	const TArray<Softs::FSolverVec3>& InOldAnimationPositions,
+	const TArray<Softs::FSolverVec3>& InAnimationNormals,
 	int32 InParticleOffset,
 	int32 InNumParticles)
 {
@@ -76,8 +76,8 @@ void FClothConstraints::CreateRules()
 		ConstraintRuleOffset = Evolution->AddConstraintRuleRange(NumConstraintRules, false);
 	}
 
-	TFunction<void(const FPBDParticles&, const FReal)>* const ConstraintInits = Evolution->ConstraintInits().GetData() + ConstraintInitOffset;
-	TFunction<void(FPBDParticles&, const FReal)>* const ConstraintRules = Evolution->ConstraintRules().GetData() + ConstraintRuleOffset;
+	TFunction<void(const Softs::FSolverParticles&, const Softs::FSolverReal)>* const ConstraintInits = Evolution->ConstraintInits().GetData() + ConstraintInitOffset;
+	TFunction<void(Softs::FSolverParticles&, const Softs::FSolverReal)>* const ConstraintRules = Evolution->ConstraintRules().GetData() + ConstraintRuleOffset;
 
 	int32 ConstraintInitIndex = 0;
 	int32 ConstraintRuleIndex = 0;
@@ -85,14 +85,14 @@ void FClothConstraints::CreateRules()
 	if (XEdgeConstraints)
 	{
 		ConstraintInits[ConstraintInitIndex++] =
-			[this](const FPBDParticles& /*Particles*/, const FReal Dt)
+			[this](const Softs::FSolverParticles& /*Particles*/, const Softs::FSolverReal Dt)
 			{
 				XEdgeConstraints->Init();
 				XEdgeConstraints->ApplyProperties(Dt, Evolution->GetIterations());
 			};
 
 		ConstraintRules[ConstraintRuleIndex++] = 
-			[this](FPBDParticles& Particles, const FReal Dt)
+			[this](Softs::FSolverParticles& Particles, const Softs::FSolverReal Dt)
 			{
 				XEdgeConstraints->Apply(Particles, Dt);
 			};
@@ -100,12 +100,12 @@ void FClothConstraints::CreateRules()
 	if (EdgeConstraints)
 	{
 		ConstraintInits[ConstraintInitIndex++] =
-			[this](const FPBDParticles& /*Particles*/, const FReal Dt)
+			[this](const Softs::FSolverParticles& /*Particles*/, const Softs::FSolverReal Dt)
 			{
 				EdgeConstraints->ApplyProperties(Dt, Evolution->GetIterations());
 			};
 		ConstraintRules[ConstraintRuleIndex++] =
-			[this](FPBDParticles& Particles, const FReal Dt)
+			[this](Softs::FSolverParticles& Particles, const Softs::FSolverReal Dt)
 			{
 				EdgeConstraints->Apply(Particles, Dt);
 			};
@@ -113,13 +113,13 @@ void FClothConstraints::CreateRules()
 	if (XBendingConstraints)
 	{
 		ConstraintInits[ConstraintInitIndex++] =
-			[this](const FPBDParticles& /*Particles*/, const FReal Dt)
+			[this](const Softs::FSolverParticles& /*Particles*/, const Softs::FSolverReal Dt)
 			{
 				XBendingConstraints->Init();
 				XBendingConstraints->ApplyProperties(Dt, Evolution->GetIterations());
 			};
 		ConstraintRules[ConstraintRuleIndex++] =
-			[this](FPBDParticles& Particles, const FReal Dt)
+			[this](Softs::FSolverParticles& Particles, const Softs::FSolverReal Dt)
 			{
 				XBendingConstraints->Apply(Particles, Dt);
 			};
@@ -127,12 +127,12 @@ void FClothConstraints::CreateRules()
 	if (BendingConstraints)
 	{
 		ConstraintInits[ConstraintInitIndex++] =
-			[this](const FPBDParticles& /*Particles*/, const FReal Dt)
+			[this](const Softs::FSolverParticles& /*Particles*/, const Softs::FSolverReal Dt)
 			{
 				BendingConstraints->ApplyProperties(Dt, Evolution->GetIterations());
 			};
 		ConstraintRules[ConstraintRuleIndex++] =
-			[this](FPBDParticles& Particles, const FReal Dt)
+			[this](Softs::FSolverParticles& Particles, const Softs::FSolverReal Dt)
 			{
 				BendingConstraints->Apply(Particles, Dt);
 			};
@@ -140,7 +140,7 @@ void FClothConstraints::CreateRules()
 	if (BendingElementConstraints)
 	{
 		ConstraintRules[ConstraintRuleIndex++] =
-			[this](FPBDParticles& Particles, const FReal Dt)
+			[this](Softs::FSolverParticles& Particles, const Softs::FSolverReal Dt)
 			{
 				BendingElementConstraints->Apply(Particles, Dt);
 			};
@@ -148,13 +148,13 @@ void FClothConstraints::CreateRules()
 	if (XAreaConstraints)
 	{
 		ConstraintInits[ConstraintInitIndex++] =
-			[this](const FPBDParticles& /*Particles*/, const FReal Dt)
+			[this](const Softs::FSolverParticles& /*Particles*/, const Softs::FSolverReal Dt)
 			{
 				XAreaConstraints->Init();
 				XAreaConstraints->ApplyProperties(Dt, Evolution->GetIterations());
 			};
 		ConstraintRules[ConstraintRuleIndex++] =
-			[this](FPBDParticles& Particles, const FReal Dt)
+			[this](Softs::FSolverParticles& Particles, const Softs::FSolverReal Dt)
 			{
 				XAreaConstraints->Apply(Particles, Dt);
 			};
@@ -162,12 +162,12 @@ void FClothConstraints::CreateRules()
 	if (AreaConstraints)
 	{
 		ConstraintInits[ConstraintInitIndex++] =
-			[this](const FPBDParticles& /*Particles*/, const FReal Dt)
+			[this](const Softs::FSolverParticles& /*Particles*/, const Softs::FSolverReal Dt)
 			{
 				AreaConstraints->ApplyProperties(Dt, Evolution->GetIterations());
 			};
 		ConstraintRules[ConstraintRuleIndex++] =
-			[this](FPBDParticles& Particles, const FReal Dt)
+			[this](Softs::FSolverParticles& Particles, const Softs::FSolverReal Dt)
 			{
 				AreaConstraints->Apply(Particles, Dt);
 			};
@@ -175,7 +175,7 @@ void FClothConstraints::CreateRules()
 	if (ThinShellVolumeConstraints)
 	{
 		ConstraintRules[ConstraintRuleIndex++] =
-			[this](FPBDParticles& Particles, const FReal Dt)
+			[this](Softs::FSolverParticles& Particles, const Softs::FSolverReal Dt)
 			{
 				ThinShellVolumeConstraints->Apply(Particles, Dt);
 			};
@@ -183,7 +183,7 @@ void FClothConstraints::CreateRules()
 	if (VolumeConstraints)
 	{
 		ConstraintRules[ConstraintRuleIndex++] =
-			[this](FPBDParticles& Particles, const FReal Dt)
+			[this](Softs::FSolverParticles& Particles, const Softs::FSolverReal Dt)
 			{
 				VolumeConstraints->Apply(Particles, Dt);
 			};
@@ -191,13 +191,13 @@ void FClothConstraints::CreateRules()
 	if (XLongRangeConstraints)
 	{
 		ConstraintInits[ConstraintInitIndex++] =
-			[this](const FPBDParticles& /*Particles*/, const FReal Dt)
+			[this](const Softs::FSolverParticles& /*Particles*/, const Softs::FSolverReal Dt)
 			{
 				XLongRangeConstraints->Init();
 				XLongRangeConstraints->ApplyProperties(Dt, Evolution->GetIterations());
 			};
 		ConstraintRules[ConstraintRuleIndex++] =
-			[this](FPBDParticles& Particles, const FReal Dt)
+			[this](Softs::FSolverParticles& Particles, const Softs::FSolverReal Dt)
 			{
 				XLongRangeConstraints->Apply(Particles, Dt);
 			};
@@ -205,12 +205,12 @@ void FClothConstraints::CreateRules()
 	if (LongRangeConstraints)
 	{
 		ConstraintInits[ConstraintInitIndex++] =
-			[this](const FPBDParticles& /*Particles*/, const FReal Dt)
+			[this](const Softs::FSolverParticles& /*Particles*/, const Softs::FSolverReal Dt)
 			{
 				LongRangeConstraints->ApplyProperties(Dt, Evolution->GetIterations());
 			};
 		ConstraintRules[ConstraintRuleIndex++] =
-			[this](FPBDParticles& Particles, const FReal Dt)
+			[this](Softs::FSolverParticles& Particles, const Softs::FSolverReal Dt)
 			{
 				LongRangeConstraints->Apply(Particles, Dt);
 			};
@@ -218,7 +218,7 @@ void FClothConstraints::CreateRules()
 	if (MaximumDistanceConstraints)
 	{
 		ConstraintRules[ConstraintRuleIndex++] =
-			[this](FPBDParticles& Particles, const FReal Dt)
+			[this](Softs::FSolverParticles& Particles, const Softs::FSolverReal Dt)
 			{
 				MaximumDistanceConstraints->Apply(Particles, Dt);
 			};
@@ -226,7 +226,7 @@ void FClothConstraints::CreateRules()
 	if (BackstopConstraints)
 	{
 		ConstraintRules[ConstraintRuleIndex++] =
-			[this](FPBDParticles& Particles, const FReal Dt)
+			[this](Softs::FSolverParticles& Particles, const Softs::FSolverReal Dt)
 			{
 				BackstopConstraints->Apply(Particles, Dt);
 			};
@@ -234,13 +234,13 @@ void FClothConstraints::CreateRules()
 	if (AnimDriveConstraints)
 	{
 		ConstraintInits[ConstraintInitIndex++] =
-			[this](const FPBDParticles& /*Particles*/, const FReal Dt)
+			[this](const Softs::FSolverParticles& /*Particles*/, const Softs::FSolverReal Dt)
 			{
 				AnimDriveConstraints->ApplyProperties(Dt, Evolution->GetIterations());
 			};
 
 		ConstraintRules[ConstraintRuleIndex++] =
-			[this](FPBDParticles& Particles, const FReal Dt)
+			[this](Softs::FSolverParticles& Particles, const Softs::FSolverReal Dt)
 			{
 				AnimDriveConstraints->Apply(Particles, Dt);
 			};
@@ -248,7 +248,7 @@ void FClothConstraints::CreateRules()
 	if (ShapeConstraints)
 	{
 		ConstraintRules[ConstraintRuleIndex++] =
-			[this](FPBDParticles& Particles, const FReal Dt)
+			[this](Softs::FSolverParticles& Particles, const Softs::FSolverReal Dt)
 			{
 				ShapeConstraints->Apply(Particles, Dt);
 			};
@@ -256,12 +256,12 @@ void FClothConstraints::CreateRules()
 	if (SelfCollisionConstraints)
 	{
 		ConstraintInits[ConstraintInitIndex++] =
-			[this](const FPBDParticles& Particles, const FReal /*Dt*/)
+			[this](const Softs::FSolverParticles& Particles, const Softs::FSolverReal /*Dt*/)
 			{
 				SelfCollisionConstraints->Init(Particles);
 			};
 		ConstraintRules[ConstraintRuleIndex++] =
-			[this](FPBDParticles& Particles, const FReal Dt)
+			[this](Softs::FSolverParticles& Particles, const Softs::FSolverReal Dt)
 			{
 				SelfCollisionConstraints->Apply(Particles, Dt);
 			};
@@ -276,23 +276,23 @@ void FClothConstraints::SetEdgeConstraints(const TArray<TVec3<int32>>& SurfaceEl
 
 	if (bUseXPBDConstraints)
 	{
-		XEdgeConstraints = MakeShared<FXPBDSpringConstraints>(
+		XEdgeConstraints = MakeShared<Softs::FXPBDSpringConstraints>(
 			Evolution->Particles(),
 			ParticleOffset, NumParticles,
 			SurfaceElements,
 			StiffnessMultipliers,
-			/*InStiffness =*/ FVec2::UnitVector,
+			/*InStiffness =*/ Softs::FSolverVec2::UnitVector,
 			/*bTrimKinematicConstraints =*/ true);
 	}
 	else
 	{
-		EdgeConstraints = MakeShared<FPBDSpringConstraints>(
+		EdgeConstraints = MakeShared<Softs::FPBDSpringConstraints>(
 			Evolution->Particles(),
 			ParticleOffset,
 			NumParticles,
 			SurfaceElements,
 			StiffnessMultipliers,
-			/*InStiffness =*/ FVec2::UnitVector,
+			/*InStiffness =*/ Softs::FSolverVec2::UnitVector,
 			/*bTrimKinematicConstraints =*/ true);
 	}
 	++NumConstraintInits;  // Uses init to update the property tables
@@ -305,34 +305,34 @@ void FClothConstraints::SetBendingConstraints(const TArray<TVec2<int32>>& Edges,
 
 	if (bUseXPBDConstraints)
 	{
-		XBendingConstraints = MakeShared<FXPBDSpringConstraints>(
+		XBendingConstraints = MakeShared<Softs::FXPBDSpringConstraints>(
 			Evolution->Particles(),
 			ParticleOffset, NumParticles,
 			Edges,
 			StiffnessMultipliers,
-			/*InStiffness =*/ FVec2::UnitVector,
+			/*InStiffness =*/ Softs::FSolverVec2::UnitVector,
 			/*bTrimKinematicConstraints =*/ true);
 	}
 	else
 	{
-		BendingConstraints = MakeShared<FPBDSpringConstraints>(
+		BendingConstraints = MakeShared<Softs::FPBDSpringConstraints>(
 			Evolution->Particles(),
 			ParticleOffset,
 			NumParticles,
 			Edges,
 			StiffnessMultipliers,
-			/*InStiffness =*/ FVec2::UnitVector,
+			/*InStiffness =*/ Softs::FSolverVec2::UnitVector,
 			/*bTrimKinematicConstraints =*/ true);
 	}
 	++NumConstraintInits;  // Uses init to update the property tables
 	++NumConstraintRules;
 }
 
-void FClothConstraints::SetBendingConstraints(TArray<TVec4<int32>>&& BendingElements, FReal BendingStiffness)
+void FClothConstraints::SetBendingConstraints(TArray<TVec4<int32>>&& BendingElements, Softs::FSolverReal BendingStiffness)
 {
 	check(Evolution);
 
-	BendingElementConstraints = MakeShared<FPBDBendingConstraints>(Evolution->Particles(), MoveTemp(BendingElements), BendingStiffness);  // TODO: Strip kinematic constraints
+	BendingElementConstraints = MakeShared<Softs::FPBDBendingConstraints>(Evolution->Particles(), MoveTemp(BendingElements), BendingStiffness);  // TODO: Strip kinematic constraints
 	++NumConstraintRules;
 }
 
@@ -342,35 +342,35 @@ void FClothConstraints::SetAreaConstraints(const TArray<TVec3<int32>>& SurfaceEl
 
 	if (bUseXPBDConstraints)
 	{
-		XAreaConstraints = MakeShared<FXPBDAxialSpringConstraints>(
+		XAreaConstraints = MakeShared<Softs::FXPBDAxialSpringConstraints>(
 			Evolution->Particles(),
 			ParticleOffset,
 			NumParticles,
 			SurfaceElements,
 			StiffnessMultipliers,
-			/*InStiffness =*/ FVec2::UnitVector,
+			/*InStiffness =*/ Softs::FSolverVec2::UnitVector,
 			/*bTrimKinematicConstraints =*/ true);
 	}
 	else
 	{
-		AreaConstraints = MakeShared<FPBDAxialSpringConstraints>(
+		AreaConstraints = MakeShared<Softs::FPBDAxialSpringConstraints>(
 			Evolution->Particles(),
 			ParticleOffset,
 			NumParticles,
 			SurfaceElements,
 			StiffnessMultipliers,
-			/*InStiffness =*/ FVec2::UnitVector,
+			/*InStiffness =*/ Softs::FSolverVec2::UnitVector,
 			/*bTrimKinematicConstraints =*/ true);
 	}
 	++NumConstraintInits;  // Uses init to update the property tables
 	++NumConstraintRules;
 }
 
-void FClothConstraints::SetVolumeConstraints(const TArray<TVec2<int32>>& DoubleBendingEdges, FReal VolumeStiffness)
+void FClothConstraints::SetVolumeConstraints(const TArray<TVec2<int32>>& DoubleBendingEdges, Softs::FSolverReal VolumeStiffness)
 {
 	check(Evolution);
 
-	ThinShellVolumeConstraints = MakeShared<FPBDSpringConstraints>(
+	ThinShellVolumeConstraints = MakeShared<Softs::FPBDSpringConstraints>(
 		Evolution->Particles(),
 		ParticleOffset,
 		NumParticles,
@@ -381,12 +381,12 @@ void FClothConstraints::SetVolumeConstraints(const TArray<TVec2<int32>>& DoubleB
 	++NumConstraintRules;
 }
 
-void FClothConstraints::SetVolumeConstraints(TArray<TVec3<int32>>&& SurfaceElements, FReal VolumeStiffness)
+void FClothConstraints::SetVolumeConstraints(TArray<TVec3<int32>>&& SurfaceElements, Softs::FSolverReal VolumeStiffness)
 {
 	check(Evolution);
 	check(VolumeStiffness > 0.f && VolumeStiffness <= 1.f);
 
-	VolumeConstraints = MakeShared<FPBDVolumeConstraint>(Evolution->Particles(), MoveTemp(SurfaceElements), VolumeStiffness);
+	VolumeConstraints = MakeShared<Softs::FPBDVolumeConstraint>(Evolution->Particles(), MoveTemp(SurfaceElements), VolumeStiffness);
 	++NumConstraintRules;
 }
 
@@ -394,32 +394,32 @@ void FClothConstraints::SetLongRangeConstraints(
 	const TArray<TConstArrayView<TTuple<int32, int32, FRealSingle>>>& Tethers,
 	const TConstArrayView<FRealSingle>& TetherStiffnessMultipliers,
 	const TConstArrayView<FRealSingle>& TetherScaleMultipliers,
-	const FVec2& TetherScale, bool bUseXPBDConstraints)
+	const Softs::FSolverVec2& TetherScale, bool bUseXPBDConstraints)
 {
 	check(Evolution);
 
 	if (bUseXPBDConstraints)
 	{
-		XLongRangeConstraints = MakeShared<FXPBDLongRangeConstraints>(
+		XLongRangeConstraints = MakeShared<Softs::FXPBDLongRangeConstraints>(
 			Evolution->Particles(),
 			ParticleOffset,
 			NumParticles,
 			Tethers,
 			TetherStiffnessMultipliers,
 			TetherScaleMultipliers,
-			/*InStiffness =*/ FVec2::UnitVector,
+			/*InStiffness =*/ Softs::FSolverVec2::UnitVector,
 			TetherScale);
 	}
 	else
 	{
-		LongRangeConstraints = MakeShared<FPBDLongRangeConstraints>(
+		LongRangeConstraints = MakeShared<Softs::FPBDLongRangeConstraints>(
 			Evolution->Particles(),
 			ParticleOffset,
 			NumParticles,
 			Tethers,
 			TetherStiffnessMultipliers,
 			TetherScaleMultipliers,
-			/*InStiffness =*/ FVec2::UnitVector,
+			/*InStiffness =*/ Softs::FSolverVec2::UnitVector,
 			TetherScale);
 	}
 	++NumConstraintInits;  // Uses init to update the property tables
@@ -428,7 +428,7 @@ void FClothConstraints::SetLongRangeConstraints(
 
 void FClothConstraints::SetMaximumDistanceConstraints(const TConstArrayView<FRealSingle>& MaxDistances)
 {
-	MaximumDistanceConstraints = MakeShared<FPBDSphericalConstraint>(
+	MaximumDistanceConstraints = MakeShared<Softs::FPBDSphericalConstraint>(
 		ParticleOffset,
 		NumParticles,
 		*AnimationPositions,
@@ -438,7 +438,7 @@ void FClothConstraints::SetMaximumDistanceConstraints(const TConstArrayView<FRea
 
 void FClothConstraints::SetBackstopConstraints(const TConstArrayView<FRealSingle>& BackstopDistances, const TConstArrayView<FRealSingle>& BackstopRadiuses, bool bUseLegacyBackstop)
 {
-	BackstopConstraints = MakeShared<FPBDSphericalBackstopConstraint>(
+	BackstopConstraints = MakeShared<Softs::FPBDSphericalBackstopConstraint>(
 		ParticleOffset,
 		NumParticles,
 		*AnimationPositions,
@@ -451,7 +451,7 @@ void FClothConstraints::SetBackstopConstraints(const TConstArrayView<FRealSingle
 
 void FClothConstraints::SetAnimDriveConstraints(const TConstArrayView<FRealSingle>& AnimDriveStiffnessMultipliers, const TConstArrayView<FRealSingle>& AnimDriveDampingMultipliers)
 {
-	AnimDriveConstraints = MakeShared<FPBDAnimDriveConstraint>(
+	AnimDriveConstraints = MakeShared<Softs::FPBDAnimDriveConstraint>(
 		ParticleOffset,
 		NumParticles,
 		*AnimationPositions,
@@ -462,12 +462,12 @@ void FClothConstraints::SetAnimDriveConstraints(const TConstArrayView<FRealSingl
 	++NumConstraintRules;
 }
 
-void FClothConstraints::SetShapeTargetConstraints(FReal ShapeTargetStiffness)
+void FClothConstraints::SetShapeTargetConstraints(Softs::FSolverReal ShapeTargetStiffness)
 {
 	// TODO: Review this constraint. Currently does nothing more than the anim drive with less controls
 	check(ShapeTargetStiffness > 0.f && ShapeTargetStiffness <= 1.f);
 
-	ShapeConstraints = MakeShared<FPBDShapeConstraints>(
+	ShapeConstraints = MakeShared<Softs::FPBDShapeConstraints>(
 		ParticleOffset,
 		NumParticles,
 		*AnimationPositions,
@@ -476,20 +476,20 @@ void FClothConstraints::SetShapeTargetConstraints(FReal ShapeTargetStiffness)
 	++NumConstraintRules;
 }
 
-void FClothConstraints::SetSelfCollisionConstraints(const TArray<TVec3<int32>>& SurfaceElements, TSet<TVec2<int32>>&& DisabledCollisionElements, FReal SelfCollisionThickness)
+void FClothConstraints::SetSelfCollisionConstraints(const TArray<TVec3<int32>>& SurfaceElements, TSet<TVec2<int32>>&& DisabledCollisionElements, Softs::FSolverReal SelfCollisionThickness)
 {
-	SelfCollisionConstraints = MakeShared<FPBDCollisionSpringConstraints>(
+	SelfCollisionConstraints = MakeShared<Softs::FPBDCollisionSpringConstraints>(
 		ParticleOffset,
 		NumParticles,
 		SurfaceElements,
 		MoveTemp(DisabledCollisionElements),
 		SelfCollisionThickness,
-		/*Stiffness =*/ 1.f);
+		/*Stiffness =*/ (Softs::FSolverReal)1.);
 	++NumConstraintInits;  // Self collision has an init
 	++NumConstraintRules;  // and a rule
 }
 
-void FClothConstraints::SetEdgeProperties(const FVec2& EdgeStiffness)
+void FClothConstraints::SetEdgeProperties(const Softs::FSolverVec2& EdgeStiffness)
 {
 	if (EdgeConstraints)
 	{
@@ -501,7 +501,7 @@ void FClothConstraints::SetEdgeProperties(const FVec2& EdgeStiffness)
 	}
 }
 
-void FClothConstraints::SetBendingProperties(const FVec2& BendingStiffness)
+void FClothConstraints::SetBendingProperties(const Softs::FSolverVec2& BendingStiffness)
 {
 	if (BendingConstraints)
 	{
@@ -513,7 +513,7 @@ void FClothConstraints::SetBendingProperties(const FVec2& BendingStiffness)
 	}
 }
 
-void FClothConstraints::SetAreaProperties(const FVec2& AreaStiffness)
+void FClothConstraints::SetAreaProperties(const Softs::FSolverVec2& AreaStiffness)
 {
 	if (AreaConstraints)
 	{
@@ -525,7 +525,7 @@ void FClothConstraints::SetAreaProperties(const FVec2& AreaStiffness)
 	}
 }
 
-void FClothConstraints::SetThinShellVolumeProperties(FReal VolumeStiffness)
+void FClothConstraints::SetThinShellVolumeProperties(Softs::FSolverReal VolumeStiffness)
 {
 	if (ThinShellVolumeConstraints)
 	{
@@ -533,7 +533,7 @@ void FClothConstraints::SetThinShellVolumeProperties(FReal VolumeStiffness)
 	}
 }
 
-void FClothConstraints::SetVolumeProperties(FReal VolumeStiffness)
+void FClothConstraints::SetVolumeProperties(Softs::FSolverReal VolumeStiffness)
 {
 	if (VolumeConstraints)
 	{
@@ -541,7 +541,7 @@ void FClothConstraints::SetVolumeProperties(FReal VolumeStiffness)
 	}
 }
 
-void FClothConstraints::SetLongRangeAttachmentProperties(const FVec2& TetherStiffness, const FVec2& TetherScale)
+void FClothConstraints::SetLongRangeAttachmentProperties(const Softs::FSolverVec2& TetherStiffness, const Softs::FSolverVec2& TetherScale)
 {
 	if (LongRangeConstraints)
 	{
@@ -555,7 +555,7 @@ void FClothConstraints::SetLongRangeAttachmentProperties(const FVec2& TetherStif
 	}
 }
 
-void FClothConstraints::SetMaximumDistanceProperties(FReal MaxDistancesMultiplier)
+void FClothConstraints::SetMaximumDistanceProperties(Softs::FSolverReal MaxDistancesMultiplier)
 {
 	if (MaximumDistanceConstraints)
 	{
@@ -563,7 +563,7 @@ void FClothConstraints::SetMaximumDistanceProperties(FReal MaxDistancesMultiplie
 	}
 }
 
-void FClothConstraints::SetAnimDriveProperties(const FVec2& AnimDriveStiffness, const FVec2& AnimDriveDamping)
+void FClothConstraints::SetAnimDriveProperties(const Softs::FSolverVec2& AnimDriveStiffness, const Softs::FSolverVec2& AnimDriveDamping)
 {
 	if (AnimDriveConstraints)
 	{
@@ -571,7 +571,7 @@ void FClothConstraints::SetAnimDriveProperties(const FVec2& AnimDriveStiffness,
 	}
 }
 
-void FClothConstraints::SetSelfCollisionProperties(FReal SelfCollisionThickness)
+void FClothConstraints::SetSelfCollisionProperties(Softs::FSolverReal SelfCollisionThickness)
 {
 	if (SelfCollisionConstraints)
 	{
@@ -579,7 +579,7 @@ void FClothConstraints::SetSelfCollisionProperties(FReal SelfCollisionThickness)
 	}
 }
 
-void FClothConstraints::SetBackstopProperties(bool bEnabled, FReal BackstopDistancesMultiplier)
+void FClothConstraints::SetBackstopProperties(bool bEnabled, Softs::FSolverReal BackstopDistancesMultiplier)
 {
 	if (BackstopConstraints)
 	{
@@ -587,3 +587,5 @@ void FClothConstraints::SetBackstopProperties(bool bEnabled, FReal BackstopDista
 		BackstopConstraints->SetSphereRadiiMultiplier(BackstopDistancesMultiplier);
 	}
 }
+
+}  // End namespace Chaos
diff --git a/Engine/Plugins/Experimental/ChaosCloth/Source/Chaos/Private/ChaosCloth/ChaosClothingSimulation.cpp b/Engine/Plugins/Experimental/ChaosCloth/Source/Chaos/Private/ChaosCloth/ChaosClothingSimulation.cpp
--- a/Engine/Plugins/Experimental/ChaosCloth/Source/Chaos/Private/ChaosCloth/ChaosClothingSimulation.cpp
+++ b/Engine/Plugins/Experimental/ChaosCloth/Source/Chaos/Private/ChaosCloth/ChaosClothingSimulation.cpp
@@ -67,20 +67,22 @@
 
 #if INTEL_ISPC && !UE_BUILD_SHIPPING
 static_assert(sizeof(ispc::FVector3f) == sizeof(FVector3f), "sizeof(ispc::FVector3f) != sizeof(FVector3f)");
-static_assert(sizeof(ispc::FVector4f) == sizeof(FQuat4f), "sizeof(ispc::FVector4f) != sizeof(Chaos::FQuat4f))");
+static_assert(sizeof(ispc::FVector) == sizeof(Chaos::FVec3), "sizeof(ispc::FVector) != sizeof(Chaos::FVec3)");
+static_assert(sizeof(ispc::FTransform) == sizeof(Chaos::FRigidTransform3), "sizeof(ispc::FTransform) != sizeof(Chaos::FRigidTransform3)");
 
 bool bChaos_GetSimData_ISPC_Enabled = true;
 FAutoConsoleVariableRef CVarChaosGetSimDataISPCEnabled(TEXT("p.Chaos.GetSimData.ISPC"), bChaos_GetSimData_ISPC_Enabled, TEXT("Whether to use ISPC optimizations when getting simulation data"));
 #endif
 
-using namespace Chaos;
-
 DECLARE_CYCLE_STAT(TEXT("Chaos Cloth Simulate"), STAT_ChaosClothSimulate, STATGROUP_ChaosCloth);
 DECLARE_CYCLE_STAT(TEXT("Chaos Cloth Create Actor"), STAT_ChaosClothCreateActor, STATGROUP_ChaosCloth);
 DECLARE_CYCLE_STAT(TEXT("Chaos Cloth Get Simulation Data"), STAT_ChaosClothGetSimulationData, STATGROUP_ChaosCloth);
 
+namespace Chaos
+{
+
 #if CHAOS_DEBUG_DRAW
-namespace ChaosClothingSimulationCVar
+namespace ClothingSimulationCVar
 {
 	TAutoConsoleVariable<bool> DebugDrawLocalSpace          (TEXT("p.ChaosCloth.DebugDrawLocalSpace"          ), false, TEXT("Whether to debug draw the Chaos Cloth local space"), ECVF_Cheat);
 	TAutoConsoleVariable<bool> DebugDrawBounds              (TEXT("p.ChaosCloth.DebugDrawBounds"              ), false, TEXT("Whether to debug draw the Chaos Cloth bounds"), ECVF_Cheat);
@@ -104,7 +106,7 @@ namespace ChaosClothingSimulationCVar
 #endif  // #if CHAOS_DEBUG_DRAW
 
 #if !(UE_BUILD_SHIPPING || UE_BUILD_TEST)
-namespace ChaosClothingSimulationConsole
+namespace ClothingSimulationConsole
 {
 	class FCommand final
 	{
@@ -288,9 +290,9 @@ FClothingSimulation::FClothingSimulation()
 	DebugClothMaterialVertex = LoadObject<UMaterial>(nullptr, TEXT("/Engine/EditorMaterials/WidgetVertexColorMaterial"), nullptr, LOAD_None, nullptr);  // LOAD_EditorOnly
 #endif  // #if WITH_EDITOR
 #if !(UE_BUILD_SHIPPING || UE_BUILD_TEST)
-	if (!ChaosClothingSimulationConsole::Command)
+	if (!ClothingSimulationConsole::Command)
 	{
-		ChaosClothingSimulationConsole::Command = MakeUnique<ChaosClothingSimulationConsole::FCommand>();
+		ClothingSimulationConsole::Command = MakeUnique<ClothingSimulationConsole::FCommand>();
 	}
 #endif
 }
@@ -505,7 +507,7 @@ void FClothingSimulation::Simulate(IClothingSimulationContext* InContext)
 {
 	TRACE_CPUPROFILER_EVENT_SCOPE(FClothingSimulation_Simulate);
 #if !(UE_BUILD_SHIPPING || UE_BUILD_TEST)
-	if (ChaosClothingSimulationConsole::Command && ChaosClothingSimulationConsole::Command->MustStep(StepCount))
+	if (ClothingSimulationConsole::Command && ClothingSimulationConsole::Command->MustStep(StepCount))
 #endif
 	{
 		SCOPE_CYCLE_COUNTER(STAT_ChaosClothSimulate);
@@ -519,7 +521,7 @@ void FClothingSimulation::Simulate(IClothingSimulationContext* InContext)
 		const float PrevSimulationTime = SimulationTime;  // Copy the atomic to prevent a re-read
 
 #if !(UE_BUILD_SHIPPING || UE_BUILD_TEST)
-		const bool bNeedsReset = (ChaosClothingSimulationConsole::Command && ChaosClothingSimulationConsole::Command->MustReset(ResetCount)) ||
+		const bool bNeedsReset = (ClothingSimulationConsole::Command && ClothingSimulationConsole::Command->MustReset(ResetCount)) ||
 			(Context->TeleportMode == EClothingTeleportMode::TeleportAndReset);
 #else
 		const bool bNeedsReset = (Context->TeleportMode == EClothingTeleportMode::TeleportAndReset);
@@ -550,7 +552,7 @@ void FClothingSimulation::Simulate(IClothingSimulationContext* InContext)
 		}
 
 		// Step the simulation
-		Solver->Update(Context->DeltaSeconds);
+		Solver->Update((Softs::FSolverReal)Context->DeltaSeconds);
 
 		// Update simulation time in ms (and provide an instant average instead of the value in real-time)
 		const float CurrSimulationTime = (float)((FPlatformTime::Seconds() - StartTime) * 1000.);
@@ -566,22 +568,22 @@ void FClothingSimulation::Simulate(IClothingSimulationContext* InContext)
 
 	// Debug draw
 #if CHAOS_DEBUG_DRAW
-	if (ChaosClothingSimulationCVar::DebugDrawLocalSpace          .GetValueOnAnyThread()) { DebugDrawLocalSpace          (); }
-	if (ChaosClothingSimulationCVar::DebugDrawBounds              .GetValueOnAnyThread()) { DebugDrawBounds              (); }
-	if (ChaosClothingSimulationCVar::DebugDrawGravity             .GetValueOnAnyThread()) { DebugDrawGravity             (); }
-	if (ChaosClothingSimulationCVar::DebugDrawPhysMeshWired       .GetValueOnAnyThread()) { DebugDrawPhysMeshWired       (); }
-	if (ChaosClothingSimulationCVar::DebugDrawAnimMeshWired       .GetValueOnAnyThread()) { DebugDrawAnimMeshWired       (); }
-	if (ChaosClothingSimulationCVar::DebugDrawAnimNormals         .GetValueOnAnyThread()) { DebugDrawAnimNormals         (); }
-	if (ChaosClothingSimulationCVar::DebugDrawPointNormals        .GetValueOnAnyThread()) { DebugDrawPointNormals        (); }
-	if (ChaosClothingSimulationCVar::DebugDrawCollision           .GetValueOnAnyThread()) { DebugDrawCollision           (); }
-	if (ChaosClothingSimulationCVar::DebugDrawBackstops           .GetValueOnAnyThread()) { DebugDrawBackstops           (); }
-	if (ChaosClothingSimulationCVar::DebugDrawBackstopDistances   .GetValueOnAnyThread()) { DebugDrawBackstopDistances   (); }
-	if (ChaosClothingSimulationCVar::DebugDrawMaxDistances        .GetValueOnAnyThread()) { DebugDrawMaxDistances        (); }
-	if (ChaosClothingSimulationCVar::DebugDrawAnimDrive           .GetValueOnAnyThread()) { DebugDrawAnimDrive           (); }
-	if (ChaosClothingSimulationCVar::DebugDrawBendingConstraint   .GetValueOnAnyThread()) { DebugDrawBendingConstraint   (); }
-	if (ChaosClothingSimulationCVar::DebugDrawLongRangeConstraint .GetValueOnAnyThread()) { DebugDrawLongRangeConstraint (); }
-	if (ChaosClothingSimulationCVar::DebugDrawWindForces          .GetValueOnAnyThread()) { DebugDrawWindForces          (); }
-	if (ChaosClothingSimulationCVar::DebugDrawSelfCollision       .GetValueOnAnyThread()) { DebugDrawSelfCollision       (); }
+	if (ClothingSimulationCVar::DebugDrawLocalSpace          .GetValueOnAnyThread()) { DebugDrawLocalSpace          (); }
+	if (ClothingSimulationCVar::DebugDrawBounds              .GetValueOnAnyThread()) { DebugDrawBounds              (); }
+	if (ClothingSimulationCVar::DebugDrawGravity             .GetValueOnAnyThread()) { DebugDrawGravity             (); }
+	if (ClothingSimulationCVar::DebugDrawPhysMeshWired       .GetValueOnAnyThread()) { DebugDrawPhysMeshWired       (); }
+	if (ClothingSimulationCVar::DebugDrawAnimMeshWired       .GetValueOnAnyThread()) { DebugDrawAnimMeshWired       (); }
+	if (ClothingSimulationCVar::DebugDrawAnimNormals         .GetValueOnAnyThread()) { DebugDrawAnimNormals         (); }
+	if (ClothingSimulationCVar::DebugDrawPointNormals        .GetValueOnAnyThread()) { DebugDrawPointNormals        (); }
+	if (ClothingSimulationCVar::DebugDrawCollision           .GetValueOnAnyThread()) { DebugDrawCollision           (); }
+	if (ClothingSimulationCVar::DebugDrawBackstops           .GetValueOnAnyThread()) { DebugDrawBackstops           (); }
+	if (ClothingSimulationCVar::DebugDrawBackstopDistances   .GetValueOnAnyThread()) { DebugDrawBackstopDistances   (); }
+	if (ClothingSimulationCVar::DebugDrawMaxDistances        .GetValueOnAnyThread()) { DebugDrawMaxDistances        (); }
+	if (ClothingSimulationCVar::DebugDrawAnimDrive           .GetValueOnAnyThread()) { DebugDrawAnimDrive           (); }
+	if (ClothingSimulationCVar::DebugDrawBendingConstraint   .GetValueOnAnyThread()) { DebugDrawBendingConstraint   (); }
+	if (ClothingSimulationCVar::DebugDrawLongRangeConstraint .GetValueOnAnyThread()) { DebugDrawLongRangeConstraint (); }
+	if (ClothingSimulationCVar::DebugDrawWindForces          .GetValueOnAnyThread()) { DebugDrawWindForces          (); }
+	if (ClothingSimulationCVar::DebugDrawSelfCollision       .GetValueOnAnyThread()) { DebugDrawSelfCollision       (); }
 #endif  // #if CHAOS_DEBUG_DRAW
 }
 
@@ -670,34 +672,29 @@ void FClothingSimulation::GetSimulationData(
 		Data.Normals = Cloth->GetParticleNormals(Solver.Get());
 
 		// Transform into the cloth reference simulation space used at the time of simulation
-		if (bRealTypeCompatibleWithISPC && bChaos_GetSimData_ISPC_Enabled)
-		{
 #if INTEL_ISPC
-			const FVector3f Location = FVector3f(LocalSpaceLocation);
-			const FQuat4f Rotation = FQuat4f(ReferenceSpaceTransform.GetRotation());
-			const FVector3f Translation = FVector3f(ReferenceSpaceTransform.GetTranslation());
-			const FVector3f Scale3D = FVector3f(ReferenceSpaceTransform.GetScale3D());
-
+		if (bChaos_GetSimData_ISPC_Enabled)
+		{
 			// ISPC is assuming float input here
-			check(sizeof(FVector3f) == Data.Positions.GetTypeSize());
-			check(sizeof(FVector3f) == Data.Normals.GetTypeSize());
+			check(sizeof(ispc::FVector3f) == Data.Positions.GetTypeSize());
+			check(sizeof(ispc::FVector3f) == Data.Normals.GetTypeSize());
 
 			ispc::GetClothingSimulationData(
 				(ispc::FVector3f*)Data.Positions.GetData(),
 				(ispc::FVector3f*)Data.Normals.GetData(),
-				(ispc::FVector4f&)Rotation,
-				(ispc::FVector3f&)Translation,
-				(ispc::FVector3f&)Scale3D,
-				(ispc::FVector3f&)Location,
+				(ispc::FTransform&)ReferenceSpaceTransform,
+				(ispc::FVector&)LocalSpaceLocation,
 				Data.Positions.Num());
-#endif
 		}
 		else
+#endif
 		{
 			for (int32 Index = 0; Index < Data.Positions.Num(); ++Index)
 			{
-				Data.Positions[Index] = (FVector3f)ReferenceSpaceTransform.InverseTransformPosition((FVector)Data.Positions[Index] + LocalSpaceLocation);  // Move into world space first
-				Data.Normals[Index] = (FVector3f)ReferenceSpaceTransform.InverseTransformVector((FVector)-Data.Normals[Index]);  // Normals are inverted due to how barycentric coordinates are calculated (see GetPointBaryAndDist in ClothingMeshUtils.cpp)
+				using FPositionsType = decltype(Data.Positions)::ElementType;
+				using FNormalsType = decltype(Data.Normals)::ElementType;
+				Data.Positions[Index] = FPositionsType(ReferenceSpaceTransform.InverseTransformPosition(FVec3(Data.Positions[Index]) + LocalSpaceLocation));  // Move into world space first
+				Data.Normals[Index] = FNormalsType(ReferenceSpaceTransform.InverseTransformVector(FVec3(-Data.Normals[Index])));  // Normals are inverted due to how barycentric coordinates are calculated (see GetPointBaryAndDist in ClothingMeshUtils.cpp)
 			}
 		}
 
@@ -903,28 +900,28 @@ void FClothingSimulation::DebugDrawPhysMeshShaded(FPrimitiveDrawInterface* PDI)
 		}
 
 		const TConstArrayView<TVec3<int32>> Elements = Cloth->GetTriangleMesh(Solver.Get()).GetElements();
-		const TConstArrayView<FVec3> Positions = Cloth->GetParticlePositions(Solver.Get());
-		const TConstArrayView<FReal> InvMasses = Cloth->GetParticleInvMasses(Solver.Get());
+		const TConstArrayView<Softs::FSolverVec3> Positions = Cloth->GetParticlePositions(Solver.Get());
+		const TConstArrayView<Softs::FSolverReal> InvMasses = Cloth->GetParticleInvMasses(Solver.Get());
 		check(InvMasses.Num() == Positions.Num());
 
 		for (int32 ElementIndex = 0; ElementIndex < Elements.Num(); ++ElementIndex, VertexIndex += 3)
 		{
 			const TVec3<int32>& Element = Elements[ElementIndex];
 
-			const FVector Pos0 = Positions[Element.X - Offset]; // TODO: Triangle Mesh shouldn't really be solver dependent (ie not use an offset)
-			const FVector Pos1 = Positions[Element.Y - Offset];
-			const FVector Pos2 = Positions[Element.Z - Offset];
+			const FVector3f Pos0(Positions[Element.X - Offset]); // TODO: Triangle Mesh shouldn't really be solver dependent (ie not use an offset)
+			const FVector3f Pos1(Positions[Element.Y - Offset]);
+			const FVector3f Pos2(Positions[Element.Z - Offset]);
 
-			const FVector3f Normal = (FVector3f)FVector::CrossProduct(Pos2 - Pos0, Pos1 - Pos0).GetSafeNormal();
-			const FVector3f Tangent = (FVector3f)((Pos1 + Pos2) * 0.5f - Pos0).GetSafeNormal();
+			const FVector3f Normal = FVector3f::CrossProduct(Pos2 - Pos0, Pos1 - Pos0).GetSafeNormal();
+			const FVector3f Tangent = ((Pos1 + Pos2) * 0.5f - Pos0).GetSafeNormal();
 
-			const bool bIsKinematic0 = (InvMasses[Element.X - Offset] == 0.f);
-			const bool bIsKinematic1 = (InvMasses[Element.Y - Offset] == 0.f);
-			const bool bIsKinematic2 = (InvMasses[Element.Z - Offset] == 0.f);
+			const bool bIsKinematic0 = (InvMasses[Element.X - Offset] == (Softs::FSolverReal)0.);
+			const bool bIsKinematic1 = (InvMasses[Element.Y - Offset] == (Softs::FSolverReal)0.);
+			const bool bIsKinematic2 = (InvMasses[Element.Z - Offset] == (Softs::FSolverReal)0.);
 
-			MeshBuilder.AddVertex(FDynamicMeshVertex((FVector3f)Pos0, (FVector3f)Tangent, (FVector3f)Normal, FVector2f(0.f, 0.f), bIsKinematic0 ? FColor::Purple : FColor::White));
-			MeshBuilder.AddVertex(FDynamicMeshVertex((FVector3f)Pos1, (FVector3f)Tangent, (FVector3f)Normal, FVector2f(0.f, 1.f), bIsKinematic1 ? FColor::Purple : FColor::White));
-			MeshBuilder.AddVertex(FDynamicMeshVertex((FVector3f)Pos2, (FVector3f)Tangent, (FVector3f)Normal, FVector2f(1.f, 1.f), bIsKinematic2 ? FColor::Purple : FColor::White));
+			MeshBuilder.AddVertex(FDynamicMeshVertex(Pos0, Tangent, Normal, FVector2f(0.f, 0.f), bIsKinematic0 ? FColor::Purple : FColor::White));
+			MeshBuilder.AddVertex(FDynamicMeshVertex(Pos1, Tangent, Normal, FVector2f(0.f, 1.f), bIsKinematic1 ? FColor::Purple : FColor::White));
+			MeshBuilder.AddVertex(FDynamicMeshVertex(Pos2, Tangent, Normal, FVector2f(1.f, 1.f), bIsKinematic2 ? FColor::Purple : FColor::White));
 			MeshBuilder.AddTriangle(VertexIndex, VertexIndex + 1, VertexIndex + 2);
 		}
 	}
@@ -961,16 +958,16 @@ void FClothingSimulation::DebugDrawParticleIndices(FCanvas* Canvas, const FScene
 			continue;
 		}
 
-		const TConstArrayView<FVec3> Positions = Cloth->GetParticlePositions(Solver.Get());
-		const TConstArrayView<FReal> InvMasses = Cloth->GetParticleInvMasses(Solver.Get());
+		const TConstArrayView<Softs::FSolverVec3> Positions = Cloth->GetParticlePositions(Solver.Get());
+		const TConstArrayView<Softs::FSolverReal> InvMasses = Cloth->GetParticleInvMasses(Solver.Get());
 		check(InvMasses.Num() == Positions.Num());
 
 		for (int32 Index = 0; Index < Positions.Num(); ++Index)
 		{
-			const FVector Position = LocalSpaceLocation + Positions[Index];
+			const FVector Position = LocalSpaceLocation + FVector(Positions[Index]);
 
 			const FText Text = FText::AsNumber(Offset + Index);
-			DrawText(Canvas, SceneView, Position, Text, InvMasses[Index] == 0.f ? KinematicColor : DynamicColor);
+			DrawText(Canvas, SceneView, Position, Text, InvMasses[Index] == (Softs::FSolverReal)0. ? KinematicColor : DynamicColor);
 		}
 	}
 }
@@ -991,16 +988,19 @@ void FClothingSimulation::DebugDrawElementIndices(FCanvas* Canvas, const FSceneV
 		}
 
 		const TArray<TVec3<int32>>& Elements = Cloth->GetTriangleMesh(Solver.Get()).GetElements();
-		const TConstArrayView<FVec3> Positions = Cloth->GetParticlePositions(Solver.Get());
-		const TConstArrayView<FReal> InvMasses = Cloth->GetParticleInvMasses(Solver.Get());
+		const TConstArrayView<Softs::FSolverVec3> Positions = Cloth->GetParticlePositions(Solver.Get());
+		const TConstArrayView<Softs::FSolverReal> InvMasses = Cloth->GetParticleInvMasses(Solver.Get());
 		check(InvMasses.Num() == Positions.Num());
 
 		for (int32 Index = 0; Index < Elements.Num(); ++Index)
 		{
 			const TVec3<int32>& Element = Elements[Index];
-			const FVector Position = LocalSpaceLocation + (Positions[Element[0]] + Positions[Element[1]] + Positions[Element[2]]) / 3.f;
+			const FVector Position = LocalSpaceLocation + FVector(Positions[Element[0]] + Positions[Element[1]] + Positions[Element[2]]) / (FReal)3.;
 
-			const FLinearColor& Color = (InvMasses[Element[0]] == 0.f && InvMasses[Element[1]] == 0.f && InvMasses[Element[2]] == 0.f) ? KinematicColor : DynamicColor;
+			const bool bIsKinematic0 = (InvMasses[Element.X - Offset] == (Softs::FSolverReal)0.);
+			const bool bIsKinematic1 = (InvMasses[Element.Y - Offset] == (Softs::FSolverReal)0.);
+			const bool bIsKinematic2 = (InvMasses[Element.Z - Offset] == (Softs::FSolverReal)0.);
+			const FLinearColor& Color = (bIsKinematic0 && bIsKinematic1 && bIsKinematic2) ? KinematicColor : DynamicColor;
 			const FText Text = FText::AsNumber(Index);
 			DrawText(Canvas, SceneView, Position, Text, Color);
 		}
@@ -1037,18 +1037,18 @@ void FClothingSimulation::DebugDrawMaxDistanceValues(FCanvas* Canvas, const FSce
 			continue;
 		}
 
-		const TConstArrayView<FVec3> Positions = Cloth->GetAnimationPositions(Solver.Get());
-		const TConstArrayView<FReal> InvMasses = Cloth->GetParticleInvMasses(Solver.Get());
+		const TConstArrayView<Softs::FSolverVec3> Positions = Cloth->GetAnimationPositions(Solver.Get());
+		const TConstArrayView<Softs::FSolverReal> InvMasses = Cloth->GetParticleInvMasses(Solver.Get());
 		check(MaxDistances.Num() == Positions.Num());
 		check(MaxDistances.Num() == InvMasses.Num());
 
 		for (int32 Index = 0; Index < MaxDistances.Num(); ++Index)
 		{
 			const FReal MaxDistance = MaxDistances[Index];
-			const FVector Position = LocalSpaceLocation + Positions[Index];
+			const FVector Position = LocalSpaceLocation + FVector(Positions[Index]);
 
 			const FText Text = FText::AsNumber(MaxDistance, &NumberFormattingOptions);
-			DrawText(Canvas, SceneView, Position, Text, InvMasses[Index] == 0.f ? KinematicColor : DynamicColor);
+			DrawText(Canvas, SceneView, Position, Text, InvMasses[Index] == (Softs::FSolverReal)0. ? KinematicColor : DynamicColor);
 		}
 	}
 }
@@ -1219,8 +1219,9 @@ static void DrawConvex(FPrimitiveDrawInterface* PDI, const FConvex& Convex, cons
 					if (ParticleIndex1 != INDEX_NONE)
 					{
 						const FVector X1(Vertices[ParticleIndex1]);
+						const FVector X2(X);
 						const FVector Position1 = Position + Rotation.RotateVector(X1);
-						const FVector Position2 = Position + Rotation.RotateVector((FVector)X);
+						const FVector Position2 = Position + Rotation.RotateVector(X2);
 						DrawLine(PDI, Position1, Position2, Color);
 						break;
 					}
@@ -1283,7 +1284,7 @@ void FClothingSimulation::DebugDrawGravity() const
 		const FAABB3 Bounds = Cloth->CalculateBoundingBox(Solver.Get());
 
 		const FVector Pos0 = Bounds.Center();
-		const FVector Pos1 = Pos0 + Cloth->GetGravity(Solver.Get());
+		const FVector Pos1 = Pos0 + FVector(Cloth->GetGravity(Solver.Get()));
 		DrawLine(nullptr, Pos0, Pos1, FLinearColor::Red);
 	}
 }
@@ -1291,6 +1292,8 @@ void FClothingSimulation::DebugDrawGravity() const
 
 void FClothingSimulation::DebugDrawPhysMeshWired(FPrimitiveDrawInterface* PDI) const
 {
+	check(Solver);
+
 	static const FLinearColor DynamicColor = FColor::White;
 	static const FLinearColor KinematicColor = FColor::Purple;
 
@@ -1305,21 +1308,21 @@ void FClothingSimulation::DebugDrawPhysMeshWired(FPrimitiveDrawInterface* PDI) c
 		}
 
 		const TConstArrayView<TVec3<int32>> Elements = Cloth->GetTriangleMesh(Solver.Get()).GetElements();
-		const TConstArrayView<FVec3> Positions = Cloth->GetParticlePositions(Solver.Get());
-		const TConstArrayView<FReal> InvMasses = Cloth->GetParticleInvMasses(Solver.Get());
+		const TConstArrayView<Softs::FSolverVec3> Positions = Cloth->GetParticlePositions(Solver.Get());
+		const TConstArrayView<Softs::FSolverReal> InvMasses = Cloth->GetParticleInvMasses(Solver.Get());
 		check(InvMasses.Num() == Positions.Num());
 
 		for (int32 ElementIndex = 0; ElementIndex < Elements.Num(); ++ElementIndex)
 		{
 			const auto& Element = Elements[ElementIndex];
 
-			const FVector Pos0 = LocalSpaceLocation + Positions[Element.X - Offset]; // TODO: Triangle Mesh shouldn't really be solver dependent (ie not use an offset)
-			const FVector Pos1 = LocalSpaceLocation + Positions[Element.Y - Offset];
-			const FVector Pos2 = LocalSpaceLocation + Positions[Element.Z - Offset];
+			const FVector Pos0 = LocalSpaceLocation + FVector(Positions[Element.X - Offset]); // TODO: Triangle Mesh shouldn't really be solver dependent (ie not use an offset)
+			const FVector Pos1 = LocalSpaceLocation + FVector(Positions[Element.Y - Offset]);
+			const FVector Pos2 = LocalSpaceLocation + FVector(Positions[Element.Z - Offset]);
 
-			const bool bIsKinematic0 = (InvMasses[Element.X - Offset] == 0.f);
-			const bool bIsKinematic1 = (InvMasses[Element.Y - Offset] == 0.f);
-			const bool bIsKinematic2 = (InvMasses[Element.Z - Offset] == 0.f);
+			const bool bIsKinematic0 = (InvMasses[Element.X - Offset] == (Softs::FSolverReal)0.);
+			const bool bIsKinematic1 = (InvMasses[Element.Y - Offset] == (Softs::FSolverReal)0.);
+			const bool bIsKinematic2 = (InvMasses[Element.Z - Offset] == (Softs::FSolverReal)0.);
 
 			DrawLine(PDI, Pos0, Pos1, bIsKinematic0 && bIsKinematic1 ? KinematicColor : DynamicColor);
 			DrawLine(PDI, Pos1, Pos2, bIsKinematic1 && bIsKinematic2 ? KinematicColor : DynamicColor);
@@ -1330,6 +1333,8 @@ void FClothingSimulation::DebugDrawPhysMeshWired(FPrimitiveDrawInterface* PDI) c
 
 void FClothingSimulation::DebugDrawAnimMeshWired(FPrimitiveDrawInterface* PDI) const
 {
+	check(Solver);
+
 	static const FLinearColor KinematicColor = FColor::Purple;
 
 	const FVec3& LocalSpaceLocation = Solver->GetLocalSpaceLocation();
@@ -1343,17 +1348,15 @@ void FClothingSimulation::DebugDrawAnimMeshWired(FPrimitiveDrawInterface* PDI) c
 		}
 
 		const TConstArrayView<TVec3<int32>> Elements = Cloth->GetTriangleMesh(Solver.Get()).GetElements();
-		const TConstArrayView<FVec3> Positions = Cloth->GetAnimationPositions(Solver.Get());
-		const TConstArrayView<FReal> InvMasses = Cloth->GetParticleInvMasses(Solver.Get());
-		check(InvMasses.Num() == Positions.Num());
+		const TConstArrayView<Softs::FSolverVec3> Positions = Cloth->GetAnimationPositions(Solver.Get());
 
 		for (int32 ElementIndex = 0; ElementIndex < Elements.Num(); ++ElementIndex)
 		{
 			const auto& Element = Elements[ElementIndex];
 
-			const FVector Pos0 = LocalSpaceLocation + Positions[Element.X - Offset]; // TODO: Triangle Mesh shouldn't really be solver dependent (ie not use an offset)
-			const FVector Pos1 = LocalSpaceLocation + Positions[Element.Y - Offset];
-			const FVector Pos2 = LocalSpaceLocation + Positions[Element.Z - Offset];
+			const FVector Pos0 = LocalSpaceLocation + FVector(Positions[Element.X - Offset]); // TODO: Triangle Mesh shouldn't really be solver dependent (ie not use an offset)
+			const FVector Pos1 = LocalSpaceLocation + FVector(Positions[Element.Y - Offset]);
+			const FVector Pos2 = LocalSpaceLocation + FVector(Positions[Element.Z - Offset]);
 
 			DrawLine(PDI, Pos0, Pos1, KinematicColor);
 			DrawLine(PDI, Pos1, Pos2, KinematicColor);
@@ -1365,7 +1368,9 @@ void FClothingSimulation::DebugDrawAnimMeshWired(FPrimitiveDrawInterface* PDI) c
 void FClothingSimulation::DebugDrawAnimNormals(FPrimitiveDrawInterface* PDI) const
 {
 	check(Solver);
+
 	static const FLinearColor KinematicColor = FColor::Magenta;
+	constexpr FReal NormalLength = (FReal)20.;
 
 	const FVec3& LocalSpaceLocation = Solver->GetLocalSpaceLocation();
 
@@ -1377,14 +1382,14 @@ void FClothingSimulation::DebugDrawAnimNormals(FPrimitiveDrawInterface* PDI) con
 			continue;
 		}
 
-		const TConstArrayView<FVec3> Positions = Cloth->GetAnimationPositions(Solver.Get());
-		const TConstArrayView<FVec3> Normals = Cloth->GetAnimationNormals(Solver.Get());
+		const TConstArrayView<Softs::FSolverVec3> Positions = Cloth->GetAnimationPositions(Solver.Get());
+		const TConstArrayView<Softs::FSolverVec3> Normals = Cloth->GetAnimationNormals(Solver.Get());
 		check(Normals.Num() == Positions.Num());
 
 		for (int32 Index = 0; Index < Positions.Num(); ++Index)
 		{
-			const FVector Pos0 = LocalSpaceLocation + Positions[Index];
-			const FVector Pos1 = Pos0 + Normals[Index] * 20.f;
+			const FVector Pos0 = LocalSpaceLocation + FVector(Positions[Index]);
+			const FVector Pos1 = Pos0 + FVector(Normals[Index]) * NormalLength;
 
 			DrawLine(PDI, Pos0, Pos1, KinematicColor);
 		}
@@ -1394,6 +1399,10 @@ void FClothingSimulation::DebugDrawAnimNormals(FPrimitiveDrawInterface* PDI) con
 void FClothingSimulation::DebugDrawPointNormals(FPrimitiveDrawInterface* PDI) const
 {
 	check(Solver);
+	static const FLinearColor DynamicColor = FColor::White;
+	static const FLinearColor KinematicColor = FColor::Purple;
+	constexpr FReal NormalLength = (FReal)20.;
+
 	const FVec3& LocalSpaceLocation = Solver->GetLocalSpaceLocation();
 
 	for (const FClothingSimulationCloth* const Cloth : Solver->GetCloths())
@@ -1404,16 +1413,19 @@ void FClothingSimulation::DebugDrawPointNormals(FPrimitiveDrawInterface* PDI) co
 			continue;
 		}
 
-		const TConstArrayView<FVec3> Positions = Cloth->GetParticlePositions(Solver.Get());
-		const TConstArrayView<FVec3> Normals = Cloth->GetParticleNormals(Solver.Get());
+		const TConstArrayView<Softs::FSolverVec3> Positions = Cloth->GetParticlePositions(Solver.Get());
+		const TConstArrayView<Softs::FSolverVec3> Normals = Cloth->GetParticleNormals(Solver.Get());
+		const TConstArrayView<Softs::FSolverReal> InvMasses = Cloth->GetParticleInvMasses(Solver.Get());
 		check(Normals.Num() == Positions.Num());
+		check(InvMasses.Num() == Positions.Num());
 
 		for (int32 Index = 0; Index < Positions.Num(); ++Index)
 		{
-			const FVector Pos0 = LocalSpaceLocation + Positions[Index];
-			const FVector Pos1 = Pos0 + Normals[Index] * 20.f;
+			const bool bIsKinematic = (InvMasses[Index] == (Softs::FSolverReal)0.);
+			const FVector Pos0 = LocalSpaceLocation + FVector(Positions[Index]);
+			const FVector Pos1 = Pos0 + FVector(Normals[Index]) * NormalLength;
 
-			DrawLine(PDI, Pos0, Pos1, FLinearColor::White);
+			DrawLine(PDI, Pos0, Pos1, bIsKinematic ? KinematicColor : DynamicColor);
 		}
 	}
 }
@@ -1421,6 +1433,10 @@ void FClothingSimulation::DebugDrawPointNormals(FPrimitiveDrawInterface* PDI) co
 void FClothingSimulation::DebugDrawInversedPointNormals(FPrimitiveDrawInterface* PDI) const
 {
 	check(Solver);
+	static const FLinearColor DynamicColor = FColor::White;
+	static const FLinearColor KinematicColor = FColor::Purple;
+	constexpr FReal NormalLength = (FReal)20.;
+
 	const FVec3& LocalSpaceLocation = Solver->GetLocalSpaceLocation();
 
 	for (const FClothingSimulationCloth* const Cloth : Solver->GetCloths())
@@ -1431,15 +1447,19 @@ void FClothingSimulation::DebugDrawInversedPointNormals(FPrimitiveDrawInterface*
 			continue;
 		}
 
-		const TConstArrayView<FVec3> Positions = Cloth->GetParticlePositions(Solver.Get());
-		const TConstArrayView<FVec3> Normals = Cloth->GetParticleNormals(Solver.Get());
+		const TConstArrayView<Softs::FSolverVec3> Positions = Cloth->GetParticlePositions(Solver.Get());
+		const TConstArrayView<Softs::FSolverVec3> Normals = Cloth->GetParticleNormals(Solver.Get());
+		const TConstArrayView<Softs::FSolverReal> InvMasses = Cloth->GetParticleInvMasses(Solver.Get());
+		check(Normals.Num() == Positions.Num());
+		check(InvMasses.Num() == Positions.Num());
 
 		for (int32 Index = 0; Index < Positions.Num(); ++Index)
 		{
-			const FVector Pos0 = LocalSpaceLocation + Positions[Index];
-			const FVector Pos1 = Pos0 - Normals[Index] * 20.f;
+			const bool bIsKinematic = (InvMasses[Index] == (Softs::FSolverReal)0.);
+			const FVector Pos0 = LocalSpaceLocation + FVector(Positions[Index]);
+			const FVector Pos1 = Pos0 - FVector(Normals[Index]) * NormalLength;
 
-			DrawLine(PDI, Pos0, Pos1, FLinearColor::White);
+			DrawLine(PDI, Pos0, Pos1, bIsKinematic ? KinematicColor : DynamicColor);
 		}
 	}
 }
@@ -1462,8 +1482,8 @@ void FClothingSimulation::DebugDrawCollision(FPrimitiveDrawInterface* PDI) const
 			const FVec3& LocalSpaceLocation = Solver->GetLocalSpaceLocation();
 
 			const TConstArrayView<TUniquePtr<FImplicitObject>> CollisionGeometries = Collider->GetCollisionGeometries(Solver.Get(), Cloth, CollisionDataType);
-			const TConstArrayView<FVec3> Translations = Collider->GetCollisionTranslations(Solver.Get(), Cloth, CollisionDataType);
-			const TConstArrayView<FRotation3> Rotations = Collider->GetCollisionRotations(Solver.Get(), Cloth, CollisionDataType);
+			const TConstArrayView<Softs::FSolverVec3> Translations = Collider->GetCollisionTranslations(Solver.Get(), Cloth, CollisionDataType);
+			const TConstArrayView<Softs::FSolverRotation3> Rotations = Collider->GetCollisionRotations(Solver.Get(), Cloth, CollisionDataType);
 			const TConstArrayView<bool> CollisionStatus = Collider->GetCollisionStatus(Solver.Get(), Cloth, CollisionDataType);
 			check(CollisionGeometries.Num() == Translations.Num());
 			check(CollisionGeometries.Num() == Rotations.Num());
@@ -1473,8 +1493,8 @@ void FClothingSimulation::DebugDrawCollision(FPrimitiveDrawInterface* PDI) const
 				if (const FImplicitObject* const Object = CollisionGeometries[Index].Get())
 				{
 					const FLinearColor Color = CollisionStatus[Index] ? CollidedColor : TypeColor;
-					const FVec3 Position = LocalSpaceLocation + Translations[Index];
-					const FRotation3 & Rotation = Rotations[Index];
+					const FVec3 Position = LocalSpaceLocation + FVec3(Translations[Index]);
+					const FRotation3 Rotation(Rotations[Index]);
 
 					switch (Object->GetType())
 					{
@@ -1598,14 +1618,14 @@ void FClothingSimulation::DebugDrawBackstops(FPrimitiveDrawInterface* PDI) const
 		}
 
 		const FClothConstraints& ClothConstraints = Solver->GetClothConstraints(Offset);
-		if (const FPBDSphericalBackstopConstraint* const BackstopConstraint = ClothConstraints.GetBackstopConstraints().Get())
+		if (const Softs::FPBDSphericalBackstopConstraint* const BackstopConstraint = ClothConstraints.GetBackstopConstraints().Get())
 		{
 			const bool bUseLegacyBackstop = BackstopConstraint->UseLegacyBackstop();
 			const TConstArrayView<FRealSingle>& BackstopDistances = Cloth->GetWeightMaps(Solver.Get())[(int32)EChaosWeightMapTarget::BackstopDistance];
 			const TConstArrayView<FRealSingle>& BackstopRadiuses = Cloth->GetWeightMaps(Solver.Get())[(int32)EChaosWeightMapTarget::BackstopRadius];
-			const TConstArrayView<FVec3> AnimationPositions = Cloth->GetAnimationPositions(Solver.Get());
-			const TConstArrayView<FVec3> AnimationNormals = Cloth->GetAnimationNormals(Solver.Get());
-			const TConstArrayView<FVec3> ParticlePositions = Cloth->GetParticlePositions(Solver.Get());
+			const TConstArrayView<Softs::FSolverVec3> AnimationPositions = Cloth->GetAnimationPositions(Solver.Get());
+			const TConstArrayView<Softs::FSolverVec3> AnimationNormals = Cloth->GetAnimationNormals(Solver.Get());
+			const TConstArrayView<Softs::FSolverVec3> ParticlePositions = Cloth->GetParticlePositions(Solver.Get());
 
 			for (int32 Index = 0; Index < AnimationPositions.Num(); ++Index)
 			{
@@ -1616,13 +1636,12 @@ void FClothingSimulation::DebugDrawBackstops(FPrimitiveDrawInterface* PDI) const
 				const FReal BackstopRadius = BackstopRadiuses[Index] * BackstopConstraint->GetSphereRadiiMultiplier();
 				const FReal BackstopDistance = BackstopDistances[Index];
 
-				const FVector AnimationPosition = LocalSpaceLocation + AnimationPositions[Index];
-				const FVector& AnimationNormal = AnimationNormals[Index];
+				const FVector AnimationNormal(AnimationNormals[Index]);
 
 				// Draw a line to show the current distance to the sphere
-				const FVector Pos0 = LocalSpaceLocation + AnimationPositions[Index];
+				const FVector Pos0 = LocalSpaceLocation + FVector(AnimationPositions[Index]);
 				const FVector Pos1 = Pos0 - (bUseLegacyBackstop ? BackstopDistance - BackstopRadius : BackstopDistance) * AnimationNormal;
-				const FVector Pos2 = LocalSpaceLocation + ParticlePositions[Index];
+				const FVector Pos2 = LocalSpaceLocation + FVector(ParticlePositions[Index]);
 				DrawLine(PDI, Pos1, Pos2, ColorLight);
 
 				// Draw the sphere
@@ -1653,13 +1672,13 @@ void FClothingSimulation::DebugDrawBackstopDistances(FPrimitiveDrawInterface* PD
 		}
 
 		const FClothConstraints& ClothConstraints = Solver->GetClothConstraints(Offset);
-		if (const FPBDSphericalBackstopConstraint* const BackstopConstraint = ClothConstraints.GetBackstopConstraints().Get())
+		if (const Softs::FPBDSphericalBackstopConstraint* const BackstopConstraint = ClothConstraints.GetBackstopConstraints().Get())
 		{
 			const bool bUseLegacyBackstop = BackstopConstraint->UseLegacyBackstop();
 			const TConstArrayView<FRealSingle>& BackstopDistances = Cloth->GetWeightMaps(Solver.Get())[(int32)EChaosWeightMapTarget::BackstopDistance];
 			const TConstArrayView<FRealSingle>& BackstopRadiuses = Cloth->GetWeightMaps(Solver.Get())[(int32)EChaosWeightMapTarget::BackstopRadius];
-			const TConstArrayView<FVec3> AnimationPositions = Cloth->GetAnimationPositions(Solver.Get());
-			const TConstArrayView<FVec3> AnimationNormals = Cloth->GetAnimationNormals(Solver.Get());
+			const TConstArrayView<Softs::FSolverVec3> AnimationPositions = Cloth->GetAnimationPositions(Solver.Get());
+			const TConstArrayView<Softs::FSolverVec3> AnimationNormals = Cloth->GetAnimationNormals(Solver.Get());
 
 			for (int32 Index = 0; Index < AnimationPositions.Num(); ++Index)
 			{
@@ -1670,11 +1689,10 @@ void FClothingSimulation::DebugDrawBackstopDistances(FPrimitiveDrawInterface* PD
 				const FReal BackstopRadius = BackstopRadiuses[Index] * BackstopConstraint->GetSphereRadiiMultiplier();
 				const FReal BackstopDistance = BackstopDistances[Index];
 
-				const FVector AnimationPosition = LocalSpaceLocation + AnimationPositions[Index];
-				const FVector& AnimationNormal = AnimationNormals[Index];
+				const FVector AnimationNormal(AnimationNormals[Index]);
 
 				// Draw a line to the sphere boundary
-				const FVector Pos0 = LocalSpaceLocation + AnimationPositions[Index];
+				const FVector Pos0 = LocalSpaceLocation + FVector(AnimationPositions[Index]);
 				const FVector Pos1 = Pos0 - (bUseLegacyBackstop ? BackstopDistance - BackstopRadius : BackstopDistance) * AnimationNormal;
 				DrawLine(PDI, Pos0, Pos1, ColorDark);
 			}
@@ -1703,18 +1721,18 @@ void FClothingSimulation::DebugDrawMaxDistances(FPrimitiveDrawInterface* PDI) co
 			continue;
 		}
 
-		const TConstArrayView<FReal> InvMasses = Cloth->GetParticleInvMasses(Solver.Get());
-		const TConstArrayView<FVec3> Positions = Cloth->GetAnimationPositions(Solver.Get());
-		const TConstArrayView<FVec3> Normals = Cloth->GetAnimationNormals(Solver.Get());
+		const TConstArrayView<Softs::FSolverReal> InvMasses = Cloth->GetParticleInvMasses(Solver.Get());
+		const TConstArrayView<Softs::FSolverVec3> Positions = Cloth->GetAnimationPositions(Solver.Get());
+		const TConstArrayView<Softs::FSolverVec3> Normals = Cloth->GetAnimationNormals(Solver.Get());
 		check(Normals.Num() == Positions.Num());
 		check(MaxDistances.Num() == Positions.Num());
 		check(InvMasses.Num() == Positions.Num());
 
 		for (int32 Index = 0; Index < MaxDistances.Num(); ++Index)
 		{
-			const FReal MaxDistance = MaxDistances[Index];
-			const FVector Position = LocalSpaceLocation + Positions[Index];
-			if (InvMasses[Index] == 0.f)
+			const FReal MaxDistance = (FReal)MaxDistances[Index];
+			const FVector Position = LocalSpaceLocation + FVector(Positions[Index]);
+			if (InvMasses[Index] == (Softs::FSolverReal)0.)
 			{
 #if WITH_EDITOR
 				DrawPoint(PDI, Position, FLinearColor::Red, DebugClothMaterialVertex);
@@ -1724,7 +1742,7 @@ void FClothingSimulation::DebugDrawMaxDistances(FPrimitiveDrawInterface* PDI) co
 			}
 			else
 			{
-				DrawLine(PDI, Position, Position + Normals[Index] * MaxDistance, FLinearColor::White);
+				DrawLine(PDI, Position, Position + FVector(Normals[Index]) * MaxDistance, FLinearColor::White);
 			}
 		}
 	}
@@ -1745,26 +1763,25 @@ void FClothingSimulation::DebugDrawAnimDrive(FPrimitiveDrawInterface* PDI) const
 		}
 
 		const FClothConstraints& ClothConstraints = Solver->GetClothConstraints(Offset);
-		if (const FPBDAnimDriveConstraint* const AnimDriveConstraint = ClothConstraints.GetAnimDriveConstraints().Get())
+		if (const Softs::FPBDAnimDriveConstraint* const AnimDriveConstraint = ClothConstraints.GetAnimDriveConstraints().Get())
 		{
 			const TConstArrayView<FRealSingle>& AnimDriveStiffnessMultipliers = Cloth->GetWeightMaps(Solver.Get())[(int32)EChaosWeightMapTarget::AnimDriveStiffness];
-			const TConstArrayView<FVec3> AnimationPositions = Cloth->GetAnimationPositions(Solver.Get());
-			const TConstArrayView<FVec3> ParticlePositions = Cloth->GetParticlePositions(Solver.Get());
+			const TConstArrayView<Softs::FSolverVec3> AnimationPositions = Cloth->GetAnimationPositions(Solver.Get());
+			const TConstArrayView<Softs::FSolverVec3> ParticlePositions = Cloth->GetParticlePositions(Solver.Get());
+			check(ParticlePositions.Num() == AnimationPositions.Num());
 
 			const FVec2 AnimDriveStiffness = AnimDriveConstraint->GetStiffness();
 			const FRealSingle StiffnessOffset = AnimDriveStiffness[0];
 			const FRealSingle StiffnessRange = AnimDriveStiffness[1] - AnimDriveStiffness[0];
 
-			check(ParticlePositions.Num() == AnimationPositions.Num());
-
 			for (int32 Index = 0; Index < ParticlePositions.Num(); ++Index)
 			{
 				const FRealSingle Stiffness = AnimDriveStiffnessMultipliers.IsValidIndex(Index) ?
 					StiffnessOffset + AnimDriveStiffnessMultipliers[Index] * StiffnessRange :
 					StiffnessOffset;
 
-				const FVector AnimationPosition = LocalSpaceLocation + AnimationPositions[Index];
-				const FVector ParticlePosition = LocalSpaceLocation + ParticlePositions[Index];
+				const FVector AnimationPosition = LocalSpaceLocation + FVector(AnimationPositions[Index]);
+				const FVector ParticlePosition = LocalSpaceLocation + FVector(ParticlePositions[Index]);
 				DrawLine(PDI, AnimationPosition, ParticlePosition, FLinearColor(FColor::Cyan) * Stiffness);
 			}
 		}
@@ -1786,16 +1803,16 @@ void FClothingSimulation::DebugDrawBendingConstraint(FPrimitiveDrawInterface* PD
 		// Draw constraints
 		const FClothConstraints& ClothConstraints = Solver->GetClothConstraints(Offset);
 
-		const TConstArrayView<FVec3> Positions = Cloth->GetParticlePositions(Solver.Get());
+		const TConstArrayView<Softs::FSolverVec3> Positions = Cloth->GetParticlePositions(Solver.Get());
 
-		if (const FPBDSpringConstraints* const BendingConstraints = ClothConstraints.GetBendingConstraints().Get())
+		if (const Softs::FPBDSpringConstraints* const BendingConstraints = ClothConstraints.GetBendingConstraints().Get())
 		{
 			const TArray<TVec2<int32>>& Constraints = BendingConstraints->GetConstraints();
 			for (const TVec2<int32>& Constraint : Constraints)
 			{
 				// Draw line
-				const FVec3 Pos0 = Positions[Constraint[0]] + LocalSpaceLocation;
-				const FVec3 Pos1 = Positions[Constraint[1]] + LocalSpaceLocation;
+				const FVec3 Pos0 = FVec3(Positions[Constraint[0]]) + LocalSpaceLocation;
+				const FVec3 Pos1 = FVec3(Positions[Constraint[1]]) + LocalSpaceLocation;
 
 				DrawLine(PDI, Pos0, Pos1, FLinearColor::Black);
 			}
@@ -1839,28 +1856,28 @@ void FClothingSimulation::DebugDrawLongRangeConstraint(FPrimitiveDrawInterface*
 
 		// Draw constraints
 		const FClothConstraints& ClothConstraints = Solver->GetClothConstraints(Offset);
-		const TConstArrayView<FVec3> Positions = Cloth->GetParticlePositions(Solver.Get());
+		const TConstArrayView<Softs::FSolverVec3> Positions = Cloth->GetParticlePositions(Solver.Get());
 
-		if (const FPBDLongRangeConstraints* const LongRangeConstraints = ClothConstraints.GetLongRangeConstraints().Get())
+		if (const Softs::FPBDLongRangeConstraints* const LongRangeConstraints = ClothConstraints.GetLongRangeConstraints().Get())
 		{
-			const TArray<TConstArrayView<FPBDLongRangeConstraints::FTether>>& Tethers = LongRangeConstraints->GetTethers();
+			const TArray<TConstArrayView<Softs::FPBDLongRangeConstraints::FTether>>& Tethers = LongRangeConstraints->GetTethers();
 
 			for (int32 BatchIndex = 0; BatchIndex < Tethers.Num(); ++BatchIndex)
 			{
 				const FLinearColor Color = PseudoRandomColor(ColorOffset + BatchIndex);
 				const FLinearColor DarkenedColor = Darken(Color);
 
-				const TConstArrayView<FPBDLongRangeConstraints::FTether>& TetherBatch = Tethers[BatchIndex];
+				const TConstArrayView<Softs::FPBDLongRangeConstraints::FTether>& TetherBatch = Tethers[BatchIndex];
 
 				// Draw tethers
-				for (const FPBDLongRangeConstraints::FTether& Tether : TetherBatch)
+				for (const Softs::FPBDLongRangeConstraints::FTether& Tether : TetherBatch)
 				{
 					const int32 KinematicIndex = LongRangeConstraints->GetStartIndex(Tether);
 					const int32 DynamicIndex = LongRangeConstraints->GetEndIndex(Tether);
 					const FReal TargetLength = LongRangeConstraints->GetTargetLength(Tether);
 
-					const FVec3 Pos0 = Positions[KinematicIndex] + LocalSpaceLocation;
-					const FVec3 Pos1 = Positions[DynamicIndex] + LocalSpaceLocation;
+					const FVec3 Pos0 = FVec3(Positions[KinematicIndex]) + LocalSpaceLocation;
+					const FVec3 Pos1 = FVec3(Positions[DynamicIndex]) + LocalSpaceLocation;
 
 					DrawLine(PDI, Pos0, Pos1, Color);
 #if WITH_EDITOR
@@ -1888,6 +1905,7 @@ void FClothingSimulation::DebugDrawWindForces(FPrimitiveDrawInterface* PDI) cons
 {
 	check(Solver);
 
+	constexpr FReal ForceLength = (FReal)10.;
 	const FVec3& LocalSpaceLocation = Solver->GetLocalSpaceLocation();
 
 	for (const FClothingSimulationCloth* const Cloth : Solver->GetCloths())
@@ -1898,21 +1916,20 @@ void FClothingSimulation::DebugDrawWindForces(FPrimitiveDrawInterface* PDI) cons
 			continue;
 		}
 
-		const FVelocityField& VelocityField = Solver->GetWindVelocityField(Cloth->GetGroupId());
+		const Softs::FVelocityField& VelocityField = Solver->GetWindVelocityField(Cloth->GetGroupId());
 
 		const TConstArrayView<TVec3<int32>>& Elements = VelocityField.GetElements();
-		const TConstArrayView<FVec3> Forces = VelocityField.GetForces();
-
-		const TConstArrayView<FVec3> Positions = Cloth->GetParticlePositions(Solver.Get());
+		const TConstArrayView<Softs::FSolverVec3> Forces = VelocityField.GetForces();
+		const TConstArrayView<Softs::FSolverVec3> Positions = Cloth->GetParticlePositions(Solver.Get());
 
 		for (int32 ElementIndex = 0; ElementIndex < Elements.Num(); ++ElementIndex)
 		{
 			const TVec3<int32>& Element = Elements[ElementIndex];
 			const FVec3 Position = LocalSpaceLocation + (
-				Positions[Element.X - Offset] +
-				Positions[Element.Y - Offset] +
-				Positions[Element.Z - Offset]) / 3.f;
-			const FVec3& Force = Forces[ElementIndex] * 10.f;
+				FVec3(Positions[Element.X - Offset]) +
+				FVec3(Positions[Element.Y - Offset]) +
+				FVec3(Positions[Element.Z - Offset])) / (FReal)3.;
+			const FVec3 Force = FVec3(Forces[ElementIndex]) * ForceLength;
 			DrawLine(PDI, Position, Position + Force, FColor::Green);
 		}
 	}
@@ -1952,26 +1969,25 @@ void FClothingSimulation::DebugDrawSelfCollision(FPrimitiveDrawInterface* PDI) c
 		// Draw constraints
 		const FClothConstraints& ClothConstraints = Solver->GetClothConstraints(Offset);
 
-		if (const FPBDCollisionSpringConstraints* const SelfCollisionConstraints = ClothConstraints.GetSelfCollisionConstraints().Get())
+		if (const Softs::FPBDCollisionSpringConstraints* const SelfCollisionConstraints = ClothConstraints.GetSelfCollisionConstraints().Get())
 		{
-			const TConstArrayView<FVec3> Positions = Cloth->GetParticlePositions(Solver.Get());
-			const TConstArrayView<FVec3> ParticleNormals = Cloth->GetParticleNormals(Solver.Get());
+			const TConstArrayView<Softs::FSolverVec3> Positions = Cloth->GetParticlePositions(Solver.Get());
 			const TArray<TVec4<int32>>& Constraints = SelfCollisionConstraints->GetConstraints();
-			const TArray<FVec3>& Barys = SelfCollisionConstraints->GetBarys();
-			const TArray<FVec3>& Normals = SelfCollisionConstraints->GetNormals();
-			const FReal Thickness = SelfCollisionConstraints->GetThickness();
+			const TArray<Softs::FSolverVec3>& Barys = SelfCollisionConstraints->GetBarys();
+			const TArray<Softs::FSolverVec3>& Normals = SelfCollisionConstraints->GetNormals();
+			const FReal Thickness = (FReal)SelfCollisionConstraints->GetThickness();
 			const FReal Height = Thickness + Thickness;
 
 			for (int32 Index = 0; Index < Constraints.Num(); ++Index)
 			{
 				const TVec4<int32>& Constraint = Constraints[Index];
-				const FVec3& Bary = Barys[Index];
-				const FVec3& Normal = Normals[Index];
+				const FVec3 Bary(Barys[Index]);
+				const FVec3 Normal(Normals[Index]);
 
-				const FVector P = LocalSpaceLocation + Positions[Constraint[0]];
-				const FVector P0 = LocalSpaceLocation + Positions[Constraint[1]];
-				const FVector P1 = LocalSpaceLocation + Positions[Constraint[2]];
-				const FVector P2 = LocalSpaceLocation + Positions[Constraint[3]];
+				const FVector P = LocalSpaceLocation + FVector(Positions[Constraint[0]]);
+				const FVector P0 = LocalSpaceLocation + FVector(Positions[Constraint[1]]);
+				const FVector P1 = LocalSpaceLocation + FVector(Positions[Constraint[2]]);
+				const FVector P2 = LocalSpaceLocation + FVector(Positions[Constraint[3]]);
 
 				const FVector Pos0 = P0 * Bary[0] + P1 * Bary[1] + P2 * Bary[2];
 				const FVector Pos1 = Pos0 + Height * Normal;
@@ -2008,4 +2024,5 @@ void FClothingSimulation::DebugDrawSelfCollision(FPrimitiveDrawInterface* PDI) c
 	}
 }
 
-#endif  // #if WITH_EDITOR || CHAOS_DEBUG_DRAW
+#endif  // #if WITH_EDITOR || CHAOS_DEBUG_DRAW?
+}  // End namespace Chaos
diff --git a/Engine/Plugins/Experimental/ChaosCloth/Source/Chaos/Private/ChaosCloth/ChaosClothingSimulation.ispc b/Engine/Plugins/Experimental/ChaosCloth/Source/Chaos/Private/ChaosCloth/ChaosClothingSimulation.ispc
--- a/Engine/Plugins/Experimental/ChaosCloth/Source/Chaos/Private/ChaosCloth/ChaosClothingSimulation.ispc
+++ b/Engine/Plugins/Experimental/ChaosCloth/Source/Chaos/Private/ChaosCloth/ChaosClothingSimulation.ispc
@@ -8,33 +8,38 @@
 
 export void GetClothingSimulationData(uniform FVector3f Positions[],
 									uniform FVector3f Normals[],
-									const uniform FVector4f &Rotation,
-									const uniform FVector3f &Translation,
-									const uniform FVector3f &Scale3D,
-									const uniform FVector3f &LocalSpaceLocation,
+									const uniform FTransform &ReferenceSpaceTransform,
+									const uniform FVector &LocalSpaceLocation,
 									const uniform int NumPositions)
 {
-	const uniform FVector4f RotationInverse = QuatInverse(Rotation);
-	const uniform FVector3f SafeReciprocal = TransformGetSafeScaleReciprocal(Scale3D, FLOAT_SMALL_NUMBER);
+	const uniform FVector4 RotationInverse = QuatInverse(ReferenceSpaceTransform.Rotation);
+	const uniform FVector Translation = SetVector(ReferenceSpaceTransform.Translation);
+	const uniform FVector SafeReciprocal = TransformGetSafeScaleReciprocal(SetVector(ReferenceSpaceTransform.Scale3D), SMALL_NUMBER);
 
 	foreach(i = 0 ... NumPositions)
 	{
 		const uniform int Index = extract(i, 0);
-		const FVector3f Position = VectorLoad(&Positions[Index]);
-		const FVector3f Normal = VectorLoad(&Normals[Index]);
+		const FVector3f PositionFloat = VectorLoad(&Positions[Index]);
+		const FVector3f NormalFloat = VectorLoad(&Normals[Index]);
 
-		const FVector3f InputPositionVector = Position + LocalSpaceLocation;
-		const FVector3f InputNormalVector = Normal * -FLOAT_ONE;
+		const FVector Position = ConvertVector3fTo3Native(PositionFloat);
+		const FVector InputPositionVector = Position + LocalSpaceLocation;
 
-		FVector3f TranslatedPosition = InputPositionVector - Translation;
+		const FVector3f InputNormalVectorFloat = NormalFloat * -FLOAT_ONE;
+		const FVector InputNormalVector = ConvertVector3fTo3Native(InputNormalVectorFloat);
 
-		const FVector3f VRPosition = VectorQuaternionRotateVector(RotationInverse, TranslatedPosition);
-		const FVector3f VRNormal = VectorQuaternionRotateVector(RotationInverse, InputNormalVector);
+		const FVector TranslatedPosition = InputPositionVector - Translation;
 
-		const FVector3f OutPosition = VRPosition * SafeReciprocal;
-		const FVector3f OutNormal = VRNormal * SafeReciprocal;
+		const FVector VRPosition = VectorQuaternionRotateVector(RotationInverse, TranslatedPosition);
+		const FVector VRNormal = VectorQuaternionRotateVector(RotationInverse, InputNormalVector);
 
-		VectorStore(&Positions[Index], OutPosition);
-		VectorStore(&Normals[Index], OutNormal);
+		const FVector OutPosition = VRPosition * SafeReciprocal;
+		const FVector OutNormal = VRNormal * SafeReciprocal;
+
+		const FVector3f OutPositionFloat = ConvertVector3NativeTo3f(OutPosition);
+		const FVector3f OutNormalFloat = ConvertVector3NativeTo3f(OutNormal);
+
+		VectorStore(&Positions[Index], OutPositionFloat);
+		VectorStore(&Normals[Index], OutNormalFloat);
 	}
 }
diff --git a/Engine/Plugins/Experimental/ChaosCloth/Source/Chaos/Private/ChaosCloth/ChaosClothingSimulationCloth.cpp b/Engine/Plugins/Experimental/ChaosCloth/Source/Chaos/Private/ChaosCloth/ChaosClothingSimulationCloth.cpp
--- a/Engine/Plugins/Experimental/ChaosCloth/Source/Chaos/Private/ChaosCloth/ChaosClothingSimulationCloth.cpp
+++ b/Engine/Plugins/Experimental/ChaosCloth/Source/Chaos/Private/ChaosCloth/ChaosClothingSimulationCloth.cpp
@@ -9,9 +9,10 @@
 #include "HAL/IConsoleManager.h"
 #include "ClothingSimulation.h"
 
-using namespace Chaos;
+namespace Chaos
+{
 
-namespace ChaosClothingSimulationClothConsoleVariables
+namespace ClothingSimulationClothConsoleVariables
 {
 	TAutoConsoleVariable<bool> CVarLegacyDisablesAccurateWind(TEXT("p.ChaosCloth.LegacyDisablesAccurateWind"), true, TEXT("Whether using the Legacy wind model switches off the accurate wind model, or adds up to it"));
 }
@@ -88,16 +89,16 @@ void FClothingSimulationCloth::FLODData::Add(FClothingSimulationSolver* Solver,
 	default:
 		check(false);
 	case EMassMode::UniformMass:
-		Solver->SetParticleMassUniform(Offset, (FReal)Cloth->MassValue, (FReal)Cloth->MinPerParticleMass, TriangleMesh, KinematicPredicate);
+		Solver->SetParticleMassUniform(Offset, Cloth->MassValue, Cloth->MinPerParticleMass, TriangleMesh, KinematicPredicate);
 		break;
 	case EMassMode::TotalMass:
-		Solver->SetParticleMassFromTotalMass(Offset, (FReal)Cloth->MassValue, (FReal)Cloth->MinPerParticleMass, TriangleMesh, KinematicPredicate);
+		Solver->SetParticleMassFromTotalMass(Offset, Cloth->MassValue, Cloth->MinPerParticleMass, TriangleMesh, KinematicPredicate);
 		break;
 	case EMassMode::Density:
-		Solver->SetParticleMassFromDensity(Offset, (FReal)Cloth->MassValue, (FReal)Cloth->MinPerParticleMass, TriangleMesh, KinematicPredicate);
+		Solver->SetParticleMassFromDensity(Offset, Cloth->MassValue, Cloth->MinPerParticleMass, TriangleMesh, KinematicPredicate);
 		break;
 	}
-	const TConstArrayView<FReal> InvMasses(Solver->GetParticleInvMasses(Offset), NumParticles);
+	const TConstArrayView<Softs::FSolverReal> InvMasses(Solver->GetParticleInvMasses(Offset), NumParticles);
 
 	// Setup solver constraints
 	FClothConstraints& ClothConstraints = Solver->GetClothConstraints(Offset);
@@ -310,17 +311,7 @@ void FClothingSimulationCloth::FLODData::ResetStartPose(FClothingSimulationSolve
 	const int32 Offset = SolverData.FindChecked(Solver).Offset;
 	check(Offset != INDEX_NONE);
 
-	FVec3* const Ps = Solver->GetParticlePs(Offset);
-	FVec3* const Xs = Solver->GetParticleXs(Offset);
-	FVec3* const Vs = Solver->GetParticleVs(Offset);
-	const FVec3* const AnimationPositions = Solver->GetAnimationPositions(Offset);
-	FVec3* const OldAnimationPositions = Solver->GetOldAnimationPositions(Offset);
-
-	for (int32 Index = 0; Index < NumParticles; ++Index)
-	{
-		Ps[Index] = Xs[Index] = OldAnimationPositions[Index] = AnimationPositions[Index];
-		Vs[Index] = FVec3(0.f);
-	}
+	Solver->ResetStartPose(Offset, NumParticles);
 }
 
 void FClothingSimulationCloth::FLODData::UpdateNormals(FClothingSimulationSolver* Solver) const
@@ -333,12 +324,12 @@ void FClothingSimulationCloth::FLODData::UpdateNormals(FClothingSimulationSolver
 
 	check(Offset != INDEX_NONE);
 
-	TConstArrayView<FVec3> Points(Solver->GetParticleXs(Offset) - Offset, Offset + NumParticles);  // TODO: TriangleMesh still uses global array
-	TArray<FVec3> FaceNormals;
+	TConstArrayView<Softs::FSolverVec3> Points(Solver->GetParticleXs(Offset) - Offset, Offset + NumParticles);  // TODO: TriangleMesh still uses global array
+	TArray<Softs::FSolverVec3> FaceNormals;
 	TriangleMesh.GetFaceNormals(FaceNormals, Points, /*ReturnEmptyOnError =*/ false);
 
-	TArrayView<FVec3> Normals(Solver->GetNormals(Offset), NumParticles);
-	TriangleMesh.GetPointNormals(Normals, FaceNormals, /*bUseGlobalArray =*/ false);
+	TArrayView<Softs::FSolverVec3> Normals(Solver->GetNormals(Offset), NumParticles);
+	TriangleMesh.GetPointNormals(Normals, TConstArrayView<Softs::FSolverVec3>(FaceNormals), /*bUseGlobalArray =*/ false);
 }
 
 FClothingSimulationCloth::FClothingSimulationCloth(
@@ -364,8 +355,8 @@ FClothingSimulationCloth::FClothingSimulationCloth(
 	bool bInUseXPBDConstraints,
 	FRealSingle InGravityScale,
 	bool bInIsGravityOverridden,
-	const FVec3& InGravityOverride,
-	const FVec3& InLinearVelocityScale,
+	const TVec3<FRealSingle>& InGravityOverride,
+	const TVec3<FRealSingle>& InLinearVelocityScale,
 	FRealSingle InAngularVelocityScale,
 	FRealSingle InFictitiousAngularScale,
 	const TVec2<FRealSingle>& InDrag,
@@ -601,7 +592,7 @@ int32 FClothingSimulationCloth::GetOffset(const FClothingSimulationSolver* Solve
 	return LODData.IsValidIndex(InLODIndex) ? LODData[InLODIndex].SolverData.FindChecked(Solver).Offset : 0;
 }
 
-FVec3 FClothingSimulationCloth::GetGravity(const FClothingSimulationSolver* Solver) const
+TVec3<FRealSingle> FClothingSimulationCloth::GetGravity(const FClothingSimulationSolver* Solver) const
 {
 	check(Solver);
 	return Solver->IsClothGravityOverrideEnabled() && bIsGravityOverridden ? GravityOverride : Solver->GetGravity() * GravityScale;
@@ -612,16 +603,16 @@ FAABB3 FClothingSimulationCloth::CalculateBoundingBox(const FClothingSimulationS
 	check(Solver);
 
 	// Calculate local space bounding box
-	FAABB3 BoundingBox = FAABB3::EmptyAABB();
+	Softs::FSolverAABB3 BoundingBox = Softs::FSolverAABB3::EmptyAABB();
 
-	const TConstArrayView<FVec3> ParticlePositions = GetParticlePositions(Solver);
-	for (const FVec3& ParticlePosition : ParticlePositions)
+	const TConstArrayView<Softs::FSolverVec3> ParticlePositions = GetParticlePositions(Solver);
+	for (const Softs::FSolverVec3& ParticlePosition : ParticlePositions)
 	{
 		BoundingBox.GrowToInclude(ParticlePosition);
 	}
 
 	// Return world space bounding box
-	return BoundingBox.TransformedAABB(FRigidTransform3(Solver->GetLocalSpaceLocation(), FRotation3::Identity));
+	return FAABB3(BoundingBox).TransformedAABB(FRigidTransform3(Solver->GetLocalSpaceLocation(), FRotation3::Identity));
 }
 
 int32 FClothingSimulationCloth::GetOffset(const FClothingSimulationSolver* Solver) const
@@ -806,13 +797,13 @@ void FClothingSimulationCloth::Update(FClothingSimulationSolver* Solver)
 		// External forces (legacy wind+field)
 		Solver->AddExternalForces(GroupId, bUseLegacyWind);
 
-		if (bUseLegacyWind && ChaosClothingSimulationClothConsoleVariables::CVarLegacyDisablesAccurateWind.GetValueOnAnyThread())
+		if (bUseLegacyWind && ClothingSimulationClothConsoleVariables::CVarLegacyDisablesAccurateWind.GetValueOnAnyThread())
 		{
-			Solver->SetWindProperties(GroupId, FVec2::ZeroVector, FVec2::ZeroVector, (FReal)0.);  // Disable the wind velocity field
+			Solver->SetWindProperties(GroupId, TVec2<FRealSingle>(0.), TVec2<FRealSingle>(0.), (FRealSingle)0.);  // Disable the wind velocity field
 		}
 		else
 		{
-			Solver->SetWindProperties(GroupId, FVec2((FReal)Drag[0], (FReal)Drag[1]), FVec2((FReal)Lift[0], (FReal)Lift[1]), (FReal)AirDensity);  // TODO: Pass AirDensity too
+			Solver->SetWindProperties(GroupId, Drag, Lift, AirDensity);
 		}
 		Solver->SetWindVelocity(GroupId, WindVelocity + Solver->GetWindVelocity());
 
@@ -840,50 +831,52 @@ void FClothingSimulationCloth::PostUpdate(FClothingSimulationSolver* Solver)
 	}
 }
 
-TConstArrayView<FVec3> FClothingSimulationCloth::GetAnimationPositions(const FClothingSimulationSolver* Solver) const
+TConstArrayView<Softs::FSolverVec3> FClothingSimulationCloth::GetAnimationPositions(const FClothingSimulationSolver* Solver) const
 {
 	check(Solver);
 	const int32 LODIndex = LODIndices.FindChecked(Solver);
 	check(GetOffset(Solver, LODIndex) != INDEX_NONE);
-	return TConstArrayView<FVec3>(Solver->GetAnimationPositions(GetOffset(Solver, LODIndex)), GetNumParticles(LODIndex));
+	return TConstArrayView<Softs::FSolverVec3>(Solver->GetAnimationPositions(GetOffset(Solver, LODIndex)), GetNumParticles(LODIndex));
 }
 
-TConstArrayView<FVec3> FClothingSimulationCloth::GetAnimationNormals(const FClothingSimulationSolver* Solver) const
+TConstArrayView<Softs::FSolverVec3> FClothingSimulationCloth::GetAnimationNormals(const FClothingSimulationSolver* Solver) const
 {
 	check(Solver);
 	const int32 LODIndex = LODIndices.FindChecked(Solver);
 	check(GetOffset(Solver, LODIndex) != INDEX_NONE);
-	return TConstArrayView<FVec3>(Solver->GetAnimationNormals(GetOffset(Solver, LODIndex)), GetNumParticles(LODIndex));
+	return TConstArrayView<Softs::FSolverVec3>(Solver->GetAnimationNormals(GetOffset(Solver, LODIndex)), GetNumParticles(LODIndex));
 }
 
-TConstArrayView<FVec3> FClothingSimulationCloth::GetParticlePositions(const FClothingSimulationSolver* Solver) const
+TConstArrayView<Softs::FSolverVec3> FClothingSimulationCloth::GetParticlePositions(const FClothingSimulationSolver* Solver) const
 {
 	check(Solver);
 	const int32 LODIndex = LODIndices.FindChecked(Solver);
 	check(GetOffset(Solver, LODIndex) != INDEX_NONE);
-	return TConstArrayView<FVec3>(Solver->GetParticleXs(GetOffset(Solver, LODIndex)), GetNumParticles(LODIndex));
+	return TConstArrayView<Softs::FSolverVec3>(Solver->GetParticleXs(GetOffset(Solver, LODIndex)), GetNumParticles(LODIndex));
 }
 
-TConstArrayView<FVec3> FClothingSimulationCloth::GetParticleOldPositions(const FClothingSimulationSolver* Solver) const
+TConstArrayView<Softs::FSolverVec3> FClothingSimulationCloth::GetParticleOldPositions(const FClothingSimulationSolver* Solver) const
 {
 	check(Solver);
 	const int32 LODIndex = LODIndices.FindChecked(Solver);
 	check(GetOffset(Solver, LODIndex) != INDEX_NONE);
-	return TConstArrayView<FVec3>(Solver->GetParticlePs(GetOffset(Solver, LODIndex)), GetNumParticles(LODIndex));
+	return TConstArrayView<Softs::FSolverVec3>(Solver->GetParticlePs(GetOffset(Solver, LODIndex)), GetNumParticles(LODIndex));
 }
 
-TConstArrayView<FVec3> FClothingSimulationCloth::GetParticleNormals(const FClothingSimulationSolver* Solver) const
+TConstArrayView<Softs::FSolverVec3> FClothingSimulationCloth::GetParticleNormals(const FClothingSimulationSolver* Solver) const
 {
 	check(Solver);
 	const int32 LODIndex = LODIndices.FindChecked(Solver);
 	check(GetOffset(Solver, LODIndex) != INDEX_NONE);
-	return TConstArrayView<FVec3>(Solver->GetNormals(GetOffset(Solver, LODIndex)), GetNumParticles(LODIndex));
+	return TConstArrayView<Softs::FSolverVec3>(Solver->GetNormals(GetOffset(Solver, LODIndex)), GetNumParticles(LODIndex));
 }
 
-TConstArrayView<FReal> FClothingSimulationCloth::GetParticleInvMasses(const FClothingSimulationSolver* Solver) const
+TConstArrayView<Softs::FSolverReal> FClothingSimulationCloth::GetParticleInvMasses(const FClothingSimulationSolver* Solver) const
 {
 	check(Solver);
 	const int32 LODIndex = LODIndices.FindChecked(Solver);
 	check(GetOffset(Solver, LODIndex) != INDEX_NONE);
-	return TConstArrayView<FReal>(Solver->GetParticleInvMasses(GetOffset(Solver, LODIndex)), GetNumParticles(LODIndex));
+	return TConstArrayView<Softs::FSolverReal>(Solver->GetParticleInvMasses(GetOffset(Solver, LODIndex)), GetNumParticles(LODIndex));
 }
+
+}  // End namespace Chaos
diff --git a/Engine/Plugins/Experimental/ChaosCloth/Source/Chaos/Private/ChaosCloth/ChaosClothingSimulationCollider.cpp b/Engine/Plugins/Experimental/ChaosCloth/Source/Chaos/Private/ChaosCloth/ChaosClothingSimulationCollider.cpp
--- a/Engine/Plugins/Experimental/ChaosCloth/Source/Chaos/Private/ChaosCloth/ChaosClothingSimulationCollider.cpp
+++ b/Engine/Plugins/Experimental/ChaosCloth/Source/Chaos/Private/ChaosCloth/ChaosClothingSimulationCollider.cpp
@@ -28,13 +28,14 @@
 
 DECLARE_CYCLE_STAT(TEXT("Chaos Cloth Update Collider"), STAT_ChaosClothingSimulationColliderUpdate, STATGROUP_ChaosCloth);
 
-namespace ChaosClothingSimulationColliderConsoleVariables
+namespace Chaos
+{
+
+namespace ClothingSimulationColliderConsoleVariables
 {
 	TAutoConsoleVariable<bool> CVarUseOptimizedTaperedCapsule(TEXT("p.ChaosCloth.UseOptimizedTaperedCapsule"), true, TEXT("Use the optimized TaperedCapsule code instead of using a tapered cylinder and two spheres"));
 }
 
-using namespace Chaos;
-
 void FClothingSimulationCollider::FLODData::Add(
 	FClothingSimulationSolver* Solver,
 	FClothingSimulationCloth* Cloth,
@@ -86,7 +87,7 @@ void FClothingSimulationCollider::FLODData::Add(
 	if (NumCapsules)
 	{
 		int32* const BoneIndices = Solver->GetCollisionBoneIndices(CapsuleOffset);
-		FRigidTransform3* const BaseTransforms = Solver->GetCollisionBaseTransforms(CapsuleOffset);
+		Softs::FSolverRigidTransform3* const BaseTransforms = Solver->GetCollisionBaseTransforms(CapsuleOffset);
 
 		for (int32 Index = 0; Index < NumCapsules; ++Index)
 		{
@@ -117,7 +118,7 @@ void FClothingSimulationCollider::FLODData::Add(
 			if (Radius0 <= Radius1) { MinRadius = Radius0; MaxRadius = Radius1; }
 			else { MinRadius = Radius1; MaxRadius = Radius0; }
 
-			BaseTransforms[Index] = FRigidTransform3::Identity;
+			BaseTransforms[Index] = Softs::FSolverRigidTransform3::Identity;
 
 			if (Axis.SizeSquared() < SMALL_NUMBER)
 			{
@@ -133,7 +134,7 @@ void FClothingSimulationCollider::FLODData::Add(
 			}
 			else
 			{
-				if (ChaosClothingSimulationColliderConsoleVariables::CVarUseOptimizedTaperedCapsule.GetValueOnAnyThread())
+				if (ClothingSimulationColliderConsoleVariables::CVarUseOptimizedTaperedCapsule.GetValueOnAnyThread())
 				{
 					Solver->SetCollisionGeometry(CapsuleOffset, Index,
 						MakeUnique<FTaperedCapsule>(P0, P1, Radius0, Radius1));
@@ -161,7 +162,7 @@ void FClothingSimulationCollider::FLODData::Add(
 	if (NumSpheres != 0)
 	{
 		int32* const BoneIndices = Solver->GetCollisionBoneIndices(SphereOffset);
-		FRigidTransform3* const BaseTransforms = Solver->GetCollisionBaseTransforms(SphereOffset);
+		Softs::FSolverRigidTransform3* const BaseTransforms = Solver->GetCollisionBaseTransforms(SphereOffset);
 
 		for (int32 Index = 0, SphereIndex = 0; SphereIndex < ClothCollisionData.Spheres.Num(); ++SphereIndex)
 		{
@@ -177,7 +178,7 @@ void FClothingSimulationCollider::FLODData::Add(
 			BoneIndices[Index] = GetMappedBoneIndex(UsedBoneIndices, Sphere.BoneIndex);
 			UE_LOG(LogChaosCloth, VeryVerbose, TEXT("Found collision sphere on bone index %d."), BoneIndices[Index]);
 
-			BaseTransforms[Index] = FRigidTransform3::Identity;
+			BaseTransforms[Index] = Softs::FSolverRigidTransform3::Identity;
 
 			Solver->SetCollisionGeometry(SphereOffset, Index,
 				MakeUnique<TSphere<FReal, 3>>(
@@ -193,14 +194,14 @@ void FClothingSimulationCollider::FLODData::Add(
 	if (NumConvexes != 0)
 	{
 		int32* const BoneIndices = Solver->GetCollisionBoneIndices(ConvexOffset);
-		FRigidTransform3* const BaseTransforms = Solver->GetCollisionBaseTransforms(ConvexOffset);
+		Softs::FSolverRigidTransform3* const BaseTransforms = Solver->GetCollisionBaseTransforms(ConvexOffset);
 
 		for (uint32 Index = 0; Index < NumConvexes; ++Index)
 		{
 			const FClothCollisionPrim_Convex& Convex = ClothCollisionData.Convexes[Index];
 
 			// Always initialize the collision particle transforms before setting any geometry as otherwise NaNs gets detected during the bounding box updates
-			BaseTransforms[Index] = FRigidTransform3::Identity;
+			BaseTransforms[Index] = Softs::FSolverRigidTransform3::Identity;
 
 			BoneIndices[Index] = GetMappedBoneIndex(UsedBoneIndices, Convex.BoneIndex);
 			UE_LOG(LogChaosCloth, VeryVerbose, TEXT("Found collision convex on bone index %d."), BoneIndices[Index]);
@@ -273,13 +274,13 @@ void FClothingSimulationCollider::FLODData::Add(
 	if (NumBoxes != 0)
 	{
 		int32* const BoneIndices = Solver->GetCollisionBoneIndices(BoxOffset);
-		FRigidTransform3* const BaseTransforms = Solver->GetCollisionBaseTransforms(BoxOffset);
+		Softs::FSolverRigidTransform3* const BaseTransforms = Solver->GetCollisionBaseTransforms(BoxOffset);
 		
 		for (int32 Index = 0; Index < NumBoxes; ++Index)
 		{
 			const FClothCollisionPrim_Box& Box = ClothCollisionData.Boxes[Index];
 			
-			BaseTransforms[Index] = FRigidTransform3(Box.LocalPosition, Box.LocalRotation);
+			BaseTransforms[Index] = Softs::FSolverRigidTransform3(Box.LocalPosition, Box.LocalRotation);
 			
 			BoneIndices[Index] = GetMappedBoneIndex(UsedBoneIndices, Box.BoneIndex);
 			UE_LOG(LogChaosCloth, VeryVerbose, TEXT("Found collision box on bone index %d."), BoneIndices[Index]);
@@ -305,20 +306,21 @@ void FClothingSimulationCollider::FLODData::Update(FClothingSimulationSolver* So
 	{
 		const int32 Offset = Offsets.FindChecked(FSolverClothPair(Solver, Cloth));
 		const int32* const BoneIndices = Solver->GetCollisionBoneIndices(Offset);
-		const FRigidTransform3* BaseTransforms = Solver->GetCollisionBaseTransforms(Offset);
-		FRigidTransform3* const CollisionTransforms = Solver->GetCollisionTransforms(Offset);
+		const Softs::FSolverRigidTransform3* BaseTransforms = Solver->GetCollisionBaseTransforms(Offset);
+		Softs::FSolverRigidTransform3* const CollisionTransforms = Solver->GetCollisionTransforms(Offset);
 
 		const TArrayView<const FTransform> BoneTransforms = Context ? Context->BoneTransforms : TArrayView<const FTransform>();
 
-		FTransform ComponentToLocalSpace = Context ? Context->ComponentToWorld : FTransform::Identity;
-		ComponentToLocalSpace.AddToTranslation(-Solver->GetLocalSpaceLocation());
+		FTransform ComponentToLocalSpaceReal = Context ? Context->ComponentToWorld : FTransform::Identity;
+		ComponentToLocalSpaceReal.AddToTranslation(-Solver->GetLocalSpaceLocation());
+		const Softs::FSolverRigidTransform3 ComponentToLocalSpace(ComponentToLocalSpaceReal);  // LWC, now in local space, therefore it is safe to use the solver transform type
 	
 		// Update the collision transforms
 		for (int32 Index = 0; Index < NumGeometries; ++Index)
 		{
 			const int32 BoneIndex = BoneIndices[Index];
 			CollisionTransforms[Index] = BoneTransforms.IsValidIndex(BoneIndex) ?
-				BaseTransforms[Index] * BoneTransforms[BoneIndex] * ComponentToLocalSpace :
+				BaseTransforms[Index] * Softs::FSolverRigidTransform3(BoneTransforms[BoneIndex]) * ComponentToLocalSpace :  // LWC requires the BoneTransform cast to the solver underlying type
 				BaseTransforms[Index] * ComponentToLocalSpace;  // External collisions don't map to a bone
 		}
 	}
@@ -342,17 +344,7 @@ void FClothingSimulationCollider::FLODData::ResetStartPose(FClothingSimulationSo
 	if (NumGeometries)
 	{
 		const int32 Offset = Offsets.FindChecked(FSolverClothPair(Solver, Cloth));
-		const FRigidTransform3* const CollisionTransforms = Solver->GetCollisionTransforms(Offset);
-		FRigidTransform3* const OldCollisionTransforms = Solver->GetOldCollisionTransforms(Offset);
-		FRotation3* const Rs = Solver->GetCollisionParticleRs(Offset);
-		FVec3* const Xs = Solver->GetCollisionParticleXs(Offset);
-
-		for (int32 Index = 0; Index < NumGeometries; ++Index)
-		{
-			OldCollisionTransforms[Index] = CollisionTransforms[Index];
-			Xs[Index] = CollisionTransforms[Index].GetTranslation();
-			Rs[Index] = CollisionTransforms[Index].GetRotation();
-		}
+		Solver->ResetCollisionStartPose(Offset, NumGeometries);
 	}
 }
 
@@ -743,37 +735,37 @@ void FClothingSimulationCollider::ResetStartPose(FClothingSimulationSolver* Solv
 	}
 }
 
-TConstArrayView<FVec3> FClothingSimulationCollider::GetCollisionTranslations(const FClothingSimulationSolver* Solver, const FClothingSimulationCloth* Cloth, ECollisionDataType CollisionDataType) const
+TConstArrayView<FSolverVec3> FClothingSimulationCollider::GetCollisionTranslations(const FClothingSimulationSolver* Solver, const FClothingSimulationCloth* Cloth, ECollisionDataType CollisionDataType) const
 {
 	check(Solver);
 	check(Cloth);
 
 	int32 Offset, NumGeometries;
 	return GetOffsetAndNumGeometries(Solver, Cloth, CollisionDataType, Offset, NumGeometries) ?
-		TConstArrayView<FVec3>(Solver->GetCollisionParticleXs(Offset), NumGeometries) :
-		TConstArrayView<FVec3>();
+		TConstArrayView<FSolverVec3>(Solver->GetCollisionParticleXs(Offset), NumGeometries) :
+		TConstArrayView<FSolverVec3>();
 }
 
-TConstArrayView<FRotation3> FClothingSimulationCollider::GetCollisionRotations(const FClothingSimulationSolver* Solver, const FClothingSimulationCloth* Cloth, ECollisionDataType CollisionDataType) const
+TConstArrayView<Softs::FSolverRotation3> FClothingSimulationCollider::GetCollisionRotations(const FClothingSimulationSolver* Solver, const FClothingSimulationCloth* Cloth, ECollisionDataType CollisionDataType) const
 {
 	check(Solver);
 	check(Cloth);
 
 	int32 Offset, NumGeometries;
 	return GetOffsetAndNumGeometries(Solver, Cloth, CollisionDataType, Offset, NumGeometries) ?
-		TConstArrayView<FRotation3>(Solver->GetCollisionParticleRs(Offset), NumGeometries) :
-		TConstArrayView<FRotation3>();
+		TConstArrayView<Softs::FSolverRotation3>(Solver->GetCollisionParticleRs(Offset), NumGeometries) :
+		TConstArrayView<Softs::FSolverRotation3>();
 }
 
-TConstArrayView<FRigidTransform3> FClothingSimulationCollider::GetOldCollisionTransforms(const FClothingSimulationSolver* Solver, const FClothingSimulationCloth* Cloth, ECollisionDataType CollisionDataType) const
+TConstArrayView<Softs::FSolverRigidTransform3> FClothingSimulationCollider::GetOldCollisionTransforms(const FClothingSimulationSolver* Solver, const FClothingSimulationCloth* Cloth, ECollisionDataType CollisionDataType) const
 {
 	check(Solver);
 	check(Cloth);
 
 	int32 Offset, NumGeometries;
 	return GetOffsetAndNumGeometries(Solver, Cloth, CollisionDataType, Offset, NumGeometries) ?
-		TConstArrayView<FRigidTransform3>(Solver->GetOldCollisionTransforms(Offset), NumGeometries) :
-		TConstArrayView<FRigidTransform3>();
+		TConstArrayView<Softs::FSolverRigidTransform3>(Solver->GetOldCollisionTransforms(Offset), NumGeometries) :
+		TConstArrayView<Softs::FSolverRigidTransform3>();
 }
 
 TConstArrayView<TUniquePtr<FImplicitObject>> FClothingSimulationCollider::GetCollisionGeometries(const FClothingSimulationSolver* Solver, const FClothingSimulationCloth* Cloth, ECollisionDataType CollisionDataType) const
@@ -797,3 +789,5 @@ TConstArrayView<bool> FClothingSimulationCollider::GetCollisionStatus(const FClo
 		TConstArrayView<bool>(Solver->GetCollisionStatus(Offset), NumGeometries) :
 		TConstArrayView<bool>();
 }
+
+}  // End namespace Chaos
