Chaos Cloth - Move back the post LWC solver to float for performance reason, since cloth is being simulated in local space.

diff --git a/Engine/Plugins/Experimental/ChaosCloth/Source/Chaos/Private/ChaosCloth/ChaosClothingSimulationInteractor.cpp b/Engine/Plugins/Experimental/ChaosCloth/Source/Chaos/Private/ChaosCloth/ChaosClothingSimulationInteractor.cpp
--- a/Engine/Plugins/Experimental/ChaosCloth/Source/Chaos/Private/ChaosCloth/ChaosClothingSimulationInteractor.cpp
+++ b/Engine/Plugins/Experimental/ChaosCloth/Source/Chaos/Private/ChaosCloth/ChaosClothingSimulationInteractor.cpp
@@ -4,8 +4,6 @@
 #include "ChaosCloth/ChaosClothingSimulationCloth.h"
 #include "ChaosCloth/ChaosClothingSimulation.h"
 
-using namespace Chaos;
-
 namespace ChaosClothingInteractor
 {
 	static const float InvStiffnessLogBase = 1.f / FMath::Loge(1.e3f);  // Log base for updating old linear stiffnesses to the new stiffness exponentiation
@@ -15,7 +13,7 @@ void UChaosClothingInteractor::Sync(IClothingSimulation* Simulation)
 {
 	check(Simulation);
 
-	if (FClothingSimulationCloth* const Cloth = static_cast<FClothingSimulation*>(Simulation)->GetCloth(ClothingId))
+	if (Chaos::FClothingSimulationCloth* const Cloth = static_cast<Chaos::FClothingSimulation*>(Simulation)->GetCloth(ClothingId))
 	{
 		for (FChaosClothingInteractorCommand& Command : Commands)
 		{
@@ -30,11 +28,11 @@ void UChaosClothingInteractor::Sync(IClothingSimulation* Simulation)
 
 void UChaosClothingInteractor::SetMaterialLinear(float EdgeStiffnessLinear, float BendingStiffnessLinear, float AreaStiffnessLinear)
 {
-	const TVec2<FRealSingle> EdgeStiffness((FMath::Clamp(FMath::Loge(EdgeStiffnessLinear) * ChaosClothingInteractor::InvStiffnessLogBase + 1.f, 0.f, 1.f)), 1.f);
-	const TVec2<FRealSingle> BendingStiffness((FMath::Clamp(FMath::Loge(BendingStiffnessLinear) * ChaosClothingInteractor::InvStiffnessLogBase + 1.f, 0.f, 1.f)), 1.f);
-	const TVec2<FRealSingle> AreaStiffness((FMath::Clamp(FMath::Loge(AreaStiffnessLinear) * ChaosClothingInteractor::InvStiffnessLogBase + 1.f, 0.f, 1.f)), 1.f);
+	const Chaos::TVec2<Chaos::FRealSingle> EdgeStiffness((FMath::Clamp(FMath::Loge(EdgeStiffnessLinear) * ChaosClothingInteractor::InvStiffnessLogBase + 1.f, 0.f, 1.f)), 1.f);
+	const Chaos::TVec2<Chaos::FRealSingle> BendingStiffness((FMath::Clamp(FMath::Loge(BendingStiffnessLinear) * ChaosClothingInteractor::InvStiffnessLogBase + 1.f, 0.f, 1.f)), 1.f);
+	const Chaos::TVec2<Chaos::FRealSingle> AreaStiffness((FMath::Clamp(FMath::Loge(AreaStiffnessLinear) * ChaosClothingInteractor::InvStiffnessLogBase + 1.f, 0.f, 1.f)), 1.f);
 
-	Commands.Add(FChaosClothingInteractorCommand::CreateLambda([EdgeStiffness, BendingStiffness, AreaStiffness](FClothingSimulationCloth* Cloth)
+	Commands.Add(FChaosClothingInteractorCommand::CreateLambda([EdgeStiffness, BendingStiffness, AreaStiffness](Chaos::FClothingSimulationCloth* Cloth)
 	{
 		Cloth->SetMaterialProperties(EdgeStiffness, BendingStiffness, AreaStiffness);
 	}));
@@ -42,35 +40,35 @@ void UChaosClothingInteractor::SetMaterialLinear(float EdgeStiffnessLinear, floa
 
 void UChaosClothingInteractor::SetMaterial(FVector2D EdgeStiffness, FVector2D BendingStiffness, FVector2D AreaStiffness)
 {
-	Commands.Add(FChaosClothingInteractorCommand::CreateLambda([EdgeStiffness, BendingStiffness, AreaStiffness](FClothingSimulationCloth* Cloth)
+	Commands.Add(FChaosClothingInteractorCommand::CreateLambda([EdgeStiffness, BendingStiffness, AreaStiffness](Chaos::FClothingSimulationCloth* Cloth)
 	{
-		Cloth->SetMaterialProperties(TVec2<FRealSingle>(EdgeStiffness[0], EdgeStiffness[1]), TVec2<FRealSingle>(BendingStiffness[0], BendingStiffness[1]), TVec2<FRealSingle>(AreaStiffness[0], AreaStiffness[1]));
+		Cloth->SetMaterialProperties(Chaos::TVec2<Chaos::FRealSingle>(EdgeStiffness[0], EdgeStiffness[1]), Chaos::TVec2<Chaos::FRealSingle>(BendingStiffness[0], BendingStiffness[1]), Chaos::TVec2<Chaos::FRealSingle>(AreaStiffness[0], AreaStiffness[1]));
 	}));
 }
 
 void UChaosClothingInteractor::SetLongRangeAttachmentLinear(float TetherStiffnessLinear, float TetherScale)
 {
 	// Deprecated
-	const TVec2<FRealSingle> TetherStiffness((FMath::Clamp(FMath::Loge(TetherStiffnessLinear) * ChaosClothingInteractor::InvStiffnessLogBase + 1.f, 0.f, 1.f)), 1.f);
-	Commands.Add(FChaosClothingInteractorCommand::CreateLambda([TetherStiffness, TetherScale](FClothingSimulationCloth* Cloth)
+	const Chaos::TVec2<Chaos::FRealSingle> TetherStiffness((FMath::Clamp(FMath::Loge(TetherStiffnessLinear) * ChaosClothingInteractor::InvStiffnessLogBase + 1.f, 0.f, 1.f)), 1.f);
+	Commands.Add(FChaosClothingInteractorCommand::CreateLambda([TetherStiffness, TetherScale](Chaos::FClothingSimulationCloth* Cloth)
 	{
-		Cloth->SetLongRangeAttachmentProperties(TetherStiffness, TVec2<FRealSingle>(TetherScale, TetherScale));
+		Cloth->SetLongRangeAttachmentProperties(TetherStiffness, Chaos::TVec2<Chaos::FRealSingle>(TetherScale, TetherScale));
 	}));
 }
 
 void UChaosClothingInteractor::SetLongRangeAttachment(FVector2D TetherStiffness, FVector2D TetherScale)
 {
-	Commands.Add(FChaosClothingInteractorCommand::CreateLambda([TetherStiffness, TetherScale](FClothingSimulationCloth* Cloth)
+	Commands.Add(FChaosClothingInteractorCommand::CreateLambda([TetherStiffness, TetherScale](Chaos::FClothingSimulationCloth* Cloth)
 	{
 		Cloth->SetLongRangeAttachmentProperties(
-			TVec2<FRealSingle>(TetherStiffness[0], TetherStiffness[1]),
-			TVec2<FRealSingle>(TetherScale[0], TetherScale[1]));
+			Chaos::TVec2<Chaos::FRealSingle>(TetherStiffness[0], TetherStiffness[1]),
+			Chaos::TVec2<Chaos::FRealSingle>(TetherScale[0], TetherScale[1]));
 	}));
 }
 
 void UChaosClothingInteractor::SetCollision(float CollisionThickness, float FrictionCoefficient, bool bUseCCD, float SelfCollisionThickness)
 {
-	Commands.Add(FChaosClothingInteractorCommand::CreateLambda([CollisionThickness, FrictionCoefficient, bUseCCD, SelfCollisionThickness](FClothingSimulationCloth* Cloth)
+	Commands.Add(FChaosClothingInteractorCommand::CreateLambda([CollisionThickness, FrictionCoefficient, bUseCCD, SelfCollisionThickness](Chaos::FClothingSimulationCloth* Cloth)
 	{
 		Cloth->SetCollisionProperties(CollisionThickness, FrictionCoefficient, bUseCCD, SelfCollisionThickness);
 	}));
@@ -78,14 +76,14 @@ void UChaosClothingInteractor::SetCollision(float CollisionThickness, float Fric
 
 void UChaosClothingInteractor::SetBackstop(bool bEnabled)
 {
-	Commands.Add(FChaosClothingInteractorCommand::CreateLambda([bEnabled](FClothingSimulationCloth* Cloth)
+	Commands.Add(FChaosClothingInteractorCommand::CreateLambda([bEnabled](Chaos::FClothingSimulationCloth* Cloth)
 	{
 		Cloth->SetBackstopProperties(bEnabled);
 	}));
 }
 void UChaosClothingInteractor::SetDamping(float DampingCoefficient)
 {
-	Commands.Add(FChaosClothingInteractorCommand::CreateLambda([DampingCoefficient](FClothingSimulationCloth* Cloth)
+	Commands.Add(FChaosClothingInteractorCommand::CreateLambda([DampingCoefficient](Chaos::FClothingSimulationCloth* Cloth)
 	{
 		Cloth->SetDampingProperties(DampingCoefficient);
 	}));
@@ -93,24 +91,24 @@ void UChaosClothingInteractor::SetDamping(float DampingCoefficient)
 
 void UChaosClothingInteractor::SetAerodynamics(float DragCoefficient, float LiftCoefficient, FVector WindVelocity)
 {
-	Commands.Add(FChaosClothingInteractorCommand::CreateLambda([DragCoefficient, LiftCoefficient, WindVelocity](FClothingSimulationCloth* Cloth)
+	Commands.Add(FChaosClothingInteractorCommand::CreateLambda([DragCoefficient, LiftCoefficient, WindVelocity](Chaos::FClothingSimulationCloth* Cloth)
 	{
-		constexpr FRealSingle AirDensity = 1.225e-6f;
-		Cloth->SetAerodynamicsProperties(TVec2<FRealSingle>(DragCoefficient, DragCoefficient), TVec2<FRealSingle>(LiftCoefficient, LiftCoefficient), AirDensity, WindVelocity);
+		constexpr Chaos::FRealSingle AirDensity = 1.225e-6f;
+		Cloth->SetAerodynamicsProperties(Chaos::TVec2<Chaos::FRealSingle>(DragCoefficient, DragCoefficient), Chaos::TVec2<Chaos::FRealSingle>(LiftCoefficient, LiftCoefficient), AirDensity, WindVelocity);
 	}));
 }
 
 void UChaosClothingInteractor::SetWind(FVector2D Drag, FVector2D Lift, float AirDensity, FVector WindVelocity)
 {
-	Commands.Add(FChaosClothingInteractorCommand::CreateLambda([Drag, Lift, AirDensity, WindVelocity](FClothingSimulationCloth* Cloth)
+	Commands.Add(FChaosClothingInteractorCommand::CreateLambda([Drag, Lift, AirDensity, WindVelocity](Chaos::FClothingSimulationCloth* Cloth)
 	{
-		Cloth->SetAerodynamicsProperties(TVec2<FRealSingle>(Drag[0], Drag[1]), TVec2<FRealSingle>(Lift[0], Lift[1]), AirDensity, WindVelocity);
+		Cloth->SetAerodynamicsProperties(Chaos::TVec2<Chaos::FRealSingle>(Drag[0], Drag[1]), Chaos::TVec2<Chaos::FRealSingle>(Lift[0], Lift[1]), AirDensity, WindVelocity);
 	}));
 }
 
 void UChaosClothingInteractor::SetGravity(float GravityScale, bool bIsGravityOverridden, FVector GravityOverride)
 {
-	Commands.Add(FChaosClothingInteractorCommand::CreateLambda([GravityScale, bIsGravityOverridden, GravityOverride](FClothingSimulationCloth* Cloth)
+	Commands.Add(FChaosClothingInteractorCommand::CreateLambda([GravityScale, bIsGravityOverridden, GravityOverride](Chaos::FClothingSimulationCloth* Cloth)
 	{
 		Cloth->SetGravityProperties(GravityScale, bIsGravityOverridden, GravityOverride);
 	}));
@@ -119,26 +117,26 @@ void UChaosClothingInteractor::SetGravity(float GravityScale, bool bIsGravityOve
 void UChaosClothingInteractor::SetAnimDriveLinear(float AnimDriveStiffnessLinear)
 {
 	// Deprecated
-	const TVec2<FRealSingle> AnimDriveStiffness(0.f, FMath::Clamp(FMath::Loge(AnimDriveStiffnessLinear) * ChaosClothingInteractor::InvStiffnessLogBase + 1.f, 0.f, 1.f));
-	Commands.Add(FChaosClothingInteractorCommand::CreateLambda([AnimDriveStiffness](FClothingSimulationCloth* Cloth)
+	const Chaos::TVec2<Chaos::FRealSingle> AnimDriveStiffness(0.f, FMath::Clamp(FMath::Loge(AnimDriveStiffnessLinear) * ChaosClothingInteractor::InvStiffnessLogBase + 1.f, 0.f, 1.f));
+	Commands.Add(FChaosClothingInteractorCommand::CreateLambda([AnimDriveStiffness](Chaos::FClothingSimulationCloth* Cloth)
 	{
 		// The Anim Drive stiffness Low value needs to be 0 in order to keep backward compatibility with existing mask (this wouldn't be an issue if this property had no legacy mask)
-		static const TVec2<FRealSingle> AnimDriveDamping(0.f, 1.f);
+		static const Chaos::TVec2<Chaos::FRealSingle> AnimDriveDamping(0.f, 1.f);
 		Cloth->SetAnimDriveProperties(AnimDriveStiffness, AnimDriveDamping);
 	}));
 }
 
 void UChaosClothingInteractor::SetAnimDrive(FVector2D AnimDriveStiffness, FVector2D AnimDriveDamping)
 {
-	Commands.Add(FChaosClothingInteractorCommand::CreateLambda([AnimDriveStiffness, AnimDriveDamping](FClothingSimulationCloth* Cloth)
+	Commands.Add(FChaosClothingInteractorCommand::CreateLambda([AnimDriveStiffness, AnimDriveDamping](Chaos::FClothingSimulationCloth* Cloth)
 	{
-		Cloth->SetAnimDriveProperties(TVec2<FRealSingle>(AnimDriveStiffness.X, AnimDriveStiffness.Y), TVec2<FRealSingle>(AnimDriveDamping.X, AnimDriveDamping.Y));
+		Cloth->SetAnimDriveProperties(Chaos::TVec2<Chaos::FRealSingle>(AnimDriveStiffness.X, AnimDriveStiffness.Y), Chaos::TVec2<Chaos::FRealSingle>(AnimDriveDamping.X, AnimDriveDamping.Y));
 	}));
 }
 
 void UChaosClothingInteractor::SetVelocityScale(FVector LinearVelocityScale, float AngularVelocityScale, float FictitiousAngularScale)
 {
-	Commands.Add(FChaosClothingInteractorCommand::CreateLambda([LinearVelocityScale, AngularVelocityScale, FictitiousAngularScale](FClothingSimulationCloth* Cloth)
+	Commands.Add(FChaosClothingInteractorCommand::CreateLambda([LinearVelocityScale, AngularVelocityScale, FictitiousAngularScale](Chaos::FClothingSimulationCloth* Cloth)
 	{
 		Cloth->SetVelocityScaleProperties(LinearVelocityScale, AngularVelocityScale, FictitiousAngularScale);
 	}));
@@ -148,14 +146,14 @@ void UChaosClothingInteractor::ResetAndTeleport(bool bReset, bool bTeleport)
 {
 	if (bReset)
 	{
-		Commands.Add(FChaosClothingInteractorCommand::CreateLambda([](FClothingSimulationCloth* Cloth)
+		Commands.Add(FChaosClothingInteractorCommand::CreateLambda([](Chaos::FClothingSimulationCloth* Cloth)
 		{
 			Cloth->Reset();
 		}));
 	}
 	if (bTeleport)
 	{
-		Commands.Add(FChaosClothingInteractorCommand::CreateLambda([](FClothingSimulationCloth* Cloth)
+		Commands.Add(FChaosClothingInteractorCommand::CreateLambda([](Chaos::FClothingSimulationCloth* Cloth)
 		{
 			Cloth->Teleport();
 		}));
@@ -169,7 +167,7 @@ void UChaosClothingSimulationInteractor::Sync(IClothingSimulation* Simulation, I
 
 	for (FChaosClothingSimulationInteractorCommand& Command : Commands)
 	{
-		Command.Execute(static_cast<FClothingSimulation*>(Simulation), static_cast<FClothingSimulationContext*>(Context));
+		Command.Execute(static_cast<Chaos::FClothingSimulation*>(Simulation), static_cast<Chaos::FClothingSimulationContext*>(Context));
 	}
 	Commands.Reset();
 
@@ -179,7 +177,7 @@ void UChaosClothingSimulationInteractor::Sync(IClothingSimulation* Simulation, I
 
 void UChaosClothingSimulationInteractor::PhysicsAssetUpdated()
 {
-	Commands.Add(FChaosClothingSimulationInteractorCommand::CreateLambda([](FClothingSimulation* Simulation, FClothingSimulationContext* /*Context*/)
+	Commands.Add(FChaosClothingSimulationInteractorCommand::CreateLambda([](Chaos::FClothingSimulation* Simulation, Chaos::FClothingSimulationContext* /*Context*/)
 	{
 		Simulation->RefreshPhysicsAsset();
 	}));
@@ -187,7 +185,7 @@ void UChaosClothingSimulationInteractor::PhysicsAssetUpdated()
 
 void UChaosClothingSimulationInteractor::ClothConfigUpdated()
 {
-	Commands.Add(FChaosClothingSimulationInteractorCommand::CreateLambda([](FClothingSimulation* Simulation, FClothingSimulationContext* Context)
+	Commands.Add(FChaosClothingSimulationInteractorCommand::CreateLambda([](Chaos::FClothingSimulation* Simulation, Chaos::FClothingSimulationContext* Context)
 	{
 		Simulation->RefreshClothConfig(Context);
 	}));
@@ -207,7 +205,7 @@ void UChaosClothingSimulationInteractor::SetAnimDriveSpringStiffness(float Stiff
 
 void UChaosClothingSimulationInteractor::EnableGravityOverride(const FVector& Gravity)
 {
-	Commands.Add(FChaosClothingSimulationInteractorCommand::CreateLambda([Gravity](FClothingSimulation* Simulation, FClothingSimulationContext* /*Context*/)
+	Commands.Add(FChaosClothingSimulationInteractorCommand::CreateLambda([Gravity](Chaos::FClothingSimulation* Simulation, Chaos::FClothingSimulationContext* /*Context*/)
 	{
 		Simulation->SetGravityOverride(Gravity);
 	}));
@@ -215,7 +213,7 @@ void UChaosClothingSimulationInteractor::EnableGravityOverride(const FVector& Gr
 
 void UChaosClothingSimulationInteractor::DisableGravityOverride()
 {
-	Commands.Add(FChaosClothingSimulationInteractorCommand::CreateLambda([](FClothingSimulation* Simulation, FClothingSimulationContext* /*Context*/)
+	Commands.Add(FChaosClothingSimulationInteractorCommand::CreateLambda([](Chaos::FClothingSimulation* Simulation, Chaos::FClothingSimulationContext* /*Context*/)
 	{
 		Simulation->DisableGravityOverride();
 	}));
@@ -223,7 +221,7 @@ void UChaosClothingSimulationInteractor::DisableGravityOverride()
 
 void UChaosClothingSimulationInteractor::SetNumIterations(int32 InNumIterations)
 {
-	Commands.Add(FChaosClothingSimulationInteractorCommand::CreateLambda([InNumIterations](FClothingSimulation* Simulation, FClothingSimulationContext* /*Context*/)
+	Commands.Add(FChaosClothingSimulationInteractorCommand::CreateLambda([InNumIterations](Chaos::FClothingSimulation* Simulation, Chaos::FClothingSimulationContext* /*Context*/)
 	{
 		Simulation->SetNumIterations(InNumIterations);
 	}));
@@ -231,7 +229,7 @@ void UChaosClothingSimulationInteractor::SetNumIterations(int32 InNumIterations)
 
 void UChaosClothingSimulationInteractor::SetNumSubsteps(int32 InNumSubsteps)
 {
-	Commands.Add(FChaosClothingSimulationInteractorCommand::CreateLambda([InNumSubsteps](FClothingSimulation* Simulation, FClothingSimulationContext* /*Context*/)
+	Commands.Add(FChaosClothingSimulationInteractorCommand::CreateLambda([InNumSubsteps](Chaos::FClothingSimulation* Simulation, Chaos::FClothingSimulationContext* /*Context*/)
 	{
 		Simulation->SetNumSubsteps(InNumSubsteps);
 	}));
diff --git a/Engine/Plugins/Experimental/ChaosCloth/Source/Chaos/Private/ChaosCloth/ChaosClothingSimulationMesh.cpp b/Engine/Plugins/Experimental/ChaosCloth/Source/Chaos/Private/ChaosCloth/ChaosClothingSimulationMesh.cpp
--- a/Engine/Plugins/Experimental/ChaosCloth/Source/Chaos/Private/ChaosCloth/ChaosClothingSimulationMesh.cpp
+++ b/Engine/Plugins/Experimental/ChaosCloth/Source/Chaos/Private/ChaosCloth/ChaosClothingSimulationMesh.cpp
@@ -16,10 +16,10 @@
 bool bChaos_SkinPhysicsMesh_ISPC_Enabled = true;
 FAutoConsoleVariableRef CVarChaosSkinPhysicsMeshISPCEnabled(TEXT("p.Chaos.SkinPhysicsMesh.ISPC"), bChaos_SkinPhysicsMesh_ISPC_Enabled, TEXT("Whether to use ISPC optimizations on skinned physics meshes"));
 
-static_assert(sizeof(ispc::FVector) == sizeof(Chaos::FVec3), "sizeof(ispc::FVector) != sizeof(FVec3)");
+static_assert(sizeof(ispc::FVector3f) == sizeof(Chaos::Softs::FSolverVec3), "sizeof(ispc::FVector3f) != sizeof(Chaos::Softs::FSolverVec3)");
 static_assert(sizeof(ispc::FVector3f) == sizeof(FVector3f), "sizeof(ispc::FVector3f) != sizeof(FVector3f)");
 static_assert(sizeof(ispc::FMatrix44f) == sizeof(FMatrix44f), "sizeof(ispc::FMatrix44f) != sizeof(FMatrix44f)");
-static_assert(sizeof(ispc::FTransform) == sizeof(FTransform), "sizeof(ispc::FTransform) != sizeof(FTransform)");
+static_assert(sizeof(ispc::FTransform3f) == sizeof(FTransform3f), "sizeof(ispc::FTransform3f) != sizeof(FTransform3f)");
 static_assert(sizeof(ispc::FClothVertBoneData) == sizeof(FClothVertBoneData), "sizeof(ispc::FClothVertBoneData) != sizeof(Chaos::FClothVertBoneData)");
 #endif
 
@@ -27,7 +27,8 @@ DECLARE_CYCLE_STAT(TEXT("Chaos Cloth Skin Physics Mesh"), STAT_ChaosClothSkinPhy
 DECLARE_CYCLE_STAT(TEXT("Chaos Cloth Wrap Deform Mesh"), STAT_ChaosClothWrapDeformMesh, STATGROUP_ChaosCloth);
 DECLARE_CYCLE_STAT(TEXT("Chaos Cloth Wrap Deform Cloth LOD"), STAT_ChaosClothWrapDeformClothLOD, STATGROUP_ChaosCloth);
 
-using namespace Chaos;
+namespace Chaos
+{
 
 FClothingSimulationMesh::FClothingSimulationMesh(const UClothingAssetCommon* InAsset, const USkeletalMeshComponent* InSkeletalMeshComponent)
 	: Asset(InAsset)
@@ -147,9 +148,9 @@ FRigidTransform3 Chaos::FClothingSimulationMesh::GetReferenceBoneTransform() con
 bool FClothingSimulationMesh::WrapDeformLOD(
 	int32 PrevLODIndex,
 	int32 LODIndex,
-	const FVec3* Normals,
-	const FVec3* Positions,
-	FVec3* OutPositions) const
+	const FSolverVec3* Normals,
+	const FSolverVec3* Positions,
+	FSolverVec3* OutPositions) const
 {
 	TRACE_CPUPROFILER_EVENT_SCOPE(FClothingSimulationMesh_WrapDeformLOD);
 	SCOPE_CYCLE_COUNTER(STAT_ChaosClothWrapDeformMesh);
@@ -175,9 +176,9 @@ bool FClothingSimulationMesh::WrapDeformLOD(
 		const int32 VertIndex2 = (int32)VertData.SourceMeshVertIndices[2];
 
 		OutPositions[Index] = 
-			Positions[VertIndex0] * VertData.PositionBaryCoordsAndDist.X + Normals[VertIndex0] * VertData.PositionBaryCoordsAndDist.W +
-			Positions[VertIndex1] * VertData.PositionBaryCoordsAndDist.Y + Normals[VertIndex1] * VertData.PositionBaryCoordsAndDist.W +
-			Positions[VertIndex2] * VertData.PositionBaryCoordsAndDist.Z + Normals[VertIndex2] * VertData.PositionBaryCoordsAndDist.W;
+			Positions[VertIndex0] * (FSolverReal)VertData.PositionBaryCoordsAndDist.X + Normals[VertIndex0] * (FSolverReal)VertData.PositionBaryCoordsAndDist.W +
+			Positions[VertIndex1] * (FSolverReal)VertData.PositionBaryCoordsAndDist.Y + Normals[VertIndex1] * (FSolverReal)VertData.PositionBaryCoordsAndDist.W +
+			Positions[VertIndex2] * (FSolverReal)VertData.PositionBaryCoordsAndDist.Z + Normals[VertIndex2] * (FSolverReal)VertData.PositionBaryCoordsAndDist.W;
 	}
 
 	return true;
@@ -186,12 +187,12 @@ bool FClothingSimulationMesh::WrapDeformLOD(
 bool FClothingSimulationMesh::WrapDeformLOD(
 	int32 PrevLODIndex,
 	int32 LODIndex,
-	const FVec3* Normals,
-	const FVec3* Positions,
-	const FVec3* Velocities,
-	FVec3* OutPositions0,
-	FVec3* OutPositions1,
-	FVec3* OutVelocities) const
+	const FSolverVec3* Normals,
+	const FSolverVec3* Positions,
+	const FSolverVec3* Velocities,
+	FSolverVec3* OutPositions0,
+	FSolverVec3* OutPositions1,
+	FSolverVec3* OutVelocities) const
 {
 	TRACE_CPUPROFILER_EVENT_SCOPE(FClothingSimulationMesh_WrapDeformLOD);
 	SCOPE_CYCLE_COUNTER(STAT_ChaosClothWrapDeformClothLOD);
@@ -217,27 +218,27 @@ bool FClothingSimulationMesh::WrapDeformLOD(
 		const int32 VertIndex2 = (int32)VertData.SourceMeshVertIndices[2];
 
 		OutPositions0[Index] = OutPositions1[Index] =
-			Positions[VertIndex0] * VertData.PositionBaryCoordsAndDist.X + Normals[VertIndex0] * VertData.PositionBaryCoordsAndDist.W +
-			Positions[VertIndex1] * VertData.PositionBaryCoordsAndDist.Y + Normals[VertIndex1] * VertData.PositionBaryCoordsAndDist.W +
-			Positions[VertIndex2] * VertData.PositionBaryCoordsAndDist.Z + Normals[VertIndex2] * VertData.PositionBaryCoordsAndDist.W;
+			Positions[VertIndex0] * (FSolverReal)VertData.PositionBaryCoordsAndDist.X + Normals[VertIndex0] * (FSolverReal)VertData.PositionBaryCoordsAndDist.W +
+			Positions[VertIndex1] * (FSolverReal)VertData.PositionBaryCoordsAndDist.Y + Normals[VertIndex1] * (FSolverReal)VertData.PositionBaryCoordsAndDist.W +
+			Positions[VertIndex2] * (FSolverReal)VertData.PositionBaryCoordsAndDist.Z + Normals[VertIndex2] * (FSolverReal)VertData.PositionBaryCoordsAndDist.W;
 
 		OutVelocities[Index] = 
-			Velocities[VertIndex0] * VertData.PositionBaryCoordsAndDist.X +
-			Velocities[VertIndex1] * VertData.PositionBaryCoordsAndDist.Y +
-			Velocities[VertIndex2] * VertData.PositionBaryCoordsAndDist.Z;
+			Velocities[VertIndex0] * (FSolverReal)VertData.PositionBaryCoordsAndDist.X +
+			Velocities[VertIndex1] * (FSolverReal)VertData.PositionBaryCoordsAndDist.Y +
+			Velocities[VertIndex2] * (FSolverReal)VertData.PositionBaryCoordsAndDist.Z;
 	}
 
 	return true;
 }
 
-// Inline function used to force the unrolling of the skinning loop
+// Inline function used to force the unrolling of the skinning loop, LWC: note skinning is all done in float to match the asset data type
 FORCEINLINE static void AddInfluence(FVector3f& OutPosition, FVector3f& OutNormal, const FVector3f& RefParticle, const FVector3f& RefNormal, const FMatrix44f& BoneMatrix, const float Weight)
 {
 	OutPosition += BoneMatrix.TransformPosition(RefParticle) * Weight;
 	OutNormal += BoneMatrix.TransformVector(RefNormal) * Weight;
 }
 
-void FClothingSimulationMesh::SkinPhysicsMesh(int32 LODIndex, const FVec3& LocalSpaceLocation, FVec3* OutPositions, FVec3* OutNormals) const
+void FClothingSimulationMesh::SkinPhysicsMesh(int32 LODIndex, const FVec3& LocalSpaceLocation, FSolverVec3* OutPositions, FSolverVec3* OutNormals) const
 {
 	TRACE_CPUPROFILER_EVENT_SCOPE(FClothingSimulationMesh_SkinPhysicsMesh);
 	SCOPE_CYCLE_COUNTER(STAT_ChaosClothSkinPhysicsMesh);
@@ -253,23 +254,23 @@ void FClothingSimulationMesh::SkinPhysicsMesh(int32 LODIndex, const FVec3& Local
 	const FClothingSimulationContextCommon* const Context = static_cast<const FClothingSimulationContextCommon*>(SkeletalMeshComponent->GetClothingSimulationContext());
 	FTransform ComponentToLocalSpaceReal = Context->ComponentToWorld;
 	ComponentToLocalSpaceReal.AddToTranslation(-LocalSpaceLocation);
-	const FTransform3f ComponentToLocalSpace(ComponentToLocalSpaceReal);
+	const FTransform3f ComponentToLocalSpace(ComponentToLocalSpaceReal);  // LWC: Now in local space, therefore it is safe to use single precision which is the asset data format
 
 	const int32* const RESTRICT BoneMap = Asset->UsedBoneIndices.GetData();
 	const FMatrix44f* const RESTRICT BoneMatrices = Context->RefToLocals.GetData();
 		
 #if INTEL_ISPC
-	if (bChaos_SkinPhysicsMesh_ISPC_Enabled && bRealTypeCompatibleWithISPC)
+	if (bChaos_SkinPhysicsMesh_ISPC_Enabled)
 	{
 		ispc::SkinPhysicsMesh(
-			(ispc::FVector*)OutPositions,
-			(ispc::FVector*)OutNormals,
+			(ispc::FVector3f*)OutPositions,
+			(ispc::FVector3f*)OutNormals,
 			(ispc::FVector3f*)PhysicalMeshData.Vertices.GetData(),
 			(ispc::FVector3f*)PhysicalMeshData.Normals.GetData(),
 			(ispc::FClothVertBoneData*)PhysicalMeshData.BoneData.GetData(),
 			BoneMap,
 			(ispc::FMatrix44f*)BoneMatrices,
-			(ispc::FTransform&)ComponentToLocalSpaceReal,
+			(ispc::FTransform3f&)ComponentToLocalSpace,
 			NumPoints);
 	}
 	else
@@ -309,8 +310,8 @@ void FClothingSimulationMesh::SkinPhysicsMesh(int32 LODIndex, const FVec3& Local
 			default: break;
 			}
 
-			OutPositions[VertIndex] = ComponentToLocalSpace.TransformPosition(Position);
-			OutNormals[VertIndex] = ComponentToLocalSpace.TransformVector(Normal).GetSafeNormal();
+			OutPositions[VertIndex] = FSolverVec3(ComponentToLocalSpace.TransformPosition(Position));
+			OutNormals[VertIndex] = FSolverVec3(ComponentToLocalSpace.TransformVector(Normal).GetSafeNormal());
 
 		}, NumPoints > MinParallelVertices ? EParallelForFlags::None : EParallelForFlags::ForceSingleThread);
 	}
@@ -333,8 +334,8 @@ void FClothingSimulationMesh::Update(
 
 	// Skin current LOD positions
 	const FVec3& LocalSpaceLocation = Solver->GetLocalSpaceLocation();
-	FVec3* const OutPositions = Solver->GetAnimationPositions(Offset);
-	FVec3* const OutNormals = Solver->GetAnimationNormals(Offset);
+	FSolverVec3* const OutPositions = Solver->GetAnimationPositions(Offset);
+	FSolverVec3* const OutNormals = Solver->GetAnimationNormals(Offset);
 	
 	SkinPhysicsMesh(LODIndex, LocalSpaceLocation, OutPositions, OutNormals);
 
@@ -342,9 +343,9 @@ void FClothingSimulationMesh::Update(
 	if (LODIndex != PrevLODIndex)
 	{
 		// TODO: Using the more accurate skinning method here would require double buffering the context at the skeletal mesh level
-		const FVec3* const SrcWrapNormals = Solver->GetAnimationNormals(PrevOffset);  // No need to keep an old normals array around, since the LOD has just changed
-		const FVec3* const SrcWrapPositions = Solver->GetOldAnimationPositions(PrevOffset);
-		FVec3* const OutOldPositions = Solver->GetOldAnimationPositions(Offset);
+		const FSolverVec3* const SrcWrapNormals = Solver->GetAnimationNormals(PrevOffset);  // No need to keep an old normals array around, since the LOD has just changed
+		const FSolverVec3* const SrcWrapPositions = Solver->GetOldAnimationPositions(PrevOffset);
+		FSolverVec3* const OutOldPositions = Solver->GetOldAnimationPositions(Offset);
 
 		const bool bValidWrap = WrapDeformLOD(PrevLODIndex, LODIndex, SrcWrapNormals, SrcWrapPositions, OutOldPositions);
 	
@@ -361,3 +362,5 @@ void FClothingSimulationMesh::Update(
 		}
 	}
 }
+
+}  // End namespace Chaos
diff --git a/Engine/Plugins/Experimental/ChaosCloth/Source/Chaos/Private/ChaosCloth/ChaosClothingSimulationMesh.ispc b/Engine/Plugins/Experimental/ChaosCloth/Source/Chaos/Private/ChaosCloth/ChaosClothingSimulationMesh.ispc
--- a/Engine/Plugins/Experimental/ChaosCloth/Source/Chaos/Private/ChaosCloth/ChaosClothingSimulationMesh.ispc
+++ b/Engine/Plugins/Experimental/ChaosCloth/Source/Chaos/Private/ChaosCloth/ChaosClothingSimulationMesh.ispc
@@ -39,22 +39,22 @@ static inline void AddInfluence(uniform FVector4f& OutPosition, uniform FVector4
 	OutNormal = OutNormal + Norm * Weight;
 }
 
-static inline uniform FVector GetSafeNormal(const uniform FVector &V)
+static inline uniform FVector3f GetSafeNormal(const uniform FVector3f &V)
 {
 	//We want N / ||N|| and to avoid inf
 	//So we want N / ||N|| < 1 / eps => N eps < ||N||, but this is clearly true for all eps < 1 and N > 0
-	const uniform FReal SizeSqr = VectorSizeSquared(V);
-	return VectorSelect(SizeSqr <= REAL_MIN, ForwardVector, V / sqrt(SizeSqr));
+	const uniform float SizeSqr = VectorSizeSquared(V);
+	return VectorSelect(SizeSqr <= FLT_MIN, FloatForwardVector, V / sqrt(SizeSqr));
 }
 
-export void SkinPhysicsMesh(uniform FVector OutPositions[],
-							uniform FVector OutNormals[],
+export void SkinPhysicsMesh(uniform FVector3f OutPositions[],
+							uniform FVector3f OutNormals[],
 							const uniform FVector3f Vertices[],
 							const uniform FVector3f Normals[],
 							const uniform FClothVertBoneData BoneData[],
 							const uniform int32 BoneMap[],
 							const uniform FMatrix44f BoneMatrices[],
-							const uniform FTransform &ComponentToLocalSpace,
+							const uniform FTransform3f &ComponentToLocalSpace,
 							const uniform int32 NumPoints)
 {
 	for(uniform int i = 0; i < NumPoints; ++i)
@@ -76,11 +76,8 @@ export void SkinPhysicsMesh(uniform FVector OutPositions[],
 			AddInfluence(OutPosition, OutNormal, RefParticle, RefNormal, BoneMatrix, Weight);
 		}
 
-		const uniform FVector4 ConvertedPosition = ConvertVector4fTo4Native(OutPosition);
-		const uniform FVector4 ConvertedNormal = ConvertVector4fTo4Native(OutNormal);
-
-		const uniform FVector Position = SetVector(ConvertedPosition);
-		const uniform FVector Normal = SetVector(ConvertedNormal);
+		const uniform FVector3f Position = SetVector(OutPosition);
+		const uniform FVector3f Normal = SetVector(OutNormal);
 
 		OutPositions[i] = TransformPosition(ComponentToLocalSpace, Position);
 		OutNormals[i] = GetSafeNormal(TransformVector(ComponentToLocalSpace, Normal));
diff --git a/Engine/Plugins/Experimental/ChaosCloth/Source/Chaos/Private/ChaosCloth/ChaosClothingSimulationSolver.cpp b/Engine/Plugins/Experimental/ChaosCloth/Source/Chaos/Private/ChaosCloth/ChaosClothingSimulationSolver.cpp
--- a/Engine/Plugins/Experimental/ChaosCloth/Source/Chaos/Private/ChaosCloth/ChaosClothingSimulationSolver.cpp
+++ b/Engine/Plugins/Experimental/ChaosCloth/Source/Chaos/Private/ChaosCloth/ChaosClothingSimulationSolver.cpp
@@ -17,8 +17,6 @@
 #define FRAMEPRO_ENABLED 0
 #endif
 
-using namespace Chaos;
-
 DECLARE_CYCLE_STAT(TEXT("Chaos Cloth Solver Update"), STAT_ChaosClothSolverUpdate, STATGROUP_ChaosCloth);
 DECLARE_CYCLE_STAT(TEXT("Chaos Cloth Solver Update Cloths"), STAT_ChaosClothSolverUpdateCloths, STATGROUP_ChaosCloth);
 DECLARE_CYCLE_STAT(TEXT("Chaos Cloth Solver Update Pre Solver Step"), STAT_ChaosClothSolverUpdatePreSolverStep, STATGROUP_ChaosCloth);
@@ -28,41 +26,45 @@ DECLARE_CYCLE_STAT(TEXT("Chaos Cloth Solver Calculate Bounds"), STAT_ChaosClothS
 DECLARE_CYCLE_STAT(TEXT("Chaos Cloth Solver Particle Pre Simulation Transforms"), STAT_ChaosClothParticlePreSimulationTransforms, STATGROUP_ChaosCloth);
 DECLARE_CYCLE_STAT(TEXT("Chaos Cloth Solver Collision Pre Simulation Transforms"), STAT_ChaosClothCollisionPreSimulationTransforms, STATGROUP_ChaosCloth);
 
-static int32 ChaosClothSolverMinParallelBatchSize = 1000;
-static bool bChaosClothSolverParallelClothPreUpdate = false;  // TODO: Doesn't seem to improve much here. Review this after the ISPC implementation.
-static bool bChaosClothSolverParallelClothUpdate = true;
-static bool bChaosClothSolverParallelClothPostUpdate = true;
-static bool bChaosClothSolverUseImprovedTimeStepSmoothing = true;
-
-#if !UE_BUILD_SHIPPING
-static int32 ChaosClothSolverDebugHitchLength = 0;
-static int32 ChaosClothSolverDebugHitchInterval = 0;
-static bool bChaosClothSolverDisableCollision = false;
-
-FAutoConsoleVariableRef CVarChaosClothSolverMinParallelBatchSize(TEXT("p.ChaosCloth.Solver.MinParallelBatchSize"), ChaosClothSolverMinParallelBatchSize, TEXT("The minimum number of particle to process in parallel batch by the solver."));
-FAutoConsoleVariableRef CVarChaosClothSolverParallelClothPreUpdate(TEXT("p.ChaosCloth.Solver.ParallelClothPreUpdate"), bChaosClothSolverParallelClothPreUpdate, TEXT("Pre-transform the cloth particles for each cloth in parallel."));
-FAutoConsoleVariableRef CVarChaosClothSolverParallelClothUpdate(TEXT("p.ChaosCloth.Solver.ParallelClothUpdate"), bChaosClothSolverParallelClothUpdate, TEXT("Skin the physics mesh and do the other cloth update for each cloth in parallel."));
-FAutoConsoleVariableRef CVarChaosClothSolverParallelClothPostUpdate(TEXT("p.ChaosCloth.Solver.ParallelClothPostUpdate"), bChaosClothSolverParallelClothPostUpdate, TEXT("Pre-transform the cloth particles for each cloth in parallel."));
-FAutoConsoleVariableRef CVarChaosClothSolverDebugHitchLength(TEXT("p.ChaosCloth.Solver.DebugHitchLength"), ChaosClothSolverDebugHitchLength, TEXT("Hitch length in ms. Create artificial hitches to debug simulation jitter. 0 to disable"));
-FAutoConsoleVariableRef CVarChaosClothSolverDebugHitchInterval(TEXT("p.ChaosCloth.Solver.DebugHitchInterval"), ChaosClothSolverDebugHitchInterval, TEXT("Hitch interval in frames. Create artificial hitches to debug simulation jitter. 0 to disable"));
-FAutoConsoleVariableRef CVarChaosClothSolverDisableCollision(TEXT("p.ChaosCloth.Solver.DisableCollision"), bChaosClothSolverDisableCollision, TEXT("Disable all collision particles. Needs reset of the simulation (p.ChaosCloth.Reset)."));
-#endif
-
 #if INTEL_ISPC && !UE_BUILD_SHIPPING
 bool bChaos_PreSimulationTransforms_ISPC_Enabled = true;
 FAutoConsoleVariableRef CVarChaosPreSimulationTransformsISPCEnabled(TEXT("p.Chaos.PreSimulationTransforms.ISPC"), bChaos_PreSimulationTransforms_ISPC_Enabled, TEXT("Whether to use ISPC optimizations in ApplySimulationTransforms"));
 bool bChaos_CalculateBounds_ISPC_Enabled = bChaos_CalculateBounds_ISPC_Enable;  // Disabled by default
 FAutoConsoleVariableRef CVarChaosCalculateBoundsISPCEnabled(TEXT("p.Chaos.CalculateBounds.ISPC"), bChaos_CalculateBounds_ISPC_Enabled, TEXT("Whether to use ISPC optimizations in CalculateBounds"));
 
-static_assert(sizeof(ispc::FVector) == sizeof(Chaos::FVec3), "sizeof(ispc::FVector) != sizeof(FVec3)");
+static_assert(sizeof(ispc::FVector3f) == sizeof(Chaos::Softs::FSolverVec3), "sizeof(ispc::FVector) != sizeof(Chaos::Softs::FSolverVec3)");
+static_assert(sizeof(ispc::FTransform3f) == sizeof(Chaos::Softs::FSolverRigidTransform3), "sizeof(ispc::FVector) != sizeof(Chaos::Softs::FSolverRigidTransform3)");
 #endif
 
-FAutoConsoleVariableRef CVarChaosClothSolverUseImprovedTimeStepSmoothing(TEXT("p.ChaosCloth.Solver.UseImprovedTimeStepSmoothing"), bChaosClothSolverUseImprovedTimeStepSmoothing, TEXT("Use the time step smoothing on input forces only rather than on the entire cloth solver, in order to avoid miscalculating velocities."));
+namespace Chaos
+{
+
+static int32 ClothSolverMinParallelBatchSize = 1000;
+static bool bClothSolverParallelClothPreUpdate = false;  // TODO: Doesn't seem to improve much here. Review this after the ISPC implementation.
+static bool bClothSolverParallelClothUpdate = true;
+static bool bClothSolverParallelClothPostUpdate = true;
+static bool bClothSolverUseImprovedTimeStepSmoothing = true;
 
-namespace ChaosClothingSimulationSolverDefault
+#if !UE_BUILD_SHIPPING
+static int32 ClothSolverDebugHitchLength = 0;
+static int32 ClothSolverDebugHitchInterval = 0;
+static bool bClothSolverDisableCollision = false;
+
+FAutoConsoleVariableRef CVarClothSolverMinParallelBatchSize(TEXT("p.ChaosCloth.Solver.MinParallelBatchSize"), ClothSolverMinParallelBatchSize, TEXT("The minimum number of particle to process in parallel batch by the solver."));
+FAutoConsoleVariableRef CVarClothSolverParallelClothPreUpdate(TEXT("p.ChaosCloth.Solver.ParallelClothPreUpdate"), bClothSolverParallelClothPreUpdate, TEXT("Pre-transform the cloth particles for each cloth in parallel."));
+FAutoConsoleVariableRef CVarClothSolverParallelClothUpdate(TEXT("p.ChaosCloth.Solver.ParallelClothUpdate"), bClothSolverParallelClothUpdate, TEXT("Skin the physics mesh and do the other cloth update for each cloth in parallel."));
+FAutoConsoleVariableRef CVarClothSolverParallelClothPostUpdate(TEXT("p.ChaosCloth.Solver.ParallelClothPostUpdate"), bClothSolverParallelClothPostUpdate, TEXT("Pre-transform the cloth particles for each cloth in parallel."));
+FAutoConsoleVariableRef CVarClothSolverDebugHitchLength(TEXT("p.ChaosCloth.Solver.DebugHitchLength"), ClothSolverDebugHitchLength, TEXT("Hitch length in ms. Create artificial hitches to debug simulation jitter. 0 to disable"));
+FAutoConsoleVariableRef CVarClothSolverDebugHitchInterval(TEXT("p.ChaosCloth.Solver.DebugHitchInterval"), ClothSolverDebugHitchInterval, TEXT("Hitch interval in frames. Create artificial hitches to debug simulation jitter. 0 to disable"));
+FAutoConsoleVariableRef CVarClothSolverDisableCollision(TEXT("p.ChaosCloth.Solver.DisableCollision"), bClothSolverDisableCollision, TEXT("Disable all collision particles. Needs reset of the simulation (p.ChaosCloth.Reset)."));
+#endif
+
+FAutoConsoleVariableRef CVarClothSolverUseImprovedTimeStepSmoothing(TEXT("p.ChaosCloth.Solver.UseImprovedTimeStepSmoothing"), bClothSolverUseImprovedTimeStepSmoothing, TEXT("Use the time step smoothing on input forces only rather than on the entire cloth solver, in order to avoid miscalculating velocities."));
+
+namespace ClothingSimulationSolverDefault
 {
-	static const FVec3 Gravity(0.f, 0.f, -980.665f);  // cm/s^2
-	static const FVec3 WindVelocity(0.f);
+	static const Softs::FSolverVec3 Gravity((Softs::FSolverReal)0., (Softs::FSolverReal)0., (Softs::FSolverReal)-980.665);  // cm/s^2
+	static const Softs::FSolverVec3 WindVelocity((Softs::FSolverReal)0.);
 	static const int32 NumIterations = 1;
 	static const int32 NumSubsteps = 1;
 	static const FRealSingle SelfCollisionThickness = 2.f;
@@ -71,42 +73,42 @@ namespace ChaosClothingSimulationSolverDefault
 	static const FRealSingle DampingCoefficient = 0.01f;
 }
 
-namespace ChaosClothingSimulationSolverConstant
+namespace ClothingSimulationSolverConstant
 {
-	static const FReal WorldScale = 100.f;  // World is in cm, but values like wind speed and density are in SI unit and relates to m.
-	static const FReal StartDeltaTime = 1.f / 30.f;  // Initialize filtered timestep at 30fps
+	static const FRealSingle WorldScale = 100.f;  // World is in cm, but values like wind speed and density are in SI unit and relates to m.
+	static const Softs::FSolverReal StartDeltaTime = (Softs::FSolverReal)(1. / 30.);  // Initialize filtered timestep at 30fps
 }
 
 FClothingSimulationSolver::FClothingSimulationSolver()
-	: OldLocalSpaceLocation(0.f)
-	, LocalSpaceLocation(0.f)
-	, Time(0)
-	, DeltaTime(ChaosClothingSimulationSolverConstant::StartDeltaTime)
-	, NumIterations(ChaosClothingSimulationSolverDefault::NumIterations)
-	, NumSubsteps(ChaosClothingSimulationSolverDefault::NumSubsteps)
+	: OldLocalSpaceLocation(0.)
+	, LocalSpaceLocation(0.)
+	, Time(0.)
+	, DeltaTime(ClothingSimulationSolverConstant::StartDeltaTime)
+	, NumIterations(ClothingSimulationSolverDefault::NumIterations)
+	, NumSubsteps(ClothingSimulationSolverDefault::NumSubsteps)
 	, CollisionParticlesOffset(0)
 	, CollisionParticlesSize(0)
-	, Gravity(ChaosClothingSimulationSolverDefault::Gravity)
-	, WindVelocity(ChaosClothingSimulationSolverDefault::WindVelocity)
-	, LegacyWindAdaption((FReal)0.)
+	, Gravity(ClothingSimulationSolverDefault::Gravity)
+	, WindVelocity(ClothingSimulationSolverDefault::WindVelocity)
+	, LegacyWindAdaption(0.f)
 	, bIsClothGravityOverrideEnabled(false)
 {
-	FPBDParticles LocalParticles;
-	FKinematicGeometryClothParticles TRigidParticles;
+	Softs::FSolverParticles LocalParticles;
+	Softs::FSolverRigidParticles RigidParticles;
 	Evolution.Reset(
-		new FPBDEvolution(
+		new Softs::FPBDEvolution(
 			MoveTemp(LocalParticles),
-			MoveTemp(TRigidParticles),
+			MoveTemp(RigidParticles),
 			{}, // CollisionTriangles
-			ChaosClothingSimulationSolverDefault::NumIterations,
-			ChaosClothingSimulationSolverDefault::CollisionThickness,
-			ChaosClothingSimulationSolverDefault::SelfCollisionThickness,
-			ChaosClothingSimulationSolverDefault::FrictionCoefficient,
-			ChaosClothingSimulationSolverDefault::DampingCoefficient));
+			ClothingSimulationSolverDefault::NumIterations,
+			(Softs::FSolverReal)ClothingSimulationSolverDefault::CollisionThickness,
+			(Softs::FSolverReal)ClothingSimulationSolverDefault::SelfCollisionThickness,
+			(Softs::FSolverReal)ClothingSimulationSolverDefault::FrictionCoefficient,
+			(Softs::FSolverReal)ClothingSimulationSolverDefault::DampingCoefficient));
 
 	// Add simulation groups arrays
 	Evolution->AddArray(&PreSimulationTransforms);
-	Evolution->AddArray(&FictitiousAngularDisplacement);
+	Evolution->AddArray(&FictitiousAngularDisplacements);
 
 	Evolution->Particles().AddArray(&Normals);
 	Evolution->Particles().AddArray(&OldAnimationPositions);
@@ -119,26 +121,26 @@ FClothingSimulationSolver::FClothingSimulationSolver()
 	Evolution->CollisionParticles().AddArray(&CollisionTransforms);
 
 	Evolution->SetKinematicUpdateFunction(
-		[this](FPBDParticles& ParticlesInput, const FReal Dt, const FReal LocalTime, const int32 Index)
+		[this](Softs::FSolverParticles& ParticlesInput, const Softs::FSolverReal Dt, const Softs::FSolverReal LocalTime, const int32 Index)
 		{
-			const FReal Alpha = (LocalTime - Time) / DeltaTime;
-			ParticlesInput.P(Index) = Alpha * AnimationPositions[Index] + (1.f - Alpha) * OldAnimationPositions[Index];  // X is the step initial condition, here it's P that needs to be updated so that constraints works with the correct step target
+			const Softs::FSolverReal Alpha = (LocalTime - Time) / DeltaTime;
+			ParticlesInput.P(Index) = Alpha * AnimationPositions[Index] + ((Softs::FSolverReal)1. - Alpha) * OldAnimationPositions[Index];  // X is the step initial condition, here it's P that needs to be updated so that constraints works with the correct step target
 		});
 
 	Evolution->SetCollisionKinematicUpdateFunction(
-		[this](FKinematicGeometryClothParticles& ParticlesInput, const FReal Dt, const FReal LocalTime, const int32 Index)
+		[this](Softs::FSolverRigidParticles& ParticlesInput, const Softs::FSolverReal Dt, const Softs::FSolverReal LocalTime, const int32 Index)
 		{
 			checkSlow(Dt > SMALL_NUMBER && DeltaTime > SMALL_NUMBER);
-			const FReal Alpha = (LocalTime - Time) / DeltaTime;
-			const FVec3 NewX =
-				Alpha * CollisionTransforms[Index].GetTranslation() + (1.f - Alpha) * OldCollisionTransforms[Index].GetTranslation();
+			const Softs::FSolverReal Alpha = (LocalTime - Time) / DeltaTime;
+			const Softs::FSolverVec3 NewX =
+				Alpha * CollisionTransforms[Index].GetTranslation() + ((Softs::FSolverReal)1. - Alpha) * OldCollisionTransforms[Index].GetTranslation();
 			ParticlesInput.V(Index) = (NewX - ParticlesInput.X(Index)) / Dt;
 			ParticlesInput.X(Index) = NewX;
-			FRotation3 NewR = FQuat::Slerp(OldCollisionTransforms[Index].GetRotation(), CollisionTransforms[Index].GetRotation(), Alpha);
-			FRotation3 Delta = NewR * ParticlesInput.R(Index).Inverse();
-			FReal Angle = Delta.GetAngle();
-			FVec3 Axis = Delta.GetRotationAxis();
-			ParticlesInput.W(Index) = (FVec3)Axis * Angle / Dt;
+			const Softs::FSolverRotation3 NewR = Softs::FSolverRotation3::Slerp(OldCollisionTransforms[Index].GetRotation(), CollisionTransforms[Index].GetRotation(), Alpha);
+			const Softs::FSolverRotation3 Delta = NewR * ParticlesInput.R(Index).Inverse();
+			const Softs::FSolverReal Angle = Delta.GetAngle();
+			const Softs::FSolverVec3 Axis = Delta.GetRotationAxis();
+			ParticlesInput.W(Index) = (Softs::FSolverVec3)Axis * Angle / Dt;
 			ParticlesInput.R(Index) = NewR;
 		});
 }
@@ -332,37 +334,52 @@ void FClothingSimulationSolver::EnableParticles(int32 Offset, bool bEnable)
 	GetClothConstraints(Offset).Enable(bEnable);
 }
 
-const FVec3* FClothingSimulationSolver::GetParticlePs(int32 Offset) const
+void FClothingSimulationSolver::ResetStartPose(int32 Offset, int32 NumParticles)
+{
+	Softs::FSolverVec3* const Ps = GetParticlePs(Offset);
+	Softs::FSolverVec3* const Xs = GetParticleXs(Offset);
+	Softs::FSolverVec3* const Vs = GetParticleVs(Offset);
+	const Softs::FSolverVec3* const Positions = GetAnimationPositions(Offset);
+	Softs::FSolverVec3* const OldPositions = GetOldAnimationPositions(Offset);
+
+	for (int32 Index = 0; Index < NumParticles; ++Index)
+	{
+		Ps[Index] = Xs[Index] = OldPositions[Index] = Positions[Index];
+		Vs[Index] = Softs::FSolverVec3(0.);
+	}
+}
+
+const Softs::FSolverVec3* FClothingSimulationSolver::GetParticlePs(int32 Offset) const
 {
 	return &Evolution->Particles().P(Offset);
 }
 
-FVec3* FClothingSimulationSolver::GetParticlePs(int32 Offset)
+Softs::FSolverVec3* FClothingSimulationSolver::GetParticlePs(int32 Offset)
 {
 	return &Evolution->Particles().P(Offset);
 }
 
-const FVec3* FClothingSimulationSolver::GetParticleXs(int32 Offset) const
+const Softs::FSolverVec3* FClothingSimulationSolver::GetParticleXs(int32 Offset) const
 {
 	return &Evolution->Particles().X(Offset);
 }
 
-FVec3* FClothingSimulationSolver::GetParticleXs(int32 Offset)
+Softs::FSolverVec3* FClothingSimulationSolver::GetParticleXs(int32 Offset)
 {
 	return &Evolution->Particles().X(Offset);
 }
 
-const FVec3* FClothingSimulationSolver::GetParticleVs(int32 Offset) const
+const Softs::FSolverVec3* FClothingSimulationSolver::GetParticleVs(int32 Offset) const
 {
 	return &Evolution->Particles().V(Offset);
 }
 
-FVec3* FClothingSimulationSolver::GetParticleVs(int32 Offset)
+Softs::FSolverVec3* FClothingSimulationSolver::GetParticleVs(int32 Offset)
 {
 	return &Evolution->Particles().V(Offset);
 }
 
-const FReal* FClothingSimulationSolver::GetParticleInvMasses(int32 Offset) const
+const Softs::FSolverReal* FClothingSimulationSolver::GetParticleInvMasses(int32 Offset) const
 {
 	return &Evolution->Particles().InvM(Offset);
 }
@@ -402,13 +419,13 @@ int32 FClothingSimulationSolver::AddCollisionParticles(int32 NumCollisionParticl
 	const int32 Offset = Evolution->AddCollisionParticleRange(NumCollisionParticles, GroupId, /*bActivate =*/ false);
 
 	// Always initialize the collision particle's transforms as otherwise setting the geometry will get NaNs detected during the bounding box updates
-	FRotation3* const Rs = GetCollisionParticleRs(Offset);
-	FVec3* const Xs = GetCollisionParticleXs(Offset);
+	Softs::FSolverRotation3* const Rs = GetCollisionParticleRs(Offset);
+	Softs::FSolverVec3* const Xs = GetCollisionParticleXs(Offset);
 
 	for (int32 Index = 0; Index < NumCollisionParticles; ++Index)
 	{
-		Xs[Index] = FVec3(0.f);
-		Rs[Index] = FRotation3::FromIdentity();
+		Xs[Index] = Softs::FSolverVec3(0.);
+		Rs[Index] = Softs::FSolverRotation3::FromIdentity();
 	}
 
 	// Always starts with particles disabled
@@ -420,7 +437,7 @@ int32 FClothingSimulationSolver::AddCollisionParticles(int32 NumCollisionParticl
 void FClothingSimulationSolver::EnableCollisionParticles(int32 Offset, bool bEnable)
 {
 #if !UE_BUILD_SHIPPING
-	if (bChaosClothSolverDisableCollision)
+	if (bClothSolverDisableCollision)
 	{
 		Evolution->ActivateCollisionParticleRange(Offset, false);
 	}
@@ -431,22 +448,37 @@ void FClothingSimulationSolver::EnableCollisionParticles(int32 Offset, bool bEna
 	}
 }
 
-const FVec3* FClothingSimulationSolver::GetCollisionParticleXs(int32 Offset) const
+void FClothingSimulationSolver::ResetCollisionStartPose(int32 Offset, int32 NumCollisionParticles)
+{
+	const Softs::FSolverRigidTransform3* const Transforms = GetCollisionTransforms(Offset);
+	Softs::FSolverRigidTransform3* const OldTransforms = GetOldCollisionTransforms(Offset);
+	Softs::FSolverRotation3* const Rs = GetCollisionParticleRs(Offset);
+	Softs::FSolverVec3* const Xs = GetCollisionParticleXs(Offset);
+
+	for (int32 Index = 0; Index < NumCollisionParticles; ++Index)
+	{
+		OldTransforms[Index] = Transforms[Index];
+		Xs[Index] = Transforms[Index].GetTranslation();
+		Rs[Index] = Transforms[Index].GetRotation();
+	}
+}
+
+const Softs::FSolverVec3* FClothingSimulationSolver::GetCollisionParticleXs(int32 Offset) const
 {
 	return &Evolution->CollisionParticles().X(Offset);
 }
 
-FVec3* FClothingSimulationSolver::GetCollisionParticleXs(int32 Offset)
+Softs::FSolverVec3* FClothingSimulationSolver::GetCollisionParticleXs(int32 Offset)
 {
 	return &Evolution->CollisionParticles().X(Offset);
 }
 
-const FRotation3* FClothingSimulationSolver::GetCollisionParticleRs(int32 Offset) const
+const Softs::FSolverRotation3* FClothingSimulationSolver::GetCollisionParticleRs(int32 Offset) const
 {
 	return &Evolution->CollisionParticles().R(Offset);
 }
 
-FRotation3* FClothingSimulationSolver::GetCollisionParticleRs(int32 Offset)
+Softs::FSolverRotation3* FClothingSimulationSolver::GetCollisionParticleRs(int32 Offset)
 {
 	return &Evolution->CollisionParticles().R(Offset);
 }
@@ -466,76 +498,76 @@ const bool* FClothingSimulationSolver::GetCollisionStatus(int32 Offset) const
 	return Evolution->GetCollisionStatus().GetData() + Offset;
 }
 
-const TArray<FVec3>& FClothingSimulationSolver::GetCollisionContacts() const
+const TArray<Softs::FSolverVec3>& FClothingSimulationSolver::GetCollisionContacts() const
 {
 	return Evolution->GetCollisionContacts();
 }
 
-const TArray<FVec3>& FClothingSimulationSolver::GetCollisionNormals() const
+const TArray<Softs::FSolverVec3>& FClothingSimulationSolver::GetCollisionNormals() const
 {
 	return Evolution->GetCollisionNormals();
 }
 
-void FClothingSimulationSolver::SetParticleMassUniform(int32 Offset, FReal UniformMass, FReal MinPerParticleMass, const FTriangleMesh& Mesh, const TFunctionRef<bool(int32)>& KinematicPredicate)
+void FClothingSimulationSolver::SetParticleMassUniform(int32 Offset, FRealSingle UniformMass, FRealSingle MinPerParticleMass, const FTriangleMesh& Mesh, const TFunctionRef<bool(int32)>& KinematicPredicate)
 {
 	// Retrieve the particle block size
 	const int32 Size = Evolution->GetParticleRangeSize(Offset);
 
 	// Set mass from uniform mass
 	const TSet<int32> Vertices = Mesh.GetVertices();
-	FPBDParticles& Particles = Evolution->Particles();
+	Softs::FSolverParticles& Particles = Evolution->Particles();
 	for (int32 Index = Offset; Index < Offset + Size; ++Index)
 	{
-		Particles.M(Index) = Vertices.Contains(Index) ? UniformMass : 0.f;
+		Particles.M(Index) = Vertices.Contains(Index) ? (Softs::FSolverReal)UniformMass : (Softs::FSolverReal)0.;
 	}
 
 	// Clamp and enslave
-	ParticleMassClampAndEnslave(Offset, Size, MinPerParticleMass, KinematicPredicate);
+	ParticleMassClampAndEnslave(Offset, Size, (Softs::FSolverReal)MinPerParticleMass, KinematicPredicate);
 }
 
-void FClothingSimulationSolver::SetParticleMassFromTotalMass(int32 Offset, FReal TotalMass, FReal MinPerParticleMass, const FTriangleMesh& Mesh, const TFunctionRef<bool(int32)>& KinematicPredicate)
+void FClothingSimulationSolver::SetParticleMassFromTotalMass(int32 Offset, FRealSingle TotalMass, FRealSingle MinPerParticleMass, const FTriangleMesh& Mesh, const TFunctionRef<bool(int32)>& KinematicPredicate)
 {
 	// Retrieve the particle block size
 	const int32 Size = Evolution->GetParticleRangeSize(Offset);
 
 	// Set mass per area
-	const FReal TotalArea = SetParticleMassPerArea(Offset, Size, Mesh);
+	const Softs::FSolverReal TotalArea = SetParticleMassPerArea(Offset, Size, Mesh);
 
 	// Find density
-	const FReal Density = TotalArea > 0.f ? TotalMass / TotalArea : 1.f;
+	const Softs::FSolverReal Density = TotalArea > (Softs::FSolverReal)0. ? (Softs::FSolverReal)TotalMass / TotalArea : (Softs::FSolverReal)1.;
 
 	// Update mass from mesh and density
 	ParticleMassUpdateDensity(Mesh, Density);
 
 	// Clamp and enslave
-	ParticleMassClampAndEnslave(Offset, Size, MinPerParticleMass, KinematicPredicate);
+	ParticleMassClampAndEnslave(Offset, Size, (Softs::FSolverReal)MinPerParticleMass, KinematicPredicate);
 }
 
-void FClothingSimulationSolver::SetParticleMassFromDensity(int32 Offset, FReal Density, FReal MinPerParticleMass, const FTriangleMesh& Mesh, const TFunctionRef<bool(int32)>& KinematicPredicate)
+void FClothingSimulationSolver::SetParticleMassFromDensity(int32 Offset, FRealSingle Density, FRealSingle MinPerParticleMass, const FTriangleMesh& Mesh, const TFunctionRef<bool(int32)>& KinematicPredicate)
 {
 	// Retrieve the particle block size
 	const int32 Size = Evolution->GetParticleRangeSize(Offset);
 
 	// Set mass per area
-	const FReal TotalArea = SetParticleMassPerArea(Offset, Size, Mesh);
+	const Softs::FSolverReal TotalArea = SetParticleMassPerArea(Offset, Size, Mesh);
 
 	// Set density from cm2 to m2
-	Density /= FMath::Square(ChaosClothingSimulationSolverConstant::WorldScale);
+	const Softs::FSolverReal DensityScaled = (Softs::FSolverReal)(Density / FMath::Square(ClothingSimulationSolverConstant::WorldScale));
 
 	// Update mass from mesh and density
-	ParticleMassUpdateDensity(Mesh, Density);
+	ParticleMassUpdateDensity(Mesh, DensityScaled);
 
 	// Clamp and enslave
-	ParticleMassClampAndEnslave(Offset, Size, MinPerParticleMass, KinematicPredicate);
+	ParticleMassClampAndEnslave(Offset, Size, (Softs::FSolverReal)MinPerParticleMass, KinematicPredicate);
 }
 
 void FClothingSimulationSolver::SetReferenceVelocityScale(
 	uint32 GroupId,
-	const FRigidTransform3& OldReferenceSpaceTransform,
+	const FRigidTransform3& OldReferenceSpaceTransform,  // Transforms are in world space so have to be FReal based for LWC
 	const FRigidTransform3& ReferenceSpaceTransform,
-	const FVec3& LinearVelocityScale,
-	FReal AngularVelocityScale,
-	FReal FictitiousAngularScale)
+	const TVec3<FRealSingle>& LinearVelocityScale,
+	FRealSingle AngularVelocityScale,
+	FRealSingle FictitiousAngularScale)
 {
 	FRigidTransform3 OldRootBoneLocalTransform = OldReferenceSpaceTransform;
 	OldRootBoneLocalTransform.AddToTranslation(-OldLocalSpaceLocation);
@@ -544,60 +576,64 @@ void FClothingSimulationSolver::SetReferenceVelocityScale(
 	const FRigidTransform3 DeltaTransform = ReferenceSpaceTransform.GetRelativeTransform(OldReferenceSpaceTransform);
 
 	// Apply linear velocity scale
-	const FVec3 LinearRatio = FVec3(1.f) - LinearVelocityScale.BoundToBox(FVec3(0.f), FVec3(1.f));
+	const FVec3 LinearRatio = FVec3(1.) - FVec3(LinearVelocityScale).BoundToBox(FVec3(0.), FVec3(1.));
 	const FVec3 DeltaPosition = LinearRatio * DeltaTransform.GetTranslation();
 
 	// Apply angular velocity scale
 	FRotation3 DeltaRotation = DeltaTransform.GetRotation();
 	FReal DeltaAngle = DeltaRotation.GetAngle();
 	FVec3 Axis = DeltaRotation.GetRotationAxis();
-	if (DeltaAngle > PI)
+	if (DeltaAngle > (FReal)PI)
 	{
-		DeltaAngle -= 2.f * PI;
+		DeltaAngle -= (FReal)2. * (FReal)PI;
 	}
 
-	const FReal PartialDeltaAngle = DeltaAngle * FMath::Clamp(1.f - AngularVelocityScale, 0.f, 1.f);
+	const FReal PartialDeltaAngle = DeltaAngle * FMath::Clamp((FReal)1. - (FReal)AngularVelocityScale, (FReal)0., (FReal)1.);
 	DeltaRotation = UE::Math::TQuat<FReal>(Axis, PartialDeltaAngle);
 
-	// Transform points back into the previous frame of reference before applying the adjusted deltas 
-	PreSimulationTransforms[GroupId] = OldRootBoneLocalTransform.Inverse() * FRigidTransform3(DeltaPosition, DeltaRotation) * OldRootBoneLocalTransform;
+	// Transform points back into the previous frame of reference before applying the adjusted deltas
+	const FRigidTransform3 PreSimulationTransform = OldRootBoneLocalTransform.Inverse() * FRigidTransform3(DeltaPosition, DeltaRotation) * OldRootBoneLocalTransform;
+	PreSimulationTransforms[GroupId] = Softs::FSolverRigidTransform3(  // Store the delta in solver precision, no need for LWC here
+		Softs::FSolverVec3(PreSimulationTransform.GetTranslation()),
+		Softs::FSolverRotation3(PreSimulationTransform.GetRotation()));
 
 	// Save the reference bone relative angular velocity for calculating the fictitious forces
-	FictitiousAngularDisplacement[GroupId] = ReferenceSpaceTransform.TransformVector(Axis * PartialDeltaAngle * FMath::Min(2.f, FictitiousAngularScale));  // Clamp to 2x the delta angle
+	const FVec3 FictitiousAngularDisplacement = ReferenceSpaceTransform.TransformVector(Axis * PartialDeltaAngle * FMath::Min((FReal)2., (FReal)FictitiousAngularScale));  // Clamp to 2x the delta angle
+	FictitiousAngularDisplacements[GroupId] = Softs::FSolverVec3(FictitiousAngularDisplacement);
 }
 
-FReal FClothingSimulationSolver::SetParticleMassPerArea(int32 Offset, int32 Size, const FTriangleMesh& Mesh)
+Softs::FSolverReal FClothingSimulationSolver::SetParticleMassPerArea(int32 Offset, int32 Size, const FTriangleMesh& Mesh)
 {
 	// Zero out masses
-	FPBDParticles& Particles = Evolution->Particles();
+	Softs::FSolverParticles& Particles = Evolution->Particles();
 	for (int32 Index = Offset; Index < Offset + Size; ++Index)
 	{
-		Particles.M(Index) = 0.f;
+		Particles.M(Index) = (Softs::FSolverReal)0.;
 	}
 
 	// Assign per particle mass proportional to connected area.
 	const TArray<TVec3<int32>>& SurfaceElements = Mesh.GetSurfaceElements();
-	FReal TotalArea = (FReal)0.f;
+	Softs::FSolverReal TotalArea = (Softs::FSolverReal)0.;
 	for (const TVec3<int32>& Tri : SurfaceElements)
 	{
-		const FReal TriArea = 0.5f * FVec3::CrossProduct(
+		const Softs::FSolverReal TriArea = (Softs::FSolverReal)0.5 * Softs::FSolverVec3::CrossProduct(
 			Particles.X(Tri[1]) - Particles.X(Tri[0]),
 			Particles.X(Tri[2]) - Particles.X(Tri[0])).Size();
 		TotalArea += TriArea;
-		const FReal ThirdTriArea = TriArea / 3.f;
+		const Softs::FSolverReal ThirdTriArea = TriArea / (Softs::FSolverReal)3.;
 		Particles.M(Tri[0]) += ThirdTriArea;
 		Particles.M(Tri[1]) += ThirdTriArea;
 		Particles.M(Tri[2]) += ThirdTriArea;
 	}
 
-	UE_LOG(LogChaosCloth, Verbose, TEXT("Total area: %f, SI total area: %f"), TotalArea, TotalArea / FMath::Square(ChaosClothingSimulationSolverConstant::WorldScale));
+	UE_LOG(LogChaosCloth, Verbose, TEXT("Total area: %f, SI total area: %f"), TotalArea, TotalArea / FMath::Square(ClothingSimulationSolverConstant::WorldScale));
 	return TotalArea;
 }
 
-void FClothingSimulationSolver::ParticleMassUpdateDensity(const FTriangleMesh& Mesh, FReal Density)
+void FClothingSimulationSolver::ParticleMassUpdateDensity(const FTriangleMesh& Mesh, Softs::FSolverReal Density)
 {
 	const TSet<int32> Vertices = Mesh.GetVertices();
-	FPBDParticles& Particles = Evolution->Particles();
+	Softs::FSolverParticles& Particles = Evolution->Particles();
 	FReal TotalMass = 0.f;
 	for (const int32 Vertex : Vertices)
 	{
@@ -608,13 +644,13 @@ void FClothingSimulationSolver::ParticleMassUpdateDensity(const FTriangleMesh& M
 	UE_LOG(LogChaosCloth, Verbose, TEXT("Total mass: %f, "), TotalMass);
 }
 
-void FClothingSimulationSolver::ParticleMassClampAndEnslave(int32 Offset, int32 Size, FReal MinPerParticleMass, const TFunctionRef<bool(int32)>& KinematicPredicate)
+void FClothingSimulationSolver::ParticleMassClampAndEnslave(int32 Offset, int32 Size, Softs::FSolverReal MinPerParticleMass, const TFunctionRef<bool(int32)>& KinematicPredicate)
 {
-	FPBDParticles& Particles = Evolution->Particles();
+	Softs::FSolverParticles& Particles = Evolution->Particles();
 	for (int32 Index = Offset; Index < Offset + Size; ++Index)
 	{
-		Particles.M(Index) = FMath::Max(Particles.M(Index), (FReal)MinPerParticleMass);
-		Particles.InvM(Index) = KinematicPredicate(Index - Offset) ? 0.f : 1.f / Particles.M(Index);
+		Particles.M(Index) = FMath::Max(Particles.M(Index), MinPerParticleMass);
+		Particles.InvM(Index) = KinematicPredicate(Index - Offset) ? (Softs::FSolverReal)0. : (Softs::FSolverReal)1. / Particles.M(Index);
 	}
 }
 
@@ -630,36 +666,36 @@ void FClothingSimulationSolver::SetUseCCD(uint32 GroupId, bool bUseCCD)
 	Evolution->SetUseCCD(bUseCCD, GroupId);
 }
 
-void FClothingSimulationSolver::SetGravity(uint32 GroupId, const FVec3& InGravity)
+void FClothingSimulationSolver::SetGravity(uint32 GroupId, const TVec3<FRealSingle>& InGravity)
 {
-	Evolution->GetGravityForces(GroupId).SetAcceleration(InGravity);
+	Evolution->SetGravity(Softs::FSolverVec3(InGravity), GroupId);
 }
 
-void FClothingSimulationSolver::SetWindVelocity(const FVec3& InWindVelocity, FRealSingle InLegacyWindAdaption)
+void FClothingSimulationSolver::SetWindVelocity(const TVec3<FRealSingle>& InWindVelocity, FRealSingle InLegacyWindAdaption)
 {
-	WindVelocity = InWindVelocity * ChaosClothingSimulationSolverConstant::WorldScale;
-	LegacyWindAdaption = (FReal)InLegacyWindAdaption;
+	WindVelocity = InWindVelocity * ClothingSimulationSolverConstant::WorldScale;
+	LegacyWindAdaption = InLegacyWindAdaption;
 }
 
-void FClothingSimulationSolver::SetWindVelocity(uint32 GroupId, const FVec3& InWindVelocity)
+void FClothingSimulationSolver::SetWindVelocity(uint32 GroupId, const TVec3<FRealSingle>& InWindVelocity)
 {
-	FVelocityField& VelocityField = Evolution->GetVelocityField(GroupId);
-	VelocityField.SetVelocity(InWindVelocity);
+	Softs::FVelocityField& VelocityField = Evolution->GetVelocityField(GroupId);
+	VelocityField.SetVelocity(Softs::FSolverVec3(InWindVelocity));
 }
 
 void FClothingSimulationSolver::SetWindGeometry(uint32 GroupId, const FTriangleMesh& TriangleMesh, const TConstArrayView<FRealSingle>& DragMultipliers, const TConstArrayView<FRealSingle>& LiftMultipliers)
 {
-	FVelocityField& VelocityField = Evolution->GetVelocityField(GroupId);
+	Softs::FVelocityField& VelocityField = Evolution->GetVelocityField(GroupId);
 	VelocityField.SetGeometry(&TriangleMesh, DragMultipliers, LiftMultipliers);
 }
 
-void FClothingSimulationSolver::SetWindProperties(uint32 GroupId, const FVec2& Drag, const FVec2& Lift, FReal AirDensity)
+void FClothingSimulationSolver::SetWindProperties(uint32 GroupId, const TVec2<FRealSingle>& Drag, const TVec2<FRealSingle>& Lift, FRealSingle AirDensity)
 {
-	FVelocityField& VelocityField = Evolution->GetVelocityField(GroupId);
+	Softs::FVelocityField& VelocityField = Evolution->GetVelocityField(GroupId);
 	VelocityField.SetProperties(Drag, Lift, AirDensity);
 }
 
-const FVelocityField& FClothingSimulationSolver::GetWindVelocityField(uint32 GroupId)
+const Softs::FVelocityField& FClothingSimulationSolver::GetWindVelocityField(uint32 GroupId)
 {
 	return Evolution->GetVelocityField(GroupId);
 }
@@ -671,14 +707,14 @@ void FClothingSimulationSolver::AddExternalForces(uint32 GroupId, bool bUseLegac
 		const bool bHasVelocityField = !PerSolverField.GetOutputResults(EFieldCommandOutputType::LinearVelocity).IsEmpty();
 		const bool bHasForceField = !PerSolverField.GetOutputResults(EFieldCommandOutputType::LinearForce).IsEmpty();
 
-		const FVec3& AngularDisplacement = FictitiousAngularDisplacement[GroupId];
+		const FVec3& AngularDisplacement = FictitiousAngularDisplacements[GroupId];
 		const bool bHasFictitiousForces = !AngularDisplacement.IsNearlyZero();
 
 		static const FReal LegacyWindMultiplier = (FReal)25.;
 		const FVec3 LegacyWindVelocity = WindVelocity * LegacyWindMultiplier;
 
 		Evolution->GetForceFunction(GroupId) =
-			[this, bHasVelocityField, bHasForceField, bHasFictitiousForces, bUseLegacyWind, LegacyWindVelocity, AngularDisplacement](FPBDParticles& Particles, const FReal Dt, const int32 Index)
+			[this, bHasVelocityField, bHasForceField, bHasFictitiousForces, bUseLegacyWind, LegacyWindVelocity, AngularDisplacement](Softs::FSolverParticles& Particles, const FReal Dt, const int32 Index)
 			{
 				FVec3 Forces((FReal)0.);
 
@@ -732,41 +768,41 @@ void FClothingSimulationSolver::AddExternalForces(uint32 GroupId, bool bUseLegac
 void FClothingSimulationSolver::ApplyPreSimulationTransforms()
 {
 	TRACE_CPUPROFILER_EVENT_SCOPE(FClothingSimulationSolver_ApplyPreSimulationTransforms);
-	const FVec3 DeltaLocalSpaceLocation = LocalSpaceLocation - OldLocalSpaceLocation;
+	const Softs::FSolverVec3 DeltaLocalSpaceLocation(LocalSpaceLocation - OldLocalSpaceLocation);  // LWC, note that LocalSpaceLocation is FReal based. but the delta can be stored in FSolverReal
 
-	const TPBDActiveView<FPBDParticles>& ParticlesActiveView = Evolution->ParticlesActiveView();
+	const TPBDActiveView<Softs::FSolverParticles>& ParticlesActiveView = Evolution->ParticlesActiveView();
 	const TArray<uint32>& ParticleGroupIds = Evolution->ParticleGroupIds();
 
 	ParticlesActiveView.RangeFor(
-		[this, &ParticleGroupIds, &DeltaLocalSpaceLocation](FPBDParticles& Particles, int32 Offset, int32 Range)
+		[this, &ParticleGroupIds, &DeltaLocalSpaceLocation](Softs::FSolverParticles& Particles, int32 Offset, int32 Range)
 		{
 			TRACE_CPUPROFILER_EVENT_SCOPE(FClothingSimulationSolver_ParticlePreSimulationTransforms);
 			SCOPE_CYCLE_COUNTER(STAT_ChaosClothParticlePreSimulationTransforms);
 
 			const int32 RangeSize = Range - Offset;
 
-			if (bChaos_PreSimulationTransforms_ISPC_Enabled)
-			{
 #if INTEL_ISPC
+			if (bRealTypeCompatibleWithISPC && bChaos_PreSimulationTransforms_ISPC_Enabled)  // TODO: Make the ISPC works with both Single and Double depending on the FSolverReal type
+			{
 				ispc::ApplyPreSimulationTransforms(
-					(ispc::FVector*)Particles.GetP().GetData(),
-					(ispc::FVector*)Particles.GetV().GetData(),
-					(ispc::FVector*)Particles.XArray().GetData(),
-					(ispc::FVector*)OldAnimationPositions.GetData(),
+					(ispc::FVector3f*)Particles.GetP().GetData(),
+					(ispc::FVector3f*)Particles.GetV().GetData(),
+					(ispc::FVector3f*)Particles.XArray().GetData(),
+					(ispc::FVector3f*)OldAnimationPositions.GetData(),
 					ParticleGroupIds.GetData(),
-					(ispc::FTransform*)PreSimulationTransforms.GetData(),
-					(ispc::FVector&)DeltaLocalSpaceLocation,
+					(ispc::FTransform3f*)PreSimulationTransforms.GetData(),
+					(ispc::FVector3f&)DeltaLocalSpaceLocation,
 					Offset,
 					Range);
-#endif
 			}
 			else
+#endif
 			{
 				PhysicsParallelFor(RangeSize,
 					[this, &ParticleGroupIds, &DeltaLocalSpaceLocation, &Particles, Offset](int32 i)
 					{
 						const int32 Index = Offset + i;
-						const FRigidTransform3& GroupSpaceTransform = PreSimulationTransforms[ParticleGroupIds[Index]];
+						const Softs::FSolverRigidTransform3& GroupSpaceTransform = PreSimulationTransforms[ParticleGroupIds[Index]];
 
 						// Update initial state for particles
 						Particles.P(Index) = Particles.X(Index) = GroupSpaceTransform.TransformPositionNoScale(Particles.X(Index)) - DeltaLocalSpaceLocation;
@@ -774,26 +810,26 @@ void FClothingSimulationSolver::ApplyPreSimulationTransforms()
 
 						// Update anim initial state (target updated by skinning)
 						OldAnimationPositions[Index] = GroupSpaceTransform.TransformPositionNoScale(OldAnimationPositions[Index]) - DeltaLocalSpaceLocation;
-					}, RangeSize < ChaosClothSolverMinParallelBatchSize);
+					}, RangeSize < ClothSolverMinParallelBatchSize);
 			}
-		}, /*bForceSingleThreaded =*/ !bChaosClothSolverParallelClothPreUpdate);
+		}, /*bForceSingleThreaded =*/ !bClothSolverParallelClothPreUpdate);
 
 #if FRAMEPRO_ENABLED
-	FRAMEPRO_CUSTOM_STAT("ChaosClothSolverMinParallelBatchSize", ChaosClothSolverMinParallelBatchSize, "ChaosClothSolver", "Particles", FRAMEPRO_COLOUR(128,0,255));
-	FRAMEPRO_CUSTOM_STAT("ChaosClothSolverParallelClothPreUpdate", bChaosClothSolverParallelClothPreUpdate, "ChaosClothSolver", "Enabled", FRAMEPRO_COLOUR(128, 128, 64));
+	FRAMEPRO_CUSTOM_STAT("ChaosClothSolverMinParallelBatchSize", ClothSolverMinParallelBatchSize, "ChaosClothSolver", "Particles", FRAMEPRO_COLOUR(128,0,255));
+	FRAMEPRO_CUSTOM_STAT("ChaosClothSolverParallelClothPreUpdate", bClothSolverParallelClothPreUpdate, "ChaosClothSolver", "Enabled", FRAMEPRO_COLOUR(128, 128, 64));
 #endif
 
 	{
 		TRACE_CPUPROFILER_EVENT_SCOPE(FClothingSimulationSolver_CollisionPreSimulationTransforms);
 		SCOPE_CYCLE_COUNTER(STAT_ChaosClothCollisionPreSimulationTransforms);
 
-		const TPBDActiveView<FKinematicGeometryClothParticles>& CollisionParticlesActiveView = Evolution->CollisionParticlesActiveView();
+		const TPBDActiveView<Softs::FSolverRigidParticles>& CollisionParticlesActiveView = Evolution->CollisionParticlesActiveView();
 		const TArray<uint32>& CollisionParticleGroupIds = Evolution->CollisionParticleGroupIds();
 
 		CollisionParticlesActiveView.SequentialFor(  // There's unlikely to ever have enough collision particles for a parallel for
-			[this, &CollisionParticleGroupIds, &DeltaLocalSpaceLocation](FKinematicGeometryClothParticles& CollisionParticles, int32 Index)
+			[this, &CollisionParticleGroupIds, &DeltaLocalSpaceLocation](Softs::FSolverRigidParticles& CollisionParticles, int32 Index)
 			{
-				const FRigidTransform3& GroupSpaceTransform = PreSimulationTransforms[CollisionParticleGroupIds[Index]];
+				const Softs::FSolverRigidTransform3& GroupSpaceTransform = PreSimulationTransforms[CollisionParticleGroupIds[Index]];
 
 				// Update initial state for collisions
 				OldCollisionTransforms[Index] = OldCollisionTransforms[Index] * GroupSpaceTransform;
@@ -825,17 +861,17 @@ void FClothingSimulationSolver::UpdateSolverField()
 	}
 }
 
-void FClothingSimulationSolver::Update(FReal InDeltaTime)
+void FClothingSimulationSolver::Update(Softs::FSolverReal InDeltaTime)
 {
 	TRACE_CPUPROFILER_EVENT_SCOPE(FClothingSimulationSolver_Update);
 	SCOPE_CYCLE_COUNTER(STAT_ChaosClothSolverUpdate);
 
-	if (!bChaosClothSolverUseImprovedTimeStepSmoothing)
+	if (!bClothSolverUseImprovedTimeStepSmoothing)
 	{
 		// Filter delta time to smoothen time variations and prevent unwanted vibrations
 		// Note: This is now deprecated and replaced by in solver input force timestep smoothing
-		static const FReal DeltaTimeDecay = 0.1f;
-		const FReal PrevDeltaTime = DeltaTime;
+		constexpr Softs::FSolverReal DeltaTimeDecay = (Softs::FSolverReal)0.1;
+		const Softs::FSolverReal PrevDeltaTime = DeltaTime;
 		DeltaTime = DeltaTime + (InDeltaTime - DeltaTime) * DeltaTimeDecay;
 	}
 	else
@@ -846,14 +882,14 @@ void FClothingSimulationSolver::Update(FReal InDeltaTime)
 
 #if !UE_BUILD_SHIPPING
 	// Introduce artificial hitches for debugging any simulation jitter
-	if (ChaosClothSolverDebugHitchLength && ChaosClothSolverDebugHitchInterval)
+	if (ClothSolverDebugHitchLength && ClothSolverDebugHitchInterval)
 	{
 		static int32 HitchCounter = 0;
 		if (--HitchCounter < 0)
 		{
-			UE_LOG(LogChaosCloth, Warning, TEXT("Hitching for %dms"), ChaosClothSolverDebugHitchLength);
-			FPlatformProcess::Sleep((float)ChaosClothSolverDebugHitchLength * 0.001f);
-			HitchCounter = ChaosClothSolverDebugHitchInterval;
+			UE_LOG(LogChaosCloth, Warning, TEXT("Hitching for %dms"), ClothSolverDebugHitchLength);
+			FPlatformProcess::Sleep((float)ClothSolverDebugHitchLength * 0.001f);
+			HitchCounter = ClothSolverDebugHitchInterval;
 		}
 	}
 #endif  // #if !UE_BUILD_SHIPPING
@@ -885,11 +921,11 @@ void FClothingSimulationSolver::Update(FReal InDeltaTime)
 			const uint32 GroupId = Cloth->GetGroupId();
 
 			// Pre-update overridable solver properties first
-			Evolution->GetGravityForces(GroupId).SetAcceleration(Gravity);
+			Evolution->SetGravity(Gravity, GroupId);
 			Evolution->GetVelocityField(GroupId).SetVelocity(WindVelocity);
 
 			Cloth->Update(this);
-		}, /*bForceSingleThreaded =*/ !bChaosClothSolverParallelClothUpdate);
+		}, /*bForceSingleThreaded =*/ !bClothSolverParallelClothUpdate);
 	}
 
 	// Pre solver step, apply group space transforms for teleport and linear/delta ratios, ...etc
@@ -912,7 +948,7 @@ void FClothingSimulationSolver::Update(FReal InDeltaTime)
 	
 		for (int32 i = 0; i < NumSubsteps; ++i)
 		{
-			Evolution->AdvanceOneTimeStep(SubstepDeltaTime, bChaosClothSolverUseImprovedTimeStepSmoothing);
+			Evolution->AdvanceOneTimeStep(SubstepDeltaTime, bClothSolverUseImprovedTimeStepSmoothing);
 		}
 
 		Time = Evolution->GetTime();
@@ -929,7 +965,7 @@ void FClothingSimulationSolver::Update(FReal InDeltaTime)
 		{
 			FClothingSimulationCloth* const Cloth = Cloths[ClothIndex];
 			Cloth->PostUpdate(this);
-		}, /*bForceSingleThreaded =*/ !bChaosClothSolverParallelClothPostUpdate);
+		}, /*bForceSingleThreaded =*/ !bClothSolverParallelClothPostUpdate);
 	}
 
 	// Save old space location for next update
@@ -941,66 +977,65 @@ FBoxSphereBounds FClothingSimulationSolver::CalculateBounds() const
 	TRACE_CPUPROFILER_EVENT_SCOPE(FClothingSimulationSolver_CalculateBounds);
 	SCOPE_CYCLE_COUNTER(STAT_ChaosClothSolverCalculateBounds);
 
-	const TPBDActiveView<FPBDParticles>& ParticlesActiveView = Evolution->ParticlesActiveView();
+	const TPBDActiveView<Softs::FSolverParticles>& ParticlesActiveView = Evolution->ParticlesActiveView();
 
 	if (ParticlesActiveView.HasActiveRange())
 	{
 		// Calculate bounding box
-		FAABB3 BoundingBox = FAABB3::EmptyAABB();
+		Softs::FSolverAABB3 BoundingBox = Softs::FSolverAABB3::EmptyAABB();
 
 #if INTEL_ISPC
-		if (bChaos_CalculateBounds_ISPC_Enabled && bRealTypeCompatibleWithISPC)
+		if (bRealTypeCompatibleWithISPC && bChaos_CalculateBounds_ISPC_Enabled)
 		{
 			ParticlesActiveView.RangeFor(
-				[&BoundingBox](FPBDParticles& Particles, int32 Offset, int32 Range)
+				[&BoundingBox](Softs::FSolverParticles& Particles, int32 Offset, int32 Range)
 				{
-					FVec3 NewMin = BoundingBox.Min();
-					FVec3 NewMax = BoundingBox.Max();
+					Softs::FSolverVec3 NewMin = BoundingBox.Min();
+					Softs::FSolverVec3 NewMax = BoundingBox.Max();
 
 					ispc::CalculateBounds(
-						(ispc::FVector&)NewMin,
-						(ispc::FVector&)NewMax,
-						(const ispc::FVector*)Particles.XArray().GetData(),
+						(ispc::FVector3f&)NewMin,
+						(ispc::FVector3f&)NewMax,
+						(const ispc::FVector3f*)Particles.XArray().GetData(),
 						Offset,
 						Range);
 
-					FAABB3 NewAABB(NewMin, NewMax);
-					BoundingBox = NewAABB;
+					BoundingBox.GrowToInclude(Softs::FSolverAABB3(NewMin, NewMax));
 				});
 		}
 		else
 #endif
 		{
 			ParticlesActiveView.SequentialFor(
-				[&BoundingBox](FPBDParticles& Particles, int32 Index)
+				[&BoundingBox](Softs::FSolverParticles& Particles, int32 Index)
 				{
 					BoundingBox.GrowToInclude(Particles.X(Index));
 				});
 		}
 
 		// Calculate (squared) radius
-		const FVec3 Center = BoundingBox.Center();
-		FReal SquaredRadius = 0.f;
+		const Softs::FSolverVec3 Center = BoundingBox.Center();
+		Softs::FSolverReal SquaredRadius = (Softs::FSolverReal)0.;
 
-		if (bChaos_CalculateBounds_ISPC_Enabled)
-		{
 #if INTEL_ISPC
+		if (bRealTypeCompatibleWithISPC && bChaos_CalculateBounds_ISPC_Enabled)
+		{
 			ParticlesActiveView.RangeFor(
-				[&SquaredRadius, &Center](FPBDParticles& Particles, int32 Offset, int32 Range)
+				[&SquaredRadius, &Center](Softs::FSolverParticles& Particles, int32 Offset, int32 Range)
 				{
 					ispc::CalculateSquaredRadius(
 						SquaredRadius,
-						(const ispc::FVector&)Center,
-						(const ispc::FVector*)Particles.XArray().GetData(),
+						(const ispc::FVector3f&)Center,
+						(const ispc::FVector3f*)Particles.XArray().GetData(),
 						Offset,
 						Range);
 				});
-#endif
 		}
 		else
+#endif
 		{
 			ParticlesActiveView.SequentialFor(
-				[&SquaredRadius, &Center](FPBDParticles& Particles, int32 Index)
+				[&SquaredRadius, &Center](Softs::FSolverParticles& Particles, int32 Index)
 				{
 					SquaredRadius = FMath::Max(SquaredRadius, (Particles.X(Index) - Center).SizeSquared());
 				});
@@ -1012,3 +1047,5 @@ FBoxSphereBounds FClothingSimulationSolver::CalculateBounds() const
 
 	return FBoxSphereBounds(LocalSpaceLocation, FVector(0.f), 0.f);
 }
+
+} // End namespace Chaos
diff --git a/Engine/Plugins/Experimental/ChaosCloth/Source/Chaos/Private/ChaosCloth/ChaosClothingSimulationSolver.ispc b/Engine/Plugins/Experimental/ChaosCloth/Source/Chaos/Private/ChaosCloth/ChaosClothingSimulationSolver.ispc
--- a/Engine/Plugins/Experimental/ChaosCloth/Source/Chaos/Private/ChaosCloth/ChaosClothingSimulationSolver.ispc
+++ b/Engine/Plugins/Experimental/ChaosCloth/Source/Chaos/Private/ChaosCloth/ChaosClothingSimulationSolver.ispc
@@ -5,13 +5,13 @@
 #include "Math/Vector.isph"
 #include "Math/Transform.isph"
 
-export void ApplyPreSimulationTransforms(uniform FVector ParticlesP[],
-										uniform FVector ParticlesV[],
-										uniform FVector ParticlesX[],
-										uniform FVector OldAnimationPositions[],
+export void ApplyPreSimulationTransforms(uniform FVector3f ParticlesP[],
+										uniform FVector3f ParticlesV[],
+										uniform FVector3f ParticlesX[],
+										uniform FVector3f OldAnimationPositions[],
 										const uniform uint32 ParticleGroupIds[],
-										const uniform FTransform PreSimulationTransforms[],
-										const uniform FVector &DeltaLocalSpaceLocation,
+										const uniform FTransform3f PreSimulationTransforms[],
+										const uniform FVector3f &DeltaLocalSpaceLocation,
 										const uniform int32 Offset,
 										const uniform int32 Range)
 {
@@ -19,38 +19,38 @@ export void ApplyPreSimulationTransforms(uniform FVector ParticlesP[],
 	{
 		const varying uint32 ParticleGroupId = ParticleGroupIds[i];
 		#pragma ignore warning(perf)
-		const varying FTransform GroupSpaceTransform = PreSimulationTransforms[ParticleGroupId];
+		const varying FTransform3f GroupSpaceTransform = PreSimulationTransforms[ParticleGroupId];
 
 		const uniform int32 Index = extract(i, 0);
 
-		const varying FVector ParticleX = VectorLoad(&ParticlesX[Index]);
-		const varying FVector PX = TransformPosition(GroupSpaceTransform, ParticleX) - DeltaLocalSpaceLocation;
+		const varying FVector3f ParticleX = VectorLoad(&ParticlesX[Index]);
+		const varying FVector3f PX = TransformPosition(GroupSpaceTransform, ParticleX) - DeltaLocalSpaceLocation;
 		VectorStore(&ParticlesP[Index], PX);
 		VectorStore(&ParticlesX[Index], PX);
 
-		const varying FVector ParticleV = VectorLoad(&ParticlesV[Index]);
-		const varying FVector V = TransformVector(GroupSpaceTransform, ParticleV);
+		const varying FVector3f ParticleV = VectorLoad(&ParticlesV[Index]);
+		const varying FVector3f V = TransformVector(GroupSpaceTransform, ParticleV);
 		VectorStore(&ParticlesV[Index], V);
 
-		const varying FVector OldAnimationPosition = VectorLoad(&OldAnimationPositions[Index]);
-		const varying FVector AP = TransformPosition(GroupSpaceTransform, OldAnimationPosition) - DeltaLocalSpaceLocation;
+		const varying FVector3f OldAnimationPosition = VectorLoad(&OldAnimationPositions[Index]);
+		const varying FVector3f AP = TransformPosition(GroupSpaceTransform, OldAnimationPosition) - DeltaLocalSpaceLocation;
 		VectorStore(&OldAnimationPositions[Index], AP);
 	}
 }
 
-export void CalculateBounds(uniform FVector &MMin,
-							uniform FVector &MMax,
-							const uniform FVector ParticlesX[],
+export void CalculateBounds(uniform FVector3f &MMin,
+							uniform FVector3f &MMax,
+							const uniform FVector3f ParticlesX[],
 							const uniform int32 Offset,
 							const uniform int32 Range)
 {
-	varying FVector Min = MMin;
-	varying FVector Max = MMax;
+	varying FVector3f Min = MMin;
+	varying FVector3f Max = MMax;
 
 	foreach(i = Offset ... Range)
 	{
 		const uniform int32 Index = extract(i, 0);
-		const varying FVector X = VectorLoad(&ParticlesX[Index]);
+		const varying FVector3f X = VectorLoad(&ParticlesX[Index]);
 
 		Min = VectorMin(Min, X);
 		Max = VectorMax(Max, X);
@@ -60,18 +60,18 @@ export void CalculateBounds(uniform FVector &MMin,
 	MMax = VectorReduceMax(Max);
 }
 
-export void CalculateSquaredRadius(uniform FReal &SquaredRadius,
-									const uniform FVector &Center,
-									const uniform FVector ParticlesX[],
+export void CalculateSquaredRadius(uniform float &SquaredRadius,
+									const uniform FVector3f &Center,
+									const uniform FVector3f ParticlesX[],
 									const uniform int32 Offset,
 									const uniform int32 Range)
 {
-	varying FReal vSquaredRadius = SquaredRadius;
+	varying float vSquaredRadius = SquaredRadius;
 
 	foreach(i = Offset ... Range)
 	{
 		const uniform int32 Index = extract(i, 0);
-		const varying FVector X = VectorLoad(&ParticlesX[Index]);
+		const varying FVector3f X = VectorLoad(&ParticlesX[Index]);
 
 		vSquaredRadius = max(vSquaredRadius, VectorSizeSquared(X - Center));
 	}
diff --git a/Engine/Plugins/Experimental/ChaosCloth/Source/Chaos/Public/ChaosCloth/ChaosClothConstraints.h b/Engine/Plugins/Experimental/ChaosCloth/Source/Chaos/Public/ChaosCloth/ChaosClothConstraints.h
--- a/Engine/Plugins/Experimental/ChaosCloth/Source/Chaos/Public/ChaosCloth/ChaosClothConstraints.h
+++ b/Engine/Plugins/Experimental/ChaosCloth/Source/Chaos/Public/ChaosCloth/ChaosClothConstraints.h
@@ -2,43 +2,27 @@
 #pragma once
 
 #include "Chaos/Core.h"
-#include "Chaos/Vector.h"
 #include "Containers/ContainersFwd.h"
 #include "Templates/SharedPointer.h"
+#include "Chaos/PBDSoftsEvolutionFwd.h"
 #include "Chaos/PBDLongRangeConstraints.h"
-#include "Chaos/PBDCollisionSpringConstraints.h"
 
 namespace Chaos
 {
-	class FPBDEvolution;
-
-	class FPBDSpringConstraints;
-	class FXPBDSpringConstraints;
-	class FPBDBendingConstraints;
-	class FPBDAxialSpringConstraints;
-	class FXPBDAxialSpringConstraints;
-	class FPBDVolumeConstraint;
-	class FXPBDLongRangeConstraints;
-	class FPBDSphericalConstraint;
-	class FPBDSphericalBackstopConstraint;
-	class FPBDAnimDriveConstraint;
-	class FPBDShapeConstraints;
-	class FPBDCollisionSpringConstraints;
-
 	class FClothConstraints final
 	{
 	public:
-		typedef FPBDLongRangeConstraints::EMode ETetherMode;
+		typedef Softs::FPBDLongRangeConstraints::EMode ETetherMode;
 
 		FClothConstraints();
 		~FClothConstraints();
 
 		// ---- Solver interface ----
 		void Initialize(
-			FPBDEvolution* InEvolution,
-			const TArray<FVec3>& InAnimationPositions,
-			const TArray<FVec3>& InOldAnimationPositions,
-			const TArray<FVec3>& InAnimationNormals,
+			Softs::FPBDEvolution* InEvolution,
+			const TArray<Softs::FSolverVec3>& InAnimationPositions,
+			const TArray<Softs::FSolverVec3>& InOldAnimationPositions,
+			const TArray<Softs::FSolverVec3>& InAnimationNormals,
 			int32 InParticleOffset,
 			int32 InNumParticles);
 		// ---- End of Solver interface ----
@@ -46,77 +30,77 @@ namespace Chaos
 		// ---- Cloth interface ----
 		void SetEdgeConstraints(const TArray<TVec3<int32>>& SurfaceElements, const TConstArrayView<FRealSingle>& StiffnessMultipliers, bool bUseXPBDConstraints);
 		void SetBendingConstraints(const TArray<TVec2<int32>>& Edges, const TConstArrayView<FRealSingle>& StiffnessMultipliers, bool bUseXPBDConstraints);
-		void SetBendingConstraints(TArray<TVec4<int32>>&& BendingElements, FReal BendingStiffness);
+		void SetBendingConstraints(TArray<TVec4<int32>>&& BendingElements, Softs::FSolverReal BendingStiffness);
 		void SetAreaConstraints(const TArray<TVec3<int32>>& SurfaceElements, const TConstArrayView<FRealSingle>& StiffnessMultipliers, bool bUseXPBDConstraints);
-		void SetVolumeConstraints(const TArray<TVec2<int32>>& DoubleBendingEdges, FReal VolumeStiffness);
-		void SetVolumeConstraints(TArray<TVec3<int32>>&& SurfaceElements, FReal VolumeStiffness);
+		void SetVolumeConstraints(const TArray<TVec2<int32>>& DoubleBendingEdges, Softs::FSolverReal VolumeStiffness);
+		void SetVolumeConstraints(TArray<TVec3<int32>>&& SurfaceElements, Softs::FSolverReal VolumeStiffness);
 		void SetLongRangeConstraints(const TArray<TConstArrayView<TTuple<int32, int32, FRealSingle>>>& Tethers,
 			const TConstArrayView<FRealSingle>& TetherStiffnessMultipliers, const TConstArrayView<FRealSingle>& TetherScaleMultipliers,
-			const FVec2& TetherScale, bool bUseXPBDConstraints);
+			const Softs::FSolverVec2& TetherScale, bool bUseXPBDConstraints);
 		void SetMaximumDistanceConstraints(const TConstArrayView<FRealSingle>& MaxDistances);
 		void SetBackstopConstraints(const TConstArrayView<FRealSingle>& BackstopDistances, const TConstArrayView<FRealSingle>& BackstopRadiuses, bool bUseLegacyBackstop);
 		void SetAnimDriveConstraints(const TConstArrayView<FRealSingle>& AnimDriveStiffnessMultipliers, const TConstArrayView<FRealSingle>& AnimDriveDampingMultipliers);
-		void SetShapeTargetConstraints(FReal ShapeTargetStiffness);
-		void SetSelfCollisionConstraints(const TArray<TVec3<int32>>& SurfaceElements, TSet<TVec2<int32>>&& DisabledCollisionElements, FReal SelfCollisionThickness);
+		void SetShapeTargetConstraints(Softs::FSolverReal ShapeTargetStiffness);
+		void SetSelfCollisionConstraints(const TArray<TVec3<int32>>& SurfaceElements, TSet<TVec2<int32>>&& DisabledCollisionElements, Softs::FSolverReal SelfCollisionThickness);
 
 		void CreateRules();
 		void Enable(bool bEnable);
 
-		void SetEdgeProperties(const FVec2& EdgeStiffness);
-		void SetBendingProperties(const FVec2& BendingStiffness);
-		void SetAreaProperties(const FVec2& AreaStiffness);
-		void SetThinShellVolumeProperties(FReal VolumeStiffness);
-		void SetVolumeProperties(FReal VolumeStiffness);
-		void SetLongRangeAttachmentProperties(const FVec2& TetherStiffness, const FVec2& TetherScale);
-		void SetMaximumDistanceProperties(FReal MaxDistancesMultiplier);
-		void SetAnimDriveProperties(const FVec2& AnimDriveStiffness, const FVec2& AnimDriveDamping);
-		void SetSelfCollisionProperties(FReal SelfCollisionThickness);
-		UE_DEPRECATED(5.0, "Use SetBackstopProperties(bool, FReal) instead.")
-		void SetBackstopProperties(bool bEnabled) { SetBackstopProperties(bEnabled, (FReal)1.); }
-		void SetBackstopProperties(bool bEnabled, FReal BackstopDistancesMultiplier);
+		void SetEdgeProperties(const Softs::FSolverVec2& EdgeStiffness);
+		void SetBendingProperties(const Softs::FSolverVec2& BendingStiffness);
+		void SetAreaProperties(const Softs::FSolverVec2& AreaStiffness);
+		void SetThinShellVolumeProperties(Softs::FSolverReal VolumeStiffness);
+		void SetVolumeProperties(Softs::FSolverReal VolumeStiffness);
+		void SetLongRangeAttachmentProperties(const Softs::FSolverVec2& TetherStiffness, const Softs::FSolverVec2& TetherScale);
+		void SetMaximumDistanceProperties(Softs::FSolverReal MaxDistancesMultiplier);
+		void SetAnimDriveProperties(const Softs::FSolverVec2& AnimDriveStiffness, const Softs::FSolverVec2& AnimDriveDamping);
+		void SetSelfCollisionProperties(Softs::FSolverReal SelfCollisionThickness);
+		UE_DEPRECATED(5.0, "Use SetBackstopProperties(bool, FSolverReal) instead.")
+		void SetBackstopProperties(bool bEnabled) { SetBackstopProperties(bEnabled, (Softs::FSolverReal)1.); }
+		void SetBackstopProperties(bool bEnabled, Softs::FSolverReal BackstopDistancesMultiplier);
 		// ---- End of Cloth interface ----
 
 		// ---- Debug functions ----
-		const TSharedPtr<FPBDSpringConstraints> GetEdgeConstraints() const { return EdgeConstraints; }
-		const TSharedPtr<FXPBDSpringConstraints>& GetXEdgeConstraints() const { return XEdgeConstraints; }
-		const TSharedPtr<FPBDSpringConstraints>& GetBendingConstraints() const { return BendingConstraints; }  
-		const TSharedPtr<FXPBDSpringConstraints>& GetXBendingConstraints() const { return XBendingConstraints; }
-		const TSharedPtr<FPBDBendingConstraints>& GetBendingElementConstraints() const { return BendingElementConstraints; }
-		const TSharedPtr<FPBDAxialSpringConstraints>& GetAreaConstraints() const { return AreaConstraints; }
-		const TSharedPtr<FXPBDAxialSpringConstraints>& GetXAreaConstraints() const { return XAreaConstraints; }
-		const TSharedPtr<FPBDSpringConstraints>& GetThinShellVolumeConstraints() const { return ThinShellVolumeConstraints; }
-		const TSharedPtr<FPBDVolumeConstraint>& GetVolumeConstraints() const { return VolumeConstraints; }
-		const TSharedPtr<FPBDLongRangeConstraints>& GetLongRangeConstraints() const { return LongRangeConstraints; }
-		const TSharedPtr<FXPBDLongRangeConstraints>& GetXLongRangeConstraints() const { return XLongRangeConstraints; }
-		const TSharedPtr<FPBDSphericalConstraint>& GetMaximumDistanceConstraints() const { return MaximumDistanceConstraints; }
-		const TSharedPtr<FPBDSphericalBackstopConstraint>& GetBackstopConstraints() const { return BackstopConstraints; }
-		const TSharedPtr<FPBDAnimDriveConstraint>& GetAnimDriveConstraints() const { return AnimDriveConstraints; }
-		const TSharedPtr<FPBDShapeConstraints>& GetShapeConstraints() const { return ShapeConstraints; }
-		const TSharedPtr<FPBDCollisionSpringConstraints>& GetSelfCollisionConstraints() const { return SelfCollisionConstraints; }
+		const TSharedPtr<Softs::FPBDSpringConstraints> GetEdgeConstraints() const { return EdgeConstraints; }
+		const TSharedPtr<Softs::FXPBDSpringConstraints>& GetXEdgeConstraints() const { return XEdgeConstraints; }
+		const TSharedPtr<Softs::FPBDSpringConstraints>& GetBendingConstraints() const { return BendingConstraints; }  
+		const TSharedPtr<Softs::FXPBDSpringConstraints>& GetXBendingConstraints() const { return XBendingConstraints; }
+		const TSharedPtr<Softs::FPBDBendingConstraints>& GetBendingElementConstraints() const { return BendingElementConstraints; }
+		const TSharedPtr<Softs::FPBDAxialSpringConstraints>& GetAreaConstraints() const { return AreaConstraints; }
+		const TSharedPtr<Softs::FXPBDAxialSpringConstraints>& GetXAreaConstraints() const { return XAreaConstraints; }
+		const TSharedPtr<Softs::FPBDSpringConstraints>& GetThinShellVolumeConstraints() const { return ThinShellVolumeConstraints; }
+		const TSharedPtr<Softs::FPBDVolumeConstraint>& GetVolumeConstraints() const { return VolumeConstraints; }
+		const TSharedPtr<Softs::FPBDLongRangeConstraints>& GetLongRangeConstraints() const { return LongRangeConstraints; }
+		const TSharedPtr<Softs::FXPBDLongRangeConstraints>& GetXLongRangeConstraints() const { return XLongRangeConstraints; }
+		const TSharedPtr<Softs::FPBDSphericalConstraint>& GetMaximumDistanceConstraints() const { return MaximumDistanceConstraints; }
+		const TSharedPtr<Softs::FPBDSphericalBackstopConstraint>& GetBackstopConstraints() const { return BackstopConstraints; }
+		const TSharedPtr<Softs::FPBDAnimDriveConstraint>& GetAnimDriveConstraints() const { return AnimDriveConstraints; }
+		const TSharedPtr<Softs::FPBDShapeConstraints>& GetShapeConstraints() const { return ShapeConstraints; }
+		const TSharedPtr<Softs::FPBDCollisionSpringConstraints>& GetSelfCollisionConstraints() const { return SelfCollisionConstraints; }
 		// ---- End of debug functions ----
 
 	private:
-		TSharedPtr<FPBDSpringConstraints> EdgeConstraints;
-		TSharedPtr<FXPBDSpringConstraints> XEdgeConstraints;
-		TSharedPtr<FPBDSpringConstraints> BendingConstraints;
-		TSharedPtr<FXPBDSpringConstraints> XBendingConstraints;
-		TSharedPtr<FPBDBendingConstraints> BendingElementConstraints;
-		TSharedPtr<FPBDAxialSpringConstraints> AreaConstraints;
-		TSharedPtr<FXPBDAxialSpringConstraints> XAreaConstraints;
-		TSharedPtr<FPBDSpringConstraints> ThinShellVolumeConstraints;
-		TSharedPtr<FPBDVolumeConstraint> VolumeConstraints;
-		TSharedPtr<FPBDLongRangeConstraints> LongRangeConstraints;
-		TSharedPtr<FXPBDLongRangeConstraints> XLongRangeConstraints;
-		TSharedPtr<FPBDSphericalConstraint> MaximumDistanceConstraints;
-		TSharedPtr<FPBDSphericalBackstopConstraint> BackstopConstraints;
-		TSharedPtr<FPBDAnimDriveConstraint> AnimDriveConstraints;
-		TSharedPtr<FPBDShapeConstraints> ShapeConstraints;
-		TSharedPtr<FPBDCollisionSpringConstraints> SelfCollisionConstraints;
+		TSharedPtr<Softs::FPBDSpringConstraints> EdgeConstraints;
+		TSharedPtr<Softs::FXPBDSpringConstraints> XEdgeConstraints;
+		TSharedPtr<Softs::FPBDSpringConstraints> BendingConstraints;
+		TSharedPtr<Softs::FXPBDSpringConstraints> XBendingConstraints;
+		TSharedPtr<Softs::FPBDBendingConstraints> BendingElementConstraints;
+		TSharedPtr<Softs::FPBDAxialSpringConstraints> AreaConstraints;
+		TSharedPtr<Softs::FXPBDAxialSpringConstraints> XAreaConstraints;
+		TSharedPtr<Softs::FPBDSpringConstraints> ThinShellVolumeConstraints;
+		TSharedPtr<Softs::FPBDVolumeConstraint> VolumeConstraints;
+		TSharedPtr<Softs::FPBDLongRangeConstraints> LongRangeConstraints;
+		TSharedPtr<Softs::FXPBDLongRangeConstraints> XLongRangeConstraints;
+		TSharedPtr<Softs::FPBDSphericalConstraint> MaximumDistanceConstraints;
+		TSharedPtr<Softs::FPBDSphericalBackstopConstraint> BackstopConstraints;
+		TSharedPtr<Softs::FPBDAnimDriveConstraint> AnimDriveConstraints;
+		TSharedPtr<Softs::FPBDShapeConstraints> ShapeConstraints;
+		TSharedPtr<Softs::FPBDCollisionSpringConstraints> SelfCollisionConstraints;
 		
-		FPBDEvolution* Evolution;
-		const TArray<FVec3>* AnimationPositions;
-		const TArray<FVec3>* OldAnimationPositions;
-		const TArray<FVec3>* AnimationNormals;
+		Softs::FPBDEvolution* Evolution;
+		const TArray<Softs::FSolverVec3>* AnimationPositions;
+		const TArray<Softs::FSolverVec3>* OldAnimationPositions;
+		const TArray<Softs::FSolverVec3>* AnimationNormals;
 
 		int32 ParticleOffset;
 		int32 NumParticles;
diff --git a/Engine/Plugins/Experimental/ChaosCloth/Source/Chaos/Public/ChaosCloth/ChaosClothingSimulation.h b/Engine/Plugins/Experimental/ChaosCloth/Source/Chaos/Public/ChaosCloth/ChaosClothingSimulation.h
--- a/Engine/Plugins/Experimental/ChaosCloth/Source/Chaos/Public/ChaosCloth/ChaosClothingSimulation.h
+++ b/Engine/Plugins/Experimental/ChaosCloth/Source/Chaos/Public/ChaosCloth/ChaosClothingSimulation.h
@@ -18,8 +18,6 @@ class UMaterial;
 namespace Chaos
 {
 	class FTriangleMesh;
-	class FPBDLongRangeConstraintsBase;
-	class FPBDEvolution;
 	class FClothingSimulationSolver;
 	class FClothingSimulationMesh;
 	class FClothingSimulationCloth;
diff --git a/Engine/Plugins/Experimental/ChaosCloth/Source/Chaos/Public/ChaosCloth/ChaosClothingSimulationCloth.h b/Engine/Plugins/Experimental/ChaosCloth/Source/Chaos/Public/ChaosCloth/ChaosClothingSimulationCloth.h
--- a/Engine/Plugins/Experimental/ChaosCloth/Source/Chaos/Public/ChaosCloth/ChaosClothingSimulationCloth.h
+++ b/Engine/Plugins/Experimental/ChaosCloth/Source/Chaos/Public/ChaosCloth/ChaosClothingSimulationCloth.h
@@ -49,8 +49,8 @@ namespace Chaos
 			bool bInUseXPBDConstraints,
 			FRealSingle InGravityScale,
 			bool bIsGravityOverridden,
-			const FVec3& InGravityOverride,
-			const FVec3& InLinearVelocityScale,
+			const TVec3<FRealSingle>& InGravityOverride,
+			const TVec3<FRealSingle>& InLinearVelocityScale,
 			FRealSingle InAngularVelocityScale,
 			FRealSingle InFictitiousAngularScale,
 			const TVec2<FRealSingle>& InDrag,
@@ -104,19 +104,19 @@ namespace Chaos
 
 		// ---- Debugging/visualization functions
 		// Return the solver's input positions for this cloth source current LOD, not thread safe, call must be done right after the solver update.
-		TConstArrayView<FVec3> GetAnimationPositions(const FClothingSimulationSolver* Solver) const;
+		TConstArrayView<Softs::FSolverVec3> GetAnimationPositions(const FClothingSimulationSolver* Solver) const;
 		// Return the solver's input normals for this cloth source current LOD, not thread safe, call must be done right after the solver update.
-		TConstArrayView<FVec3> GetAnimationNormals(const FClothingSimulationSolver* Solver) const;
+		TConstArrayView<Softs::FSolverVec3> GetAnimationNormals(const FClothingSimulationSolver* Solver) const;
 		// Return the solver's positions for this cloth current LOD, not thread safe, call must be done right after the solver update.
-		TConstArrayView<FVec3> GetParticlePositions(const FClothingSimulationSolver* Solver) const;
+		TConstArrayView<Softs::FSolverVec3> GetParticlePositions(const FClothingSimulationSolver* Solver) const;
 		// Return the solver's previous frame positions for this cloth current LOD, not thread safe, call must be done right after the solver update.
-		TConstArrayView<FVec3> GetParticleOldPositions(const FClothingSimulationSolver* Solver) const;
+		TConstArrayView<Softs::FSolverVec3> GetParticleOldPositions(const FClothingSimulationSolver* Solver) const;
 		// Return the solver's normals for this cloth current LOD, not thread safe, call must be done right after the solver update.
-		TConstArrayView<FVec3> GetParticleNormals(const FClothingSimulationSolver* Solver) const;
+		TConstArrayView<Softs::FSolverVec3> GetParticleNormals(const FClothingSimulationSolver* Solver) const;
 		// Return the solver's normals for this cloth current LOD, not thread safe, call must be done right after the solver update.
-		TConstArrayView<FReal> GetParticleInvMasses(const FClothingSimulationSolver* Solver) const;
+		TConstArrayView<Softs::FSolverReal> GetParticleInvMasses(const FClothingSimulationSolver* Solver) const;
 		// Return the current gravity as applied by the solver using the various overrides, not thread safe, call must be done right after the solver update.
-		FVec3 GetGravity(const FClothingSimulationSolver* Solver) const;
+		TVec3<FRealSingle> GetGravity(const FClothingSimulationSolver* Solver) const;
 		// Return the current bounding box based on a given solver, not thread safe, call must be done right after the solver update.
 		FAABB3 CalculateBoundingBox(const FClothingSimulationSolver* Solver) const;
 		// Return the current LOD offset in the solver's particle array, or INDEX_NONE if no LOD is currently selected.
@@ -208,8 +208,8 @@ namespace Chaos
 		bool bUseXPBDConstraints;
 		FRealSingle GravityScale;
 		bool bIsGravityOverridden;
-		FVec3 GravityOverride;
-		FVec3 LinearVelocityScale;  // Linear ratio applied to the reference bone transforms
+		TVec3<FRealSingle> GravityOverride;
+		TVec3<FRealSingle> LinearVelocityScale;  // Linear ratio applied to the reference bone transforms
 		FRealSingle AngularVelocityScale;  // Angular ratio factor applied to the reference bone transforms
 		FRealSingle FictitiousAngularScale;
 		TVec2<FRealSingle> Drag;
diff --git a/Engine/Plugins/Experimental/ChaosCloth/Source/Chaos/Public/ChaosCloth/ChaosClothingSimulationCollider.h b/Engine/Plugins/Experimental/ChaosCloth/Source/Chaos/Public/ChaosCloth/ChaosClothingSimulationCollider.h
--- a/Engine/Plugins/Experimental/ChaosCloth/Source/Chaos/Public/ChaosCloth/ChaosClothingSimulationCollider.h
+++ b/Engine/Plugins/Experimental/ChaosCloth/Source/Chaos/Public/ChaosCloth/ChaosClothingSimulationCollider.h
@@ -3,8 +3,7 @@
 
 #include "ClothCollisionData.h"
 #include "Containers/ContainersFwd.h"
-#include "Chaos/Rotation.h"
-#include "Chaos/ImplicitObject.h"
+#include "Chaos/PBDSoftsEvolutionFwd.h"
 
 class USkeletalMeshComponent;
 class UClothingAssetCommon;
@@ -12,6 +11,8 @@ class FClothingSimulationContextCommon;
 
 namespace Chaos
 {
+	class FImplicitObject;
+
 	class FClothingSimulationSolver;
 	class FClothingSimulationCloth;
 
@@ -54,13 +55,13 @@ namespace Chaos
 
 		// ---- Debugging and visualization functions ----
 		// Return current active LOD collision particles translations, not thread safe, to use after solver update.
-		TConstArrayView<FVec3> GetCollisionTranslations(const FClothingSimulationSolver* Solver, const FClothingSimulationCloth* Cloth, ECollisionDataType CollisionDataType) const;
+		TConstArrayView<Softs::FSolverVec3> GetCollisionTranslations(const FClothingSimulationSolver* Solver, const FClothingSimulationCloth* Cloth, ECollisionDataType CollisionDataType) const;
 
 		// Return current active LOD collision particles rotations, not thread safe, to use after solver update.
-		TConstArrayView<FRotation3> GetCollisionRotations(const FClothingSimulationSolver* Solver, const FClothingSimulationCloth* Cloth, ECollisionDataType CollisionDataType) const;
+		TConstArrayView<Softs::FSolverRotation3> GetCollisionRotations(const FClothingSimulationSolver* Solver, const FClothingSimulationCloth* Cloth, ECollisionDataType CollisionDataType) const;
 
 		// Return current active LOD previous frame collision particles transforms, not thread safe, to use after solver update.
-		TConstArrayView<FRigidTransform3> GetOldCollisionTransforms(const FClothingSimulationSolver* Solver, const FClothingSimulationCloth* Cloth, ECollisionDataType CollisionDataType) const;
+		TConstArrayView<Softs::FSolverRigidTransform3> GetOldCollisionTransforms(const FClothingSimulationSolver* Solver, const FClothingSimulationCloth* Cloth, ECollisionDataType CollisionDataType) const;
 
 		// Return current active LOD collision geometries, not thread safe, to use after solver update.
 		TConstArrayView<TUniquePtr<FImplicitObject>> GetCollisionGeometries(const FClothingSimulationSolver* Solver, const FClothingSimulationCloth* Cloth, ECollisionDataType CollisionDataType) const;
