CIS - Fix non unity build error caused by missing include in PerParticleDampVelocity.cpp.
- Fix static analysis warning on normal conversion in ChaosClothingSimulation.cpp.

diff --git a/Engine/Plugins/Experimental/ChaosCloth/Source/Chaos/Private/ChaosCloth/ChaosClothingSimulation.cpp b/Engine/Plugins/Experimental/ChaosCloth/Source/Chaos/Private/ChaosCloth/ChaosClothingSimulation.cpp
--- a/Engine/Plugins/Experimental/ChaosCloth/Source/Chaos/Private/ChaosCloth/ChaosClothingSimulation.cpp
+++ b/Engine/Plugins/Experimental/ChaosCloth/Source/Chaos/Private/ChaosCloth/ChaosClothingSimulation.cpp
@@ -1570,11 +1570,13 @@ void FClothingSimulation::DebugDrawCollision(FPrimitiveDrawInterface* PDI) const
 
 	// Draw contacts
 	check(Solver->GetCollisionContacts().Num() == Solver->GetCollisionNormals().Num());
+	constexpr FReal NormalLength = (FReal)10.;
+
 	const FVec3& LocalSpaceLocation = Solver->GetLocalSpaceLocation();
 	for (int32 i = 0; i < Solver->GetCollisionContacts().Num(); ++i)
 	{
-		const FVec3& Pos0 = LocalSpaceLocation + Solver->GetCollisionContacts()[i];
-		const FVec3& Normal = Solver->GetCollisionNormals()[i];
+		const FVec3 Pos0 = LocalSpaceLocation + FVec3(Solver->GetCollisionContacts()[i]);
+		const FVec3 Normal = FVec3(Solver->GetCollisionNormals()[i]);
 
 		// Draw contact
 		FVec3 TangentU, TangentV;
@@ -1587,7 +1589,7 @@ void FClothingSimulation::DebugDrawCollision(FPrimitiveDrawInterface* PDI) const
 
 		// Draw normal
 		static const FLinearColor Brown(0.1f, 0.05f, 0.f);
-		const FVec3 Pos1 = Pos0 + 10.f * Normal;
+		const FVec3 Pos1 = Pos0 + NormalLength * Normal;
 		DrawLine(PDI, Pos0, Pos1, Brown);
 	}
 }
diff --git a/Engine/Source/Runtime/Experimental/Chaos/Private/Chaos/PerParticleDampVelocity.cpp b/Engine/Source/Runtime/Experimental/Chaos/Private/Chaos/PerParticleDampVelocity.cpp
--- a/Engine/Source/Runtime/Experimental/Chaos/Private/Chaos/PerParticleDampVelocity.cpp
+++ b/Engine/Source/Runtime/Experimental/Chaos/Private/Chaos/PerParticleDampVelocity.cpp
@@ -4,6 +4,10 @@
 #include "Chaos/Matrix.h"
 #include "GenericPlatform/GenericPlatformMath.h"
 
+#if INTEL_ISPC && !UE_BUILD_SHIPPING
+#include "HAL/IConsoleManager.h"
+#endif
+
 #if INTEL_ISPC
 #include "PerParticleDampVelocity.ispc.generated.h"
 #endif
