Chaos cloth - fix scaled collisions.

This was a regression from switching back to single precision for cloth solves. FTransforms (rotation, translation, scale) were being cast to FSolverRigidTransform3s (rotation, translation) and losing scale.

==CHANGES==
PBDSoftsEvolutionFwd.h
Add new FSolverTransform3 typedef

ChaosClothingSimulationCollider.cpp
cast to FSolverTransform3 when working with FTransforms

==TESTING==
Tested in Cloth map from ContentExamples. Verified collisions are fixed.

diff --git a/Engine/Plugins/Experimental/ChaosCloth/Source/Chaos/Private/ChaosCloth/ChaosClothingSimulationCollider.cpp b/Engine/Plugins/Experimental/ChaosCloth/Source/Chaos/Private/ChaosCloth/ChaosClothingSimulationCollider.cpp
--- a/Engine/Plugins/Experimental/ChaosCloth/Source/Chaos/Private/ChaosCloth/ChaosClothingSimulationCollider.cpp
+++ b/Engine/Plugins/Experimental/ChaosCloth/Source/Chaos/Private/ChaosCloth/ChaosClothingSimulationCollider.cpp
@@ -313,14 +313,14 @@ void FClothingSimulationCollider::FLODData::Update(FClothingSimulationSolver* So
 
 		FTransform ComponentToLocalSpaceReal = Context ? Context->ComponentToWorld : FTransform::Identity;
 		ComponentToLocalSpaceReal.AddToTranslation(-Solver->GetLocalSpaceLocation());
-		const Softs::FSolverRigidTransform3 ComponentToLocalSpace(ComponentToLocalSpaceReal);  // LWC, now in local space, therefore it is safe to use the solver transform type
+		const Softs::FSolverTransform3 ComponentToLocalSpace(ComponentToLocalSpaceReal);  // LWC, now in local space, therefore it is safe to use the solver transform type
 	
 		// Update the collision transforms
 		for (int32 Index = 0; Index < NumGeometries; ++Index)
 		{
 			const int32 BoneIndex = BoneIndices[Index];
 			CollisionTransforms[Index] = BoneTransforms.IsValidIndex(BoneIndex) ?
-				BaseTransforms[Index] * Softs::FSolverRigidTransform3(BoneTransforms[BoneIndex]) * ComponentToLocalSpace :  // LWC requires the BoneTransform cast to the solver underlying type
+				BaseTransforms[Index] * Softs::FSolverTransform3(BoneTransforms[BoneIndex]) * ComponentToLocalSpace :  // LWC requires the BoneTransform cast to the solver underlying type
 				BaseTransforms[Index] * ComponentToLocalSpace;  // External collisions don't map to a bone
 		}
 	}
diff --git a/Engine/Source/Runtime/Experimental/Chaos/Public/Chaos/PBDSoftsEvolutionFwd.h b/Engine/Source/Runtime/Experimental/Chaos/Public/Chaos/PBDSoftsEvolutionFwd.h
--- a/Engine/Source/Runtime/Experimental/Chaos/Public/Chaos/PBDSoftsEvolutionFwd.h
+++ b/Engine/Source/Runtime/Experimental/Chaos/Public/Chaos/PBDSoftsEvolutionFwd.h
@@ -3,6 +3,11 @@
 
 #include "Chaos/GeometryParticlesfwd.h"  // for EGeometryParticlesSimType
 
+namespace UE::Math
+{
+	template<typename T> struct TTransform;
+}
+
 namespace Chaos
 {
 	// Chaos types
@@ -30,6 +35,7 @@ namespace Chaos
 		using FSolverRigidParticles = TKinematicGeometryParticlesImp<FSolverReal, 3, EGeometryParticlesSimType::Other>;
 		using FSolverRotation3 = TRotation<FSolverReal, 3>;
 		using FSolverRigidTransform3 = TRigidTransform<FSolverReal, 3>;
+		using FSolverTransform3 = UE::Math::TTransform<FSolverReal>;
 		using FSolverAABB3 = TAABB<FSolverReal, 3>;
 
 		// Softs solver class
