Chaos Cloth - Run the Long Range Attachment tether constraints as a pre-iteration step instead as a per iteration step to improve simulation performance.
Saves about 40ms from Owen's time as it now runs only 1 iteration on the tethers rather than 4.
This change also seems to improve behavior without seemingly affecting the overall stiffness, due to the tethers no longer overbearing the other constraints.

diff --git a/Engine/Plugins/Experimental/ChaosCloth/Source/Chaos/Private/ChaosCloth/ChaosClothConstraints.cpp b/Engine/Plugins/Experimental/ChaosCloth/Source/Chaos/Private/ChaosCloth/ChaosClothConstraints.cpp
--- a/Engine/Plugins/Experimental/ChaosCloth/Source/Chaos/Private/ChaosCloth/ChaosClothConstraints.cpp
+++ b/Engine/Plugins/Experimental/ChaosCloth/Source/Chaos/Private/ChaosCloth/ChaosClothConstraints.cpp
@@ -76,7 +76,7 @@ void FClothConstraints::CreateRules()
 		ConstraintRuleOffset = Evolution->AddConstraintRuleRange(NumConstraintRules, false);
 	}
 
-	TFunction<void(const Softs::FSolverParticles&, const Softs::FSolverReal)>* const ConstraintInits = Evolution->ConstraintInits().GetData() + ConstraintInitOffset;
+	TFunction<void(Softs::FSolverParticles&, const Softs::FSolverReal)>* const ConstraintInits = Evolution->ConstraintInits().GetData() + ConstraintInitOffset;
 	TFunction<void(Softs::FSolverParticles&, const Softs::FSolverReal)>* const ConstraintRules = Evolution->ConstraintRules().GetData() + ConstraintRuleOffset;
 
 	int32 ConstraintInitIndex = 0;
@@ -85,7 +85,7 @@ void FClothConstraints::CreateRules()
 	if (XEdgeConstraints)
 	{
 		ConstraintInits[ConstraintInitIndex++] =
-			[this](const Softs::FSolverParticles& /*Particles*/, const Softs::FSolverReal Dt)
+			[this](Softs::FSolverParticles& /*Particles*/, const Softs::FSolverReal Dt)
 			{
 				XEdgeConstraints->Init();
 				XEdgeConstraints->ApplyProperties(Dt, Evolution->GetIterations());
@@ -100,7 +100,7 @@ void FClothConstraints::CreateRules()
 	if (EdgeConstraints)
 	{
 		ConstraintInits[ConstraintInitIndex++] =
-			[this](const Softs::FSolverParticles& /*Particles*/, const Softs::FSolverReal Dt)
+			[this](Softs::FSolverParticles& /*Particles*/, const Softs::FSolverReal Dt)
 			{
 				EdgeConstraints->ApplyProperties(Dt, Evolution->GetIterations());
 			};
@@ -113,7 +113,7 @@ void FClothConstraints::CreateRules()
 	if (XBendingConstraints)
 	{
 		ConstraintInits[ConstraintInitIndex++] =
-			[this](const Softs::FSolverParticles& /*Particles*/, const Softs::FSolverReal Dt)
+			[this](Softs::FSolverParticles& /*Particles*/, const Softs::FSolverReal Dt)
 			{
 				XBendingConstraints->Init();
 				XBendingConstraints->ApplyProperties(Dt, Evolution->GetIterations());
@@ -127,7 +127,7 @@ void FClothConstraints::CreateRules()
 	if (BendingConstraints)
 	{
 		ConstraintInits[ConstraintInitIndex++] =
-			[this](const Softs::FSolverParticles& /*Particles*/, const Softs::FSolverReal Dt)
+			[this](Softs::FSolverParticles& /*Particles*/, const Softs::FSolverReal Dt)
 			{
 				BendingConstraints->ApplyProperties(Dt, Evolution->GetIterations());
 			};
@@ -148,7 +148,7 @@ void FClothConstraints::CreateRules()
 	if (XAreaConstraints)
 	{
 		ConstraintInits[ConstraintInitIndex++] =
-			[this](const Softs::FSolverParticles& /*Particles*/, const Softs::FSolverReal Dt)
+			[this](Softs::FSolverParticles& /*Particles*/, const Softs::FSolverReal Dt)
 			{
 				XAreaConstraints->Init();
 				XAreaConstraints->ApplyProperties(Dt, Evolution->GetIterations());
@@ -162,7 +162,7 @@ void FClothConstraints::CreateRules()
 	if (AreaConstraints)
 	{
 		ConstraintInits[ConstraintInitIndex++] =
-			[this](const Softs::FSolverParticles& /*Particles*/, const Softs::FSolverReal Dt)
+			[this](Softs::FSolverParticles& /*Particles*/, const Softs::FSolverReal Dt)
 			{
 				AreaConstraints->ApplyProperties(Dt, Evolution->GetIterations());
 			};
@@ -191,28 +191,20 @@ void FClothConstraints::CreateRules()
 	if (XLongRangeConstraints)
 	{
 		ConstraintInits[ConstraintInitIndex++] =
-			[this](const Softs::FSolverParticles& /*Particles*/, const Softs::FSolverReal Dt)
+			[this](Softs::FSolverParticles& Particles, const Softs::FSolverReal Dt)
 			{
 				XLongRangeConstraints->Init();
 				XLongRangeConstraints->ApplyProperties(Dt, Evolution->GetIterations());
-			};
-		ConstraintRules[ConstraintRuleIndex++] =
-			[this](Softs::FSolverParticles& Particles, const Softs::FSolverReal Dt)
-			{
-				XLongRangeConstraints->Apply(Particles, Dt);
+				XLongRangeConstraints->Apply(Particles, Dt);  // Run the LRA constraint only once per timestep
 			};
 	}
 	if (LongRangeConstraints)
 	{
 		ConstraintInits[ConstraintInitIndex++] =
-			[this](const Softs::FSolverParticles& /*Particles*/, const Softs::FSolverReal Dt)
-			{
-				LongRangeConstraints->ApplyProperties(Dt, Evolution->GetIterations());
-			};
-		ConstraintRules[ConstraintRuleIndex++] =
 			[this](Softs::FSolverParticles& Particles, const Softs::FSolverReal Dt)
 			{
-				LongRangeConstraints->Apply(Particles, Dt);
+				LongRangeConstraints->ApplyProperties(Dt, Evolution->GetIterations());
+				LongRangeConstraints->Apply(Particles, Dt);  // Run the LRA constraint only once per timestep
 			};
 	}
 	if (MaximumDistanceConstraints)
@@ -234,7 +226,7 @@ void FClothConstraints::CreateRules()
 	if (AnimDriveConstraints)
 	{
 		ConstraintInits[ConstraintInitIndex++] =
-			[this](const Softs::FSolverParticles& /*Particles*/, const Softs::FSolverReal Dt)
+			[this](Softs::FSolverParticles& /*Particles*/, const Softs::FSolverReal Dt)
 			{
 				AnimDriveConstraints->ApplyProperties(Dt, Evolution->GetIterations());
 			};
@@ -256,7 +248,7 @@ void FClothConstraints::CreateRules()
 	if (SelfCollisionConstraints)
 	{
 		ConstraintInits[ConstraintInitIndex++] =
-			[this](const Softs::FSolverParticles& Particles, const Softs::FSolverReal /*Dt*/)
+			[this](Softs::FSolverParticles& Particles, const Softs::FSolverReal /*Dt*/)
 			{
 				SelfCollisionConstraints->Init(Particles);
 			};
@@ -422,8 +414,7 @@ void FClothConstraints::SetLongRangeConstraints(
 			/*InStiffness =*/ Softs::FSolverVec2::UnitVector,
 			TetherScale);
 	}
-	++NumConstraintInits;  // Uses init to update the property tables
-	++NumConstraintRules;
+	++NumConstraintInits;  // Uses init to both update the property tables and apply the constraint
 }
 
 void FClothConstraints::SetMaximumDistanceConstraints(const TConstArrayView<FRealSingle>& MaxDistances)
diff --git a/Engine/Plugins/Experimental/ChaosCloth/Source/Chaos/Private/ChaosCloth/ChaosClothingSimulationSolver.cpp b/Engine/Plugins/Experimental/ChaosCloth/Source/Chaos/Private/ChaosCloth/ChaosClothingSimulationSolver.cpp
--- a/Engine/Plugins/Experimental/ChaosCloth/Source/Chaos/Private/ChaosCloth/ChaosClothingSimulationSolver.cpp
+++ b/Engine/Plugins/Experimental/ChaosCloth/Source/Chaos/Private/ChaosCloth/ChaosClothingSimulationSolver.cpp
@@ -40,7 +40,7 @@ namespace Chaos
 {
 
 static int32 ClothSolverMinParallelBatchSize = 1000;
-static bool bClothSolverParallelClothPreUpdate = false;  // TODO: Doesn't seem to improve much here. Review this after the ISPC implementation.
+static bool bClothSolverParallelClothPreUpdate = true;
 static bool bClothSolverParallelClothUpdate = true;
 static bool bClothSolverParallelClothPostUpdate = true;
 static bool bClothSolverUseImprovedTimeStepSmoothing = true;
diff --git a/Engine/Source/Runtime/Experimental/Chaos/Private/Chaos/PBDEvolution.cpp b/Engine/Source/Runtime/Experimental/Chaos/Private/Chaos/PBDEvolution.cpp
--- a/Engine/Source/Runtime/Experimental/Chaos/Private/Chaos/PBDEvolution.cpp
+++ b/Engine/Source/Runtime/Experimental/Chaos/Private/Chaos/PBDEvolution.cpp
@@ -377,7 +377,7 @@ void FPBDEvolution::AdvanceOneTimeStep(const FSolverReal Dt, const bool bSmoothD
 		TRACE_CPUPROFILER_EVENT_SCOPE(ChaosXPBDConstraintsInit);
 		SCOPE_CYCLE_COUNTER(STAT_ChaosXPBDConstraintsInit);
 		MConstraintInitsActiveView.SequentialFor(
-			[this, Dt](TArray<TFunction<void(const FSolverParticles&, const FSolverReal)>>& ConstraintInits, int32 Index)
+			[this, Dt](TArray<TFunction<void(FSolverParticles&, const FSolverReal)>>& ConstraintInits, int32 Index)
 			{
 				ConstraintInits[Index](MParticles, Dt);
 			});
diff --git a/Engine/Source/Runtime/Experimental/Chaos/Public/Chaos/PBDEvolution.h b/Engine/Source/Runtime/Experimental/Chaos/Public/Chaos/PBDEvolution.h
--- a/Engine/Source/Runtime/Experimental/Chaos/Public/Chaos/PBDEvolution.h
+++ b/Engine/Source/Runtime/Experimental/Chaos/Public/Chaos/PBDEvolution.h
@@ -76,8 +76,8 @@ class CHAOS_API FPBDEvolution : public TArrayCollection
 	void ActivateConstraintRuleRange(int32 Offset, bool bActivate) { MConstraintRulesActiveView.ActivateRange(Offset, bActivate); }
 
 	// Constraint accessors
-	const TArray<TFunction<void(const FSolverParticles&, const FSolverReal)>>& ConstraintInits() const { return MConstraintInits; }
-	TArray<TFunction<void(const FSolverParticles&, const FSolverReal)>>& ConstraintInits() { return MConstraintInits; }
+	const TArray<TFunction<void(FSolverParticles&, const FSolverReal)>>& ConstraintInits() const { return MConstraintInits; }
+	TArray<TFunction<void(FSolverParticles&, const FSolverReal)>>& ConstraintInits() { return MConstraintInits; }
 	const TArray<TFunction<void(FSolverParticles&, const FSolverReal)>>& ConstraintRules() const { return MConstraintRules; }
 	TArray<TFunction<void(FSolverParticles&, const FSolverReal)>>& ConstraintRules() { return MConstraintRules; }
 	
@@ -151,8 +151,8 @@ private:
 	TArrayCollectionArray<FSolverReal> MGroupDampings;
 	TArrayCollectionArray<bool> MGroupUseCCDs;
 	
-	TArray<TFunction<void(const FSolverParticles&, const FSolverReal)>> MConstraintInits;
-	TPBDActiveView<TArray<TFunction<void(const FSolverParticles&, const FSolverReal)>>> MConstraintInitsActiveView;
+	TArray<TFunction<void(FSolverParticles&, const FSolverReal)>> MConstraintInits;
+	TPBDActiveView<TArray<TFunction<void(FSolverParticles&, const FSolverReal)>>> MConstraintInitsActiveView;
 	TArray<TFunction<void(FSolverParticles&, const FSolverReal)>> MConstraintRules;
 	TPBDActiveView<TArray<TFunction<void(FSolverParticles&, const FSolverReal)>>> MConstraintRulesActiveView;
 
