Chaos Cloth - Make the number of iteration frame rate dependent and add a new MaxIterations property to the cloth shared config.

diff --git a/Engine/Plugins/Experimental/ChaosCloth/Source/Chaos/Private/ChaosCloth/ChaosClothConfig.cpp b/Engine/Plugins/Experimental/ChaosCloth/Source/Chaos/Private/ChaosCloth/ChaosClothConfig.cpp
--- a/Engine/Plugins/Experimental/ChaosCloth/Source/Chaos/Private/ChaosCloth/ChaosClothConfig.cpp
+++ b/Engine/Plugins/Experimental/ChaosCloth/Source/Chaos/Private/ChaosCloth/ChaosClothConfig.cpp
@@ -219,7 +219,8 @@ UChaosClothSharedSimConfig::~UChaosClothSharedSimConfig()
 void UChaosClothSharedSimConfig::MigrateFrom(const FClothConfig_Legacy& ClothConfig)
 {
 #if WITH_EDITORONLY_DATA
-	IterationCount = FMath::Clamp(int32(ClothConfig.SolverFrequency / 60.f), 1, 100);
+	IterationCount = FMath::Clamp(int32(ClothConfig.SolverFrequency / 120.f), 1, 100);
+	MaxIterationCount = FMath::Clamp(int32(ClothConfig.SolverFrequency / 30.f), 1, 100);
 
 	bUseDampingOverride_DEPRECATED = false;  // Damping is migrated to per cloth configs
 #endif
diff --git a/Engine/Plugins/Experimental/ChaosCloth/Source/Chaos/Private/ChaosCloth/ChaosClothingSimulation.cpp b/Engine/Plugins/Experimental/ChaosCloth/Source/Chaos/Private/ChaosCloth/ChaosClothingSimulation.cpp
--- a/Engine/Plugins/Experimental/ChaosCloth/Source/Chaos/Private/ChaosCloth/ChaosClothingSimulation.cpp
+++ b/Engine/Plugins/Experimental/ChaosCloth/Source/Chaos/Private/ChaosCloth/ChaosClothingSimulation.cpp
@@ -479,6 +479,7 @@ void FClothingSimulation::UpdateSimulationFromSharedSimConfig()
 		// Set common simulation parameters
 		Solver->SetNumSubsteps(ClothSharedSimConfig->SubdivisionCount);
 		Solver->SetNumIterations(ClothSharedSimConfig->IterationCount);
+		Solver->SetMaxNumIterations(ClothSharedSimConfig->MaxIterationCount);
 	}
 }
 
@@ -487,6 +488,11 @@ void FClothingSimulation::SetNumIterations(int32 InNumIterations)
 	Solver->SetNumIterations(InNumIterations);
 }
 
+void FClothingSimulation::SetMaxNumIterations(int32 MaxNumIterations)
+{
+	Solver->SetMaxNumIterations(MaxNumIterations);
+}
+
 void FClothingSimulation::SetNumSubsteps(int32 InNumSubsteps)
 {
 	Solver->SetNumSubsteps(InNumSubsteps);
@@ -555,6 +561,9 @@ void FClothingSimulation::Simulate(IClothingSimulationContext* InContext)
 		// Step the simulation
 		Solver->Update((Softs::FSolverReal)Context->DeltaSeconds);
 
+		// Keep the actual used number of iterations for the stats
+		NumIterations = Solver->GetNumUsedIterations();
+
 		// Update simulation time in ms (and provide an instant average instead of the value in real-time)
 		const float CurrSimulationTime = (float)((FPlatformTime::Seconds() - StartTime) * 1000.);
 		static const float SimulationTimeDecay = 0.03f; // 0.03 seems to provide a good rate of update for the instant average
diff --git a/Engine/Plugins/Experimental/ChaosCloth/Source/Chaos/Private/ChaosCloth/ChaosClothingSimulationInteractor.cpp b/Engine/Plugins/Experimental/ChaosCloth/Source/Chaos/Private/ChaosCloth/ChaosClothingSimulationInteractor.cpp
--- a/Engine/Plugins/Experimental/ChaosCloth/Source/Chaos/Private/ChaosCloth/ChaosClothingSimulationInteractor.cpp
+++ b/Engine/Plugins/Experimental/ChaosCloth/Source/Chaos/Private/ChaosCloth/ChaosClothingSimulationInteractor.cpp
@@ -219,19 +219,27 @@ void UChaosClothingSimulationInteractor::DisableGravityOverride()
 	}));
 }
 
-void UChaosClothingSimulationInteractor::SetNumIterations(int32 InNumIterations)
+void UChaosClothingSimulationInteractor::SetNumIterations(int32 NumIterations)
 {
-	Commands.Add(FChaosClothingSimulationInteractorCommand::CreateLambda([InNumIterations](Chaos::FClothingSimulation* Simulation, Chaos::FClothingSimulationContext* /*Context*/)
+	Commands.Add(FChaosClothingSimulationInteractorCommand::CreateLambda([NumIterations](Chaos::FClothingSimulation* Simulation, Chaos::FClothingSimulationContext* /*Context*/)
 	{
-		Simulation->SetNumIterations(InNumIterations);
+		Simulation->SetNumIterations(NumIterations);
 	}));
 }
 
-void UChaosClothingSimulationInteractor::SetNumSubsteps(int32 InNumSubsteps)
+void UChaosClothingSimulationInteractor::SetMaxNumIterations(int32 MaxNumIterations)
 {
-	Commands.Add(FChaosClothingSimulationInteractorCommand::CreateLambda([InNumSubsteps](Chaos::FClothingSimulation* Simulation, Chaos::FClothingSimulationContext* /*Context*/)
+	Commands.Add(FChaosClothingSimulationInteractorCommand::CreateLambda([MaxNumIterations](Chaos::FClothingSimulation* Simulation, Chaos::FClothingSimulationContext* /*Context*/)
 	{
-		Simulation->SetNumSubsteps(InNumSubsteps);
+		Simulation->SetMaxNumIterations(MaxNumIterations);
+	}));
+}
+
+void UChaosClothingSimulationInteractor::SetNumSubsteps(int32 NumSubsteps)
+{
+	Commands.Add(FChaosClothingSimulationInteractorCommand::CreateLambda([NumSubsteps](Chaos::FClothingSimulation* Simulation, Chaos::FClothingSimulationContext* /*Context*/)
+	{
+		Simulation->SetNumSubsteps(NumSubsteps);
 	}));
 }
 
diff --git a/Engine/Plugins/Experimental/ChaosCloth/Source/Chaos/Private/ChaosCloth/ChaosClothingSimulationSolver.cpp b/Engine/Plugins/Experimental/ChaosCloth/Source/Chaos/Private/ChaosCloth/ChaosClothingSimulationSolver.cpp
--- a/Engine/Plugins/Experimental/ChaosCloth/Source/Chaos/Private/ChaosCloth/ChaosClothingSimulationSolver.cpp
+++ b/Engine/Plugins/Experimental/ChaosCloth/Source/Chaos/Private/ChaosCloth/ChaosClothingSimulationSolver.cpp
@@ -44,6 +44,7 @@ static bool bClothSolverParallelClothPreUpdate = true;
 static bool bClothSolverParallelClothUpdate = true;
 static bool bClothSolverParallelClothPostUpdate = true;
 static bool bClothSolverUseImprovedTimeStepSmoothing = true;
+static bool bClothSolverDisableTimeDependentNumIterations = false;
 
 #if !UE_BUILD_SHIPPING
 static int32 ClothSolverDebugHitchLength = 0;
@@ -60,12 +61,14 @@ FAutoConsoleVariableRef CVarClothSolverDisableCollision(TEXT("p.ChaosCloth.Solve
 #endif
 
 FAutoConsoleVariableRef CVarClothSolverUseImprovedTimeStepSmoothing(TEXT("p.ChaosCloth.Solver.UseImprovedTimeStepSmoothing"), bClothSolverUseImprovedTimeStepSmoothing, TEXT("Use the time step smoothing on input forces only rather than on the entire cloth solver, in order to avoid miscalculating velocities."));
+FAutoConsoleVariableRef CVarClothSolverDisableTimeDependentNumIterations(TEXT("p.ChaosCloth.Solver.DisableTimeDependentNumIterations"), bClothSolverDisableTimeDependentNumIterations, TEXT("Make the number of iterations independent from the time step."));
 
 namespace ClothingSimulationSolverDefault
 {
 	static const Softs::FSolverVec3 Gravity((Softs::FSolverReal)0., (Softs::FSolverReal)0., (Softs::FSolverReal)-980.665);  // cm/s^2
 	static const Softs::FSolverVec3 WindVelocity((Softs::FSolverReal)0.);
 	static const int32 NumIterations = 1;
+	static const int32 MaxNumIterations = 10;
 	static const int32 NumSubsteps = 1;
 	static const FRealSingle SelfCollisionThickness = 2.f;
 	static const FRealSingle CollisionThickness = 1.2f;
@@ -86,6 +89,7 @@ FClothingSimulationSolver::FClothingSimulationSolver()
 	, Time(0.)
 	, DeltaTime(ClothingSimulationSolverConstant::StartDeltaTime)
 	, NumIterations(ClothingSimulationSolverDefault::NumIterations)
+	, MaxNumIterations(ClothingSimulationSolverDefault::MaxNumIterations)
 	, NumSubsteps(ClothingSimulationSolverDefault::NumSubsteps)
 	, CollisionParticlesOffset(0)
 	, CollisionParticlesSize(0)
@@ -101,7 +105,7 @@ FClothingSimulationSolver::FClothingSimulationSolver()
 			MoveTemp(LocalParticles),
 			MoveTemp(RigidParticles),
 			{}, // CollisionTriangles
-			ClothingSimulationSolverDefault::NumIterations,
+			FMath::Min(ClothingSimulationSolverDefault::NumIterations, ClothingSimulationSolverDefault::MaxNumIterations),
 			(Softs::FSolverReal)ClothingSimulationSolverDefault::CollisionThickness,
 			(Softs::FSolverReal)ClothingSimulationSolverDefault::SelfCollisionThickness,
 			(Softs::FSolverReal)ClothingSimulationSolverDefault::FrictionCoefficient,
@@ -948,9 +952,17 @@ void FClothingSimulationSolver::Update(Softs::FSolverReal InDeltaTime)
 		SCOPE_CYCLE_COUNTER(STAT_ChaosClothSolverUpdateSolverStep);
 		SCOPE_CYCLE_COUNTER(STAT_ClothInternalSolve);
 
-		Evolution->SetIterations(NumIterations);
+		// Update solver time dependent parameters
+		constexpr Softs::FSolverReal SolverFrequency = 60.f;  // TODO: This could become a solver property
 
-		const FReal SubstepDeltaTime = DeltaTime / (FReal)NumSubsteps;
+		const int32 TimeDependentNumIterations = bClothSolverDisableTimeDependentNumIterations ?
+			NumIterations :
+			(int32)(SolverFrequency * DeltaTime * (Softs::FSolverReal)NumIterations);
+
+		Evolution->SetIterations(FMath::Clamp(TimeDependentNumIterations, 1, MaxNumIterations));
+
+		// Advance substeps
+		const Softs::FSolverReal SubstepDeltaTime = DeltaTime / (Softs::FSolverReal)NumSubsteps;
 	
 		for (int32 i = 0; i < NumSubsteps; ++i)
 		{
@@ -978,6 +990,11 @@ void FClothingSimulationSolver::Update(Softs::FSolverReal InDeltaTime)
 	OldLocalSpaceLocation = LocalSpaceLocation;
 }
 
+int32 FClothingSimulationSolver::GetNumUsedIterations() const
+{
+	return Evolution->GetIterations();
+}
+
 FBoxSphereBounds FClothingSimulationSolver::CalculateBounds() const
 {
 	TRACE_CPUPROFILER_EVENT_SCOPE(FClothingSimulationSolver_CalculateBounds);
diff --git a/Engine/Plugins/Experimental/ChaosCloth/Source/Chaos/Public/ChaosCloth/ChaosClothConfig.h b/Engine/Plugins/Experimental/ChaosCloth/Source/Chaos/Public/ChaosCloth/ChaosClothConfig.h
--- a/Engine/Plugins/Experimental/ChaosCloth/Source/Chaos/Public/ChaosCloth/ChaosClothConfig.h
+++ b/Engine/Plugins/Experimental/ChaosCloth/Source/Chaos/Public/ChaosCloth/ChaosClothConfig.h
@@ -392,13 +392,25 @@ public:
 	virtual void PostEditChangeChainProperty(FPropertyChangedChainEvent& ChainEvent) override;
 #endif
 
-	// The number of solver iterations.
-	// This will increase the stiffness of all constraints but will increase the CPU cost.
-	UPROPERTY(EditAnywhere, Category = Simulation, meta = (UIMin = "1", UIMax = "10", ClampMin = "0", ClampMax = "100"))
-	int32 IterationCount = 2;
+	/**
+	 * The number of time step dependent solver iterations. This sets iterations at 60fps. 
+	 * This will increase the stiffness of all constraints but will increase the CPU cost.
+	 * If the frame rate is higher, the actual number of iterations used by the solver might be less, if the frame rate is lower it might be more.
+	 */
+	UPROPERTY(EditAnywhere, Category = Simulation, meta = (UIMin = "1", UIMax = "10", ClampMin = "1", ClampMax = "100"))
+	int32 IterationCount = 1;
 
-	// The number of solver substeps.
-	// This will increase the precision of the collision inputs and help with constraint resolutions but will increase the CPU cost.
+	/**
+	 * The maximum number of solver iterations.
+	 * This is the upper limit of the number of iterations set in solver, when the frame rate is lower than 60fps.
+	 */
+	UPROPERTY(EditAnywhere, Category = Simulation, meta = (UIMin = "1", UIMax = "10", ClampMin = "1", ClampMax = "100"))
+	int32 MaxIterationCount = 10;
+
+	/**
+	 * The number of solver substeps.
+	 * This will increase the precision of the collision inputs and help with constraint resolutions but will increase the CPU cost.
+	 */
 	UPROPERTY(EditAnywhere, Category = Simulation, meta = (UIMin = "1", UIMax = "10", ClampMin = "1", ClampMax = "100"))
 	int32 SubdivisionCount = 1;
 
diff --git a/Engine/Plugins/Experimental/ChaosCloth/Source/Chaos/Public/ChaosCloth/ChaosClothingSimulation.h b/Engine/Plugins/Experimental/ChaosCloth/Source/Chaos/Public/ChaosCloth/ChaosClothingSimulation.h
--- a/Engine/Plugins/Experimental/ChaosCloth/Source/Chaos/Public/ChaosCloth/ChaosClothingSimulation.h
+++ b/Engine/Plugins/Experimental/ChaosCloth/Source/Chaos/Public/ChaosCloth/ChaosClothingSimulation.h
@@ -70,6 +70,7 @@ namespace Chaos
 
 		// IClothingSimulation interface
 		virtual void SetNumIterations(int32 NumIterations) override;
+		virtual void SetMaxNumIterations(int32 MaxNumIterations) override;
 		virtual void SetNumSubsteps(int32 NumSubsteps) override;
 		virtual int32 GetNumCloths() const override { return NumCloths; }
 		virtual int32 GetNumKinematicParticles() const override { return NumKinematicParticles; }
diff --git a/Engine/Plugins/Experimental/ChaosCloth/Source/Chaos/Public/ChaosCloth/ChaosClothingSimulationInteractor.h b/Engine/Plugins/Experimental/ChaosCloth/Source/Chaos/Public/ChaosCloth/ChaosClothingSimulationInteractor.h
--- a/Engine/Plugins/Experimental/ChaosCloth/Source/Chaos/Public/ChaosCloth/ChaosClothingSimulationInteractor.h
+++ b/Engine/Plugins/Experimental/ChaosCloth/Source/Chaos/Public/ChaosCloth/ChaosClothingSimulationInteractor.h
@@ -85,7 +85,8 @@ public:
 	virtual void SetAnimDriveSpringStiffness(float Stiffness) override;
 	virtual void EnableGravityOverride(const FVector& Gravity) override;
 	virtual void DisableGravityOverride() override;
-	virtual void SetNumIterations(int32 NumIterations = 2) override;
+	virtual void SetNumIterations(int32 NumIterations = 1) override;
+	virtual void SetMaxNumIterations(int32 MaxNumIterations = 10) override;
 	virtual void SetNumSubsteps(int32 NumSubsteps = 1) override;
 
 protected:
diff --git a/Engine/Plugins/Experimental/ChaosCloth/Source/Chaos/Public/ChaosCloth/ChaosClothingSimulationSolver.h b/Engine/Plugins/Experimental/ChaosCloth/Source/Chaos/Public/ChaosCloth/ChaosClothingSimulationSolver.h
--- a/Engine/Plugins/Experimental/ChaosCloth/Source/Chaos/Public/ChaosCloth/ChaosClothingSimulationSolver.h
+++ b/Engine/Plugins/Experimental/ChaosCloth/Source/Chaos/Public/ChaosCloth/ChaosClothingSimulationSolver.h
@@ -44,6 +44,8 @@ namespace Chaos
 
 		void SetNumIterations(int32 InNumIterations) { NumIterations = InNumIterations; }
 		int32 GetNumIterations() const { return NumIterations; }
+		void SetMaxNumIterations(int32 InMaxNumIterations) { MaxNumIterations = InMaxNumIterations; }
+		int32 GetMaxNumIterations() const { return MaxNumIterations; }
 		void SetNumSubsteps(int32 InNumSubsteps) { NumSubsteps = InNumSubsteps; }
 		int32 GetNumSubsteps() const { return NumSubsteps; }
 		// ---- End of the animatable property setters ----
@@ -62,6 +64,9 @@ namespace Chaos
 		// Update solver properties before simulation.
 		void Update(FSolverReal InDeltaTime);
 
+		// Return the actual of number of iterations used by the Evolution solver after the update (different from the number of iterations, depends on frame rate)
+		int32 GetNumUsedIterations() const;
+
 		FBoxSphereBounds CalculateBounds() const;
 		// ---- End of the object management functions ----
 
@@ -71,7 +76,7 @@ namespace Chaos
 
 		void ResetStartPose(int32 Offset, int32 NumParticles);
 
-		// Get the currrent solver time
+		// Get the current solver time
 		FSolverReal GetTime() const { return Time; }
 		void SetParticleMassUniform(int32 Offset, FRealSingle UniformMass, FRealSingle MinPerParticleMass, const FTriangleMesh& Mesh, const TFunctionRef<bool(int32)>& KinematicPredicate);
 		void SetParticleMassFromTotalMass(int32 Offset, FRealSingle TotalMass, FRealSingle MinPerParticleMass, const FTriangleMesh& Mesh, const TFunctionRef<bool(int32)>& KinematicPredicate);
@@ -207,6 +212,7 @@ namespace Chaos
 		FSolverReal Time;
 		FSolverReal DeltaTime;
 		int32 NumIterations;
+		int32 MaxNumIterations;
 		int32 NumSubsteps;
 
 		// Solver colliders offset
diff --git a/Engine/Source/Editor/Persona/Private/AnimationEditorViewportClient.cpp b/Engine/Source/Editor/Persona/Private/AnimationEditorViewportClient.cpp
--- a/Engine/Source/Editor/Persona/Private/AnimationEditorViewportClient.cpp
+++ b/Engine/Source/Editor/Persona/Private/AnimationEditorViewportClient.cpp
@@ -1091,6 +1091,10 @@ FText FAnimationViewportClient::GetDisplayInfo(bool bDisplayAllInfo) const
 		{
 			TextValue = ConcatenateLine(TextValue, FText::Format(LOCTEXT("NumDynamicParticles", "Dynamic Particles: {0}"), NumDynamicParticles));
 		}
+		if (const int32 NumIterations = ClothingSimulation->GetNumIterations())
+		{
+			TextValue = ConcatenateLine(TextValue, FText::Format(LOCTEXT("NumIterations", "Iterations: {0}"), NumIterations));
+		}
 		if (const float SimulationTime = ClothingSimulation->GetSimulationTime())
 		{
 			FNumberFormattingOptions NumberFormatOptions;
diff --git a/Engine/Source/Runtime/ClothingSystemRuntimeInterface/Public/ClothingSimulationInteractor.h b/Engine/Source/Runtime/ClothingSystemRuntimeInterface/Public/ClothingSimulationInteractor.h
--- a/Engine/Source/Runtime/ClothingSystemRuntimeInterface/Public/ClothingSimulationInteractor.h
+++ b/Engine/Source/Runtime/ClothingSystemRuntimeInterface/Public/ClothingSimulationInteractor.h
@@ -81,11 +81,16 @@ public:
 	virtual void DisableGravityOverride()
 	PURE_VIRTUAL(UClothingSimulationInteractor::DisableGravityOverride, );
 
-	/** Set the number of solver iterations. */
+	/** Set the number of time dependent solver iterations. */
 	UFUNCTION(BlueprintCallable, Category = ClothingSimulation)
-	virtual void SetNumIterations(int32 NumIterations = 2)
+	virtual void SetNumIterations(int32 NumIterations = 1)
 	PURE_VIRTUAL(UClothingSimulationInteractor::SetNumIterations, );
 
+	/** Set the maximum number of solver iterations. */
+	UFUNCTION(BlueprintCallable, Category = ClothingSimulation)
+	virtual void SetMaxNumIterations(int32 MaxNumIterations = 10)
+	PURE_VIRTUAL(UClothingSimulationInteractor::SetMaxNumIterations, );
+
 	/** Set the number of substeps or subdivisions. */
 	UFUNCTION(BlueprintCallable, Category = ClothingSimulation, Meta = (Keywords = "Subdivisions"))
 	virtual void SetNumSubsteps(int32 NumSubsteps = 1)
diff --git a/Engine/Source/Runtime/ClothingSystemRuntimeInterface/Public/ClothingSimulationInterface.h b/Engine/Source/Runtime/ClothingSystemRuntimeInterface/Public/ClothingSimulationInterface.h
--- a/Engine/Source/Runtime/ClothingSystemRuntimeInterface/Public/ClothingSimulationInterface.h
+++ b/Engine/Source/Runtime/ClothingSystemRuntimeInterface/Public/ClothingSimulationInterface.h
@@ -142,6 +142,11 @@ public:
 	 * Implementation is not considered thread safe, and therefore this function must be called at a synchronization point.
 	 */
 	virtual void SetNumIterations(int32 /*NumIterations*/) {}
+	/**
+	 * Set the maximum number of iterations used by the solver.
+	 * Implementation is not considered thread safe, and therefore this function must be called at a synchronization point.
+	 */
+	virtual void SetMaxNumIterations(int32 /*MaxNumIterations*/) {}
 	/**
 	 * Set the number of substeps used by the solver.
 	 * Implementation is not considered thread safe, and therefore this function must be called at a synchronization point.
