Chaos Cloth - Fix PhysX config stiffnesses migration to work with the new adaptive number of iterations optimization.

diff --git a/Engine/Plugins/Experimental/ChaosCloth/Source/Chaos/Private/ChaosCloth/ChaosClothConfig.cpp b/Engine/Plugins/Experimental/ChaosCloth/Source/Chaos/Private/ChaosCloth/ChaosClothConfig.cpp
--- a/Engine/Plugins/Experimental/ChaosCloth/Source/Chaos/Private/ChaosCloth/ChaosClothConfig.cpp
+++ b/Engine/Plugins/Experimental/ChaosCloth/Source/Chaos/Private/ChaosCloth/ChaosClothConfig.cpp
@@ -37,16 +37,16 @@ void UChaosClothConfig::MigrateFrom(const FClothConfig_Legacy& ClothConfig)
 		ClothConfig.VerticalConstraintConfig.Stiffness;
 	const float HorizontalStiffness =
 		ClothConfig.HorizontalConstraintConfig.Stiffness;
-	EdgeStiffnessWeighted.Low = EdgeStiffnessWeighted.High = FMath::Clamp((VerticalStiffness + HorizontalStiffness) * 0.5f, 0.f, 1.f);
+	EdgeStiffnessWeighted.Low = EdgeStiffnessWeighted.High = FMath::Clamp((VerticalStiffness + HorizontalStiffness) * 0.5f + 0.5f, 0.f, 1.f);
 
 	BendingStiffnessWeighted.Low = BendingStiffnessWeighted.High = FMath::Clamp(
-		ClothConfig.BendConstraintConfig.Stiffness, 0.f, 1.f);
+		ClothConfig.BendConstraintConfig.Stiffness + 0.5f, 0.f, 1.f);
 
 	AreaStiffnessWeighted.Low = AreaStiffnessWeighted.High = FMath::Clamp(
-		ClothConfig.ShearConstraintConfig.Stiffness, 0.f, 1.f);
+		ClothConfig.ShearConstraintConfig.Stiffness + 0.5f, 0.f, 1.f);
 
 	AnimDriveStiffness.Low = 0.f;
-	AnimDriveStiffness.High = FMath::Clamp(ClothConfig.AnimDriveSpringStiffness, 0.f, 1.f);
+	AnimDriveStiffness.High = FMath::Clamp(ClothConfig.AnimDriveSpringStiffness + 0.5f, 0.f, 1.f);
 
 	AnimDriveDamping.Low = 0.f;
 	AnimDriveDamping.High = FMath::Clamp(ClothConfig.AnimDriveDamperStiffness, 0.f, 1.f);
@@ -56,7 +56,7 @@ void UChaosClothConfig::MigrateFrom(const FClothConfig_Legacy& ClothConfig)
 	bUseBendingElements = false;
 	bUseSelfCollisions = (ClothConfig.SelfCollisionRadius > 0.f && ClothConfig.SelfCollisionStiffness > 0.f);
 
-	TetherStiffness.Low = FMath::Clamp(ClothConfig.TetherStiffness, 0.f, 1.f);
+	TetherStiffness.Low = FMath::Clamp(ClothConfig.TetherStiffness + 0.5f, 0.f, 1.f);
 	TetherStiffness.High = 1.f;
 	TetherScale.Low = TetherScale.High = FMath::Clamp(ClothConfig.TetherLimit, 0.01f, 10.f);
 	ShapeTargetStiffness = 0.f;
