Changing the p.ChaosCloth.Ispc command to register during static initialization (and in test builds) so we can use it for automated testing.

diff --git a/Engine/Plugins/Experimental/ChaosCloth/Source/Chaos/Private/ChaosCloth/ChaosClothingSimulation.cpp b/Engine/Plugins/Experimental/ChaosCloth/Source/Chaos/Private/ChaosCloth/ChaosClothingSimulation.cpp
--- a/Engine/Plugins/Experimental/ChaosCloth/Source/Chaos/Private/ChaosCloth/ChaosClothingSimulation.cpp
+++ b/Engine/Plugins/Experimental/ChaosCloth/Source/Chaos/Private/ChaosCloth/ChaosClothingSimulation.cpp
@@ -53,7 +53,7 @@
 #include "ChaosClothingSimulation.ispc.generated.h"
 #endif
 
-#if INTEL_ISPC && !(UE_BUILD_SHIPPING || UE_BUILD_TEST)  // Include only used for the Ispc command
+#if INTEL_ISPC && !UE_BUILD_SHIPPING  // Include only used for the Ispc command
 #include "Chaos/PBDCollisionConstraints.h"
 #include "Chaos/PBDEvolution.h"
 #include "Chaos/PBDSpringConstraints.h"
@@ -105,9 +105,58 @@ namespace ClothingSimulationCVar
 }
 #endif  // #if CHAOS_DEBUG_DRAW
 
-#if !(UE_BUILD_SHIPPING || UE_BUILD_TEST)
+#if !UE_BUILD_SHIPPING
 namespace ClothingSimulationConsole
 {
+#if INTEL_ISPC
+	static FAutoConsoleCommand CommandIspc(
+		TEXT("p.ChaosCloth.Ispc"),
+		TEXT("Enable or disable ISPC optimizations for cloth simulation."),
+		FConsoleCommandWithArgsDelegate::CreateLambda([](const TArray<FString>& Args)
+		{
+			bool bEnableISPC;
+			switch (Args.Num())
+			{
+			default:
+				break; // Invalid arguments
+			case 1:
+				if (Args[0] == TEXT("1") || Args[0] == TEXT("true") || Args[0] == TEXT("on"))
+				{
+					bEnableISPC = true;
+				}
+				else if (Args[0] == TEXT("0") || Args[0] == TEXT("false") || Args[0] == TEXT("off"))
+				{
+					bEnableISPC = false;
+				}
+				else
+				{
+					break; // Invalid arguments
+				}
+				bChaos_AxialSpring_ISPC_Enabled =
+					bChaos_LongRange_ISPC_Enabled =
+					bChaos_Spherical_ISPC_Enabled =
+					bChaos_Spring_ISPC_Enabled =
+					bChaos_DampVelocity_ISPC_Enabled =
+					bChaos_PerParticleCollision_ISPC_Enabled =
+					bChaos_VelocityField_ISPC_Enabled =
+					bChaos_GetSimData_ISPC_Enabled =
+					bChaos_SkinPhysicsMesh_ISPC_Enabled =
+					bChaos_PreSimulationTransforms_ISPC_Enabled =
+					bChaos_CalculateBounds_ISPC_Enabled =
+					bChaos_PostIterationUpdates_ISPC_Enabled =
+					bEnableISPC;
+				return;
+			}
+
+			UE_LOG(LogChaosCloth, Display, TEXT("Invalid arguments."));
+			UE_LOG(LogChaosCloth, Display, TEXT("Usage:"));
+			UE_LOG(LogChaosCloth, Display, TEXT("  p.ChaosCloth.Ispc [0|1]|[true|false]|[on|off]"));
+			UE_LOG(LogChaosCloth, Display, TEXT("Example: p.Chaos.Ispc on"));
+		}),
+		ECVF_Cheat);
+#endif // #if INTEL_ISPC
+
+#if !UE_BUILD_TEST
 	class FCommand final
 	{
 	public:
@@ -116,15 +165,6 @@ namespace ClothingSimulationConsole
 			, bIsPaused(false)
 			, ResetCount(0)
 		{
-#if INTEL_ISPC
-			// Register Ispc console command
-			ConsoleObjects.Add(IConsoleManager::Get().RegisterConsoleCommand(
-				TEXT("p.ChaosCloth.Ispc"),
-				TEXT("Enable or disable ISPC optimizations for cloth simulation."),
-				FConsoleCommandWithArgsDelegate::CreateRaw(this, &FCommand::Ispc),
-				ECVF_Cheat));
-#endif // #if INTEL_ISPC
-
 			// Register DebugStep console command
 			ConsoleObjects.Add(IConsoleManager::Get().RegisterConsoleCommand(
 				TEXT("p.ChaosCloth.DebugStep"),
@@ -163,49 +203,6 @@ namespace ClothingSimulationConsole
 		}
 
 	private:
-#if INTEL_ISPC
-		void Ispc(const TArray<FString>& Args)
-		{
-			bool bEnableISPC;
-			switch (Args.Num())
-			{
-			default:
-				break;  // Invalid arguments
-			case 1:
-				if (Args[0] == TEXT("1") || Args[0] == TEXT("true") || Args[0] == TEXT("on"))
-				{
-					bEnableISPC = true;
-				}
-				else if (Args[0] == TEXT("0") || Args[0] == TEXT("false") || Args[0] == TEXT("off"))
-				{
-					bEnableISPC = false;
-				}
-				else
-				{
-					break;  // Invalid arguments
-				}
-				bChaos_AxialSpring_ISPC_Enabled =
-				bChaos_LongRange_ISPC_Enabled =
-				bChaos_Spherical_ISPC_Enabled =
-				bChaos_Spring_ISPC_Enabled =
-				bChaos_DampVelocity_ISPC_Enabled =
-				bChaos_PerParticleCollision_ISPC_Enabled =
-				bChaos_VelocityField_ISPC_Enabled =
-				bChaos_GetSimData_ISPC_Enabled =
-				bChaos_SkinPhysicsMesh_ISPC_Enabled =
-				bChaos_PreSimulationTransforms_ISPC_Enabled =
-				bChaos_CalculateBounds_ISPC_Enabled =
-				bChaos_PostIterationUpdates_ISPC_Enabled =
-					bEnableISPC;
-				return;
-			}
-			UE_LOG(LogChaosCloth, Display, TEXT("Invalid arguments."));
-			UE_LOG(LogChaosCloth, Display, TEXT("Usage:"));
-			UE_LOG(LogChaosCloth, Display, TEXT("  p.ChaosCloth.Ispc [0|1]|[true|false]|[on|off]"));
-			UE_LOG(LogChaosCloth, Display, TEXT("Example: p.Chaos.Ispc on"));
-		}
-#endif // #if INTEL_ISPC
-
 		void DebugStep(const TArray<FString>& Args)
 		{
 			switch (Args.Num())
@@ -263,8 +260,9 @@ namespace ClothingSimulationConsole
 		TAtomic<int32> ResetCount;
 	};
 	static TUniquePtr<FCommand> Command;
+#endif // #if !UE_BUILD_TEST
 }
-#endif  // #if !(UE_BUILD_SHIPPING || UE_BUILD_TEST)
+#endif  // #if !UE_BUILD_SHIPPING
 
 // Default parameters, will be overwritten when cloth assets are loaded
 namespace ChaosClothingSimulationDefault
