Chaos Cloth - Fix the tethers' scale when the mesh skeleton contains scaled bones.
- Tidy up some LWC type conversions on the cloth constraint property setters.

diff --git a/Engine/Plugins/Experimental/ChaosCloth/Source/Chaos/Private/ChaosCloth/ChaosClothingSimulationCloth.cpp b/Engine/Plugins/Experimental/ChaosCloth/Source/Chaos/Private/ChaosCloth/ChaosClothingSimulationCloth.cpp
--- a/Engine/Plugins/Experimental/ChaosCloth/Source/Chaos/Private/ChaosCloth/ChaosClothingSimulationCloth.cpp
+++ b/Engine/Plugins/Experimental/ChaosCloth/Source/Chaos/Private/ChaosCloth/ChaosClothingSimulationCloth.cpp
@@ -204,7 +204,7 @@ void FClothingSimulationCloth::FLODData::Add(FClothingSimulationSolver* Solver,
 	{
 		const TConstArrayView<FRealSingle>& TetherScaleMultipliers = WeightMaps[(int32)EChaosWeightMapTarget::TetherScale];
 
-		const Softs::FSolverReal MeshScale(Cloth->Mesh->GetReferenceBoneTransform().GetScale3D().GetMax());
+		const Softs::FSolverReal MeshScale = Cloth->Mesh->GetScale();
 
 		ClothConstraints.SetLongRangeConstraints(
 			Tethers,
@@ -275,23 +275,23 @@ void FClothingSimulationCloth::FLODData::Update(FClothingSimulationSolver* Solve
 	const int32 Offset = SolverData.FindChecked(Solver).Offset;
 	check(Offset != INDEX_NONE);
 
-	const FReal MeshScale = Cloth->Mesh->GetReferenceBoneTransform().GetScale3D().GetMax();
+	const Softs::FSolverReal MeshScale = Cloth->Mesh->GetScale();
 
 	// Update the animatable constraint parameters
 	FClothConstraints& ClothConstraints = Solver->GetClothConstraints(Offset);
-	ClothConstraints.SetMaximumDistanceProperties((FReal)Cloth->MaxDistancesMultiplier * MeshScale);
-	ClothConstraints.SetEdgeProperties(FVec2((FReal)Cloth->EdgeStiffness[0], (FReal)Cloth->EdgeStiffness[1]));
-	ClothConstraints.SetBendingProperties(FVec2((FReal)Cloth->BendingStiffness[0], (FReal)Cloth->BendingStiffness[1]));
-	ClothConstraints.SetAreaProperties(FVec2((FReal)Cloth->AreaStiffness[0], (FReal)Cloth->AreaStiffness[1]));
+	ClothConstraints.SetMaximumDistanceProperties((Softs::FSolverReal)Cloth->MaxDistancesMultiplier * MeshScale);
+	ClothConstraints.SetEdgeProperties(Softs::FSolverVec2((Softs::FSolverReal)Cloth->EdgeStiffness[0], (Softs::FSolverReal)Cloth->EdgeStiffness[1]));
+	ClothConstraints.SetBendingProperties(Softs::FSolverVec2((Softs::FSolverReal)Cloth->BendingStiffness[0], (Softs::FSolverReal)Cloth->BendingStiffness[1]));
+	ClothConstraints.SetAreaProperties(Softs::FSolverVec2((FSolverReal)Cloth->AreaStiffness[0], (Softs::FSolverReal)Cloth->AreaStiffness[1]));
 	ClothConstraints.SetLongRangeAttachmentProperties(
-		FVec2((FReal)Cloth->TetherStiffness[0], (FReal)Cloth->TetherStiffness[1]),
-		FVec2((FReal)Cloth->TetherScale[0], (FReal)Cloth->TetherScale[1]) * MeshScale);
-	ClothConstraints.SetSelfCollisionProperties((FReal)Cloth->SelfCollisionThickness);
+		Softs::FSolverVec2((Softs::FSolverReal)Cloth->TetherStiffness[0], (Softs::FSolverReal)Cloth->TetherStiffness[1]),
+		Softs::FSolverVec2((Softs::FSolverReal)Cloth->TetherScale[0], (Softs::FSolverReal)Cloth->TetherScale[1]) * MeshScale);
+	ClothConstraints.SetSelfCollisionProperties((Softs::FSolverReal)Cloth->SelfCollisionThickness);
 	ClothConstraints.SetAnimDriveProperties(
-		FVec2((FReal)Cloth->AnimDriveStiffness[0], (FReal)Cloth->AnimDriveStiffness[1]),
-		FVec2((FReal)Cloth->AnimDriveDamping[0], (FReal)Cloth->AnimDriveDamping[1]));
-	ClothConstraints.SetThinShellVolumeProperties((FReal)Cloth->VolumeStiffness);
-	ClothConstraints.SetVolumeProperties((FReal)Cloth->VolumeStiffness);
+		Softs::FSolverVec2((Softs::FSolverReal)Cloth->AnimDriveStiffness[0], (Softs::FSolverReal)Cloth->AnimDriveStiffness[1]),
+		Softs::FSolverVec2((Softs::FSolverReal)Cloth->AnimDriveDamping[0], (Softs::FSolverReal)Cloth->AnimDriveDamping[1]));
+	ClothConstraints.SetThinShellVolumeProperties((Softs::FSolverReal)Cloth->VolumeStiffness);
+	ClothConstraints.SetVolumeProperties((Softs::FSolverReal)Cloth->VolumeStiffness);
 	ClothConstraints.SetBackstopProperties(Cloth->bEnableBackstop, MeshScale);
 }
 
diff --git a/Engine/Plugins/Experimental/ChaosCloth/Source/Chaos/Private/ChaosCloth/ChaosClothingSimulationMesh.cpp b/Engine/Plugins/Experimental/ChaosCloth/Source/Chaos/Private/ChaosCloth/ChaosClothingSimulationMesh.cpp
--- a/Engine/Plugins/Experimental/ChaosCloth/Source/Chaos/Private/ChaosCloth/ChaosClothingSimulationMesh.cpp
+++ b/Engine/Plugins/Experimental/ChaosCloth/Source/Chaos/Private/ChaosCloth/ChaosClothingSimulationMesh.cpp
@@ -128,7 +128,7 @@ int32 FClothingSimulationMesh::GetReferenceBoneIndex() const
 	return Asset ? Asset->ReferenceBoneIndex : INDEX_NONE;
 }
 
-FRigidTransform3 Chaos::FClothingSimulationMesh::GetReferenceBoneTransform() const
+FRigidTransform3 FClothingSimulationMesh::GetReferenceBoneTransform() const
 {
 	if (SkeletalMeshComponent)
 	{
@@ -146,6 +146,19 @@ FRigidTransform3 Chaos::FClothingSimulationMesh::GetReferenceBoneTransform() con
 	return FRigidTransform3::Identity;
 }
 
+Softs::FSolverReal FClothingSimulationMesh::GetScale() const
+{
+	if (SkeletalMeshComponent)
+	{
+		if (const FClothingSimulationContextCommon* const Context =
+			static_cast<const FClothingSimulationContextCommon*>(SkeletalMeshComponent->GetClothingSimulationContext()))
+		{
+			return (Softs::FSolverReal)Context->ComponentToWorld.GetScale3D().GetMax();
+		}
+	}
+	return (Softs::FSolverReal)1.;
+}
+
 bool FClothingSimulationMesh::WrapDeformLOD(
 	int32 PrevLODIndex,
 	int32 LODIndex,
diff --git a/Engine/Plugins/Experimental/ChaosCloth/Source/Chaos/Public/ChaosCloth/ChaosClothingSimulationMesh.h b/Engine/Plugins/Experimental/ChaosCloth/Source/Chaos/Public/ChaosCloth/ChaosClothingSimulationMesh.h
--- a/Engine/Plugins/Experimental/ChaosCloth/Source/Chaos/Public/ChaosCloth/ChaosClothingSimulationMesh.h
+++ b/Engine/Plugins/Experimental/ChaosCloth/Source/Chaos/Public/ChaosCloth/ChaosClothingSimulationMesh.h
@@ -42,6 +42,9 @@ namespace Chaos
 		int32 GetReferenceBoneIndex() const;
 		FRigidTransform3 GetReferenceBoneTransform() const;
 
+		// Return this mesh component's scale (the max of the three axis scale values)
+		Softs::FSolverReal GetScale() const;
+
 		bool WrapDeformLOD(
 			int32 PrevLODIndex,
 			int32 LODIndex,
