Chaos Cloth - Fix suspended simulation positions and bounds by reverting part of changes made in CL 18170330 and fixing the local space bounds calculation.

diff --git a/Engine/Plugins/Experimental/ChaosCloth/Source/Chaos/Private/ChaosCloth/ChaosClothingSimulation.cpp b/Engine/Plugins/Experimental/ChaosCloth/Source/Chaos/Private/ChaosCloth/ChaosClothingSimulation.cpp
--- a/Engine/Plugins/Experimental/ChaosCloth/Source/Chaos/Private/ChaosCloth/ChaosClothingSimulation.cpp
+++ b/Engine/Plugins/Experimental/ChaosCloth/Source/Chaos/Private/ChaosCloth/ChaosClothingSimulation.cpp
@@ -351,7 +351,8 @@ void FClothingSimulation::CreateActor(USkeletalMeshComponent* InOwnerComponent,
 		const FClothingSimulationContext* const Context = static_cast<const FClothingSimulationContext*>(InOwnerComponent->GetClothingSimulationContext());
 		check(Context);
 		static const bool bReset = true;
-		Solver->SetLocalSpaceLocation(bUseLocalSpaceSimulation ? Context->ComponentToWorld.GetLocation() : FVec3(0.), bReset);
+		Solver->SetLocalSpaceLocation(bUseLocalSpaceSimulation ? (FVec3)Context->ComponentToWorld.GetLocation() : FVec3(0.), bReset);
+		Solver->SetLocalSpaceRotation(bUseLocalSpaceSimulation ? (FQuat)Context->ComponentToWorld.GetRotation() : FQuat::Identity);
 	}
 	else
 	{
@@ -536,6 +537,7 @@ void FClothingSimulation::Simulate(IClothingSimulationContext* InContext)
 
 		// Update Solver animatable parameters
 		Solver->SetLocalSpaceLocation(bUseLocalSpaceSimulation ? (FVec3)Context->ComponentToWorld.GetLocation() : FVec3(0), bNeedsReset);
+		Solver->SetLocalSpaceRotation(bUseLocalSpaceSimulation ? (FQuat)Context->ComponentToWorld.GetRotation() : FQuat::Identity);
 		Solver->SetWindVelocity(Context->WindVelocity, Context->WindAdaption);
 		Solver->SetGravity(bUseGravityOverride ? GravityOverride : Context->WorldGravity);
 		Solver->EnableClothGravityOverride(!bUseGravityOverride);  // Disable all cloth gravity overrides when the interactor takes over
@@ -615,16 +617,16 @@ void FClothingSimulation::GetSimulationData(
 		OutData.Reset();
 	}
 
+	// Get the solver's local space
+	const FVec3& LocalSpaceLocation = Solver->GetLocalSpaceLocation(); // Note: Since the ReferenceSpaceTransform can be suspended with the simulation, it is important that the suspended local space location is used too in order to get the simulation data back into reference space
+
 	// Retrieve the component transforms
 	const FClothingSimulationContext* const Context = static_cast<const FClothingSimulationContext*>(InOwnerComponent->GetClothingSimulationContext());
-	check(Context);  // A simuation can't be created without context
+	check(Context);  // A simulation can't be created without context
 	const FTransform& OwnerTransform = Context->ComponentToWorld;
 
 	const TArray<FTransform>& ComponentSpaceTransforms = InOverrideComponent ? InOverrideComponent->GetComponentSpaceTransforms() : InOwnerComponent->GetComponentSpaceTransforms();
 
-	// Get the solver's local space location
-	const FVec3 LocalSpaceLocation = bUseLocalSpaceSimulation ? OwnerTransform.GetLocation() : FVec3(0.);  // Note: The component could be moving while the simulation is suspended, this means the owner component location needs to be used instead of the solver's local space location
-
 	// Set the simulation data for each of the cloths
 	for (const TUniquePtr<FClothingSimulationCloth>& Cloth : Cloths)
 	{
@@ -723,24 +725,21 @@ FBoxSphereBounds FClothingSimulation::GetBounds(const USkeletalMeshComponent* In
 	check(Solver);
 	FBoxSphereBounds Bounds = Solver->CalculateBounds();
 
-	if (InOwnerComponent)
+	if (bUseLocalSpaceSimulation)
+	{
+		// The component could be moving while the simulation is suspended so getting the bounds
+		// in world space isn't good enough and the bounds origin needs to be continuously updated
+		Bounds = Bounds.TransformBy(FTransform((FQuat)Solver->GetLocalSpaceRotation(), (FVector)Solver->GetLocalSpaceLocation()).Inverse());
+	}
+	else if (InOwnerComponent)
 	{
 		const FClothingSimulationContext* const Context = static_cast<const FClothingSimulationContext*>(InOwnerComponent->GetClothingSimulationContext());
-		check(Context);  // A simuation can't be created without context
+		check(Context);  // A simulation can't be created without context
 		const FTransform& OwnerTransform = Context->ComponentToWorld;
-
-		if (bUseLocalSpaceSimulation)
-		{
-			// The component could be moving while the simulation is suspended so getting the bounds
-			// in world space isn't good enough and the bounds origin needs to be continuously updated
-			const FVec3 CurrentLocalSpaceLocation = OwnerTransform.GetLocation();
-			const FVec3& SolverLocalSpaceLocation = Solver->GetLocalSpaceLocation();
-			Bounds.Origin += FVector(CurrentLocalSpaceLocation - SolverLocalSpaceLocation);
-		}
-
 		// Return local bounds
-		return Bounds.TransformBy(OwnerTransform.Inverse());
+		Bounds = Bounds.TransformBy(OwnerTransform.Inverse());
 	}
+	// Else return bounds in world space
 	return Bounds;
 }
 
@@ -789,6 +788,7 @@ void FClothingSimulation::RefreshClothConfig(const IClothingSimulationContext* I
 	const FClothingSimulationContext* const Context = static_cast<const FClothingSimulationContext*>(InContext);
 	static const bool bReset = true;
 	Solver->SetLocalSpaceLocation(bUseLocalSpaceSimulation ? (FVec3)Context->ComponentToWorld.GetLocation() : FVec3(0), bReset);
+	Solver->SetLocalSpaceRotation(bUseLocalSpaceSimulation ? (FQuat)Context->ComponentToWorld.GetRotation() : FQuat::Identity);
 
 	// Reset stats
 	ResetStats();
diff --git a/Engine/Plugins/Experimental/ChaosCloth/Source/Chaos/Public/ChaosCloth/ChaosClothingSimulationSolver.h b/Engine/Plugins/Experimental/ChaosCloth/Source/Chaos/Public/ChaosCloth/ChaosClothingSimulationSolver.h
--- a/Engine/Plugins/Experimental/ChaosCloth/Source/Chaos/Public/ChaosCloth/ChaosClothingSimulationSolver.h
+++ b/Engine/Plugins/Experimental/ChaosCloth/Source/Chaos/Public/ChaosCloth/ChaosClothingSimulationSolver.h
@@ -29,6 +29,8 @@ namespace Chaos
 		// ---- Animatable property setters ----
 		void SetLocalSpaceLocation(const FVec3& InLocalSpaceLocation, bool bReset = false);
 		const FVec3& GetLocalSpaceLocation() const { return LocalSpaceLocation; }
+		void SetLocalSpaceRotation(const FQuat& InLocalSpaceRotation) { LocalSpaceRotation = InLocalSpaceRotation; }
+		const FRotation3& GetLocalSpaceRotation() const { return LocalSpaceRotation; }
 
 		// Disables all Cloths gravity override mechanism
 		void EnableClothGravityOverride(bool bInIsClothGravityOverrideEnabled) { bIsClothGravityOverrideEnabled = bInIsClothGravityOverrideEnabled; }
@@ -207,6 +209,7 @@ namespace Chaos
 		// Local space simulation
 		FVec3 OldLocalSpaceLocation;  // This is used to translate between world space and simulation space,
 		FVec3 LocalSpaceLocation;     // add this to simulation space coordinates to get world space coordinates, must keep FReal as underlying type for LWC
+		FRotation3 LocalSpaceRotation;
 
 		// Time stepping
 		FSolverReal Time;
