Chaos Cloth - Add global point damping to offer a faster alternative to the current local damping formulation. Saves 40us on the pre-iteration step of the Owen asset in Test builds.
- Rename the existing Damping parameter to LocalDamping, set it disabled by default, and add a new more accurate damping config migration path for the NvCloth assets.
- Tidy up a couple of mismatched LWC conversions.
- Fix the cloth position initializations by adding a reset as it was impeding the damping tests.
[Edigrate CL 19067856]

diff --git a/Engine/Plugins/Experimental/ChaosCloth/Source/Chaos/Private/ChaosCloth/ChaosClothConfig.cpp b/Engine/Plugins/Experimental/ChaosCloth/Source/Chaos/Private/ChaosCloth/ChaosClothConfig.cpp
--- a/Engine/Plugins/Experimental/ChaosCloth/Source/Chaos/Private/ChaosCloth/ChaosClothConfig.cpp
+++ b/Engine/Plugins/Experimental/ChaosCloth/Source/Chaos/Private/ChaosCloth/ChaosClothConfig.cpp
@@ -149,7 +149,7 @@ void UChaosClothConfig::PostLoad()
 		bUseLegacyBackstop = true;
 	}
 
-	if (PhysicsObjectVersion < FPhysicsObjectVersion::ChaosClothAddWeightedValue ||
+	if (PhysicsObjectVersion < FPhysicsObjectVersion::ChaosClothAddWeightedValue &&
 		FortniteMainBranchObjectVersion < FFortniteMainBranchObjectVersion::ChaosClothAddWeightedValue)
 	{
 		AnimDriveStiffness.Low = 0.f;
@@ -168,7 +168,7 @@ void UChaosClothConfig::PostLoad()
 		FictitiousAngularScale = 0.f;  // Maintain early behavior with no fictitious forces
 	}
 
-	if (PhysicsObjectVersion < FPhysicsObjectVersion::ChaosClothAddTetherStiffnessWeightMap ||
+	if (PhysicsObjectVersion < FPhysicsObjectVersion::ChaosClothAddTetherStiffnessWeightMap &&
 		FortniteMainBranchObjectVersion < FFortniteMainBranchObjectVersion::ChaosClothAddTetherStiffnessWeightMap)
 	{
 		// Note: Unlike AnimDriveStiffness, Low is updated here, because there was no existing weight map before this version
@@ -190,7 +190,8 @@ void UChaosClothConfig::PostLoad()
 		AreaStiffnessWeighted.Low = AreaStiffnessWeighted.High = AreaStiffness_DEPRECATED;
 	}
 
-	if (UE5ReleaseStreamObjectVersion < FUE5ReleaseStreamObjectVersion::ChaosClothFasterDamping)
+	if (FortniteMainBranchObjectVersion < FFortniteMainBranchObjectVersion::ChaosClothFasterDamping &&
+		UE5ReleaseStreamObjectVersion < FUE5ReleaseStreamObjectVersion::ChaosClothFasterDamping)
 	{
 		// Note: The previous damping has been renamed LocalDamping to make space for a faster but more primitive global point damping.
 		LocalDampingCoefficient = DampingCoefficient;
diff --git a/Engine/Source/Runtime/Core/Public/UObject/FortniteMainBranchObjectVersion.h b/Engine/Source/Runtime/Core/Public/UObject/FortniteMainBranchObjectVersion.h
--- a/Engine/Source/Runtime/Core/Public/UObject/FortniteMainBranchObjectVersion.h
+++ b/Engine/Source/Runtime/Core/Public/UObject/FortniteMainBranchObjectVersion.h
@@ -196,6 +196,9 @@ struct CORE_API FFortniteMainBranchObjectVersion
 		// Decoupled the generation of the water texture from the Water Brush and the landscape
 		WaterZonesRefactor,
 
+		// Add faster damping calculations to the cloth simulation and rename previous Damping parameter to LocalDamping.
+		ChaosClothFasterDamping,
+
 		// -----<new versions can be added above this line>-------------------------------------------------
 		VersionPlusOne,
 		LatestVersion = VersionPlusOne - 1
