Fix compile warning from deprecating XLongRangeConstraints in 19741727

diff --git a/Engine/Plugins/Experimental/ChaosCloth/Source/Chaos/Private/ChaosCloth/ChaosClothConstraints.cpp b/Engine/Plugins/Experimental/ChaosCloth/Source/Chaos/Private/ChaosCloth/ChaosClothConstraints.cpp
--- a/Engine/Plugins/Experimental/ChaosCloth/Source/Chaos/Private/ChaosCloth/ChaosClothConstraints.cpp
+++ b/Engine/Plugins/Experimental/ChaosCloth/Source/Chaos/Private/ChaosCloth/ChaosClothConstraints.cpp
@@ -33,9 +33,11 @@ FClothConstraints::FClothConstraints()
 {
 }
 
+PRAGMA_DISABLE_DEPRECATION_WARNINGS
 FClothConstraints::~FClothConstraints()
 {
 }
+PRAGMA_ENABLE_DEPRECATION_WARNINGS
 
 void FClothConstraints::Initialize(
 	Softs::FPBDEvolution* InEvolution,
