Choas Cloth - Add missing Local Damping Coefficient in the SetDamping Blueprint function.

diff --git a/Engine/Plugins/Experimental/ChaosCloth/Source/Chaos/Private/ChaosCloth/ChaosClothingSimulationInteractor.cpp b/Engine/Plugins/Experimental/ChaosCloth/Source/Chaos/Private/ChaosCloth/ChaosClothingSimulationInteractor.cpp
--- a/Engine/Plugins/Experimental/ChaosCloth/Source/Chaos/Private/ChaosCloth/ChaosClothingSimulationInteractor.cpp
+++ b/Engine/Plugins/Experimental/ChaosCloth/Source/Chaos/Private/ChaosCloth/ChaosClothingSimulationInteractor.cpp
@@ -81,11 +81,11 @@ void UChaosClothingInteractor::SetBackstop(bool bEnabled)
 		Cloth->SetBackstopProperties(bEnabled);
 	}));
 }
-void UChaosClothingInteractor::SetDamping(float DampingCoefficient)
+void UChaosClothingInteractor::SetDamping(float DampingCoefficient, float LocalDampingCoefficient)
 {
-	Commands.Add(FChaosClothingInteractorCommand::CreateLambda([DampingCoefficient](Chaos::FClothingSimulationCloth* Cloth)
+	Commands.Add(FChaosClothingInteractorCommand::CreateLambda([DampingCoefficient, LocalDampingCoefficient](Chaos::FClothingSimulationCloth* Cloth)
 	{
-		Cloth->SetDampingProperties(DampingCoefficient);
+		Cloth->SetDampingProperties(DampingCoefficient, LocalDampingCoefficient);
 	}));
 }
 
diff --git a/Engine/Plugins/Experimental/ChaosCloth/Source/Chaos/Public/ChaosCloth/ChaosClothingSimulationCloth.h b/Engine/Plugins/Experimental/ChaosCloth/Source/Chaos/Public/ChaosCloth/ChaosClothingSimulationCloth.h
--- a/Engine/Plugins/Experimental/ChaosCloth/Source/Chaos/Public/ChaosCloth/ChaosClothingSimulationCloth.h
+++ b/Engine/Plugins/Experimental/ChaosCloth/Source/Chaos/Public/ChaosCloth/ChaosClothingSimulationCloth.h
@@ -86,7 +86,7 @@ namespace Chaos
 		void SetLongRangeAttachmentProperties(const TVec2<FRealSingle>& InTetherStiffness, const TVec2<FRealSingle>& InTetherScale) { TetherStiffness = InTetherStiffness; TetherScale = InTetherScale;  }
 		void SetCollisionProperties(FRealSingle InCollisionThickness, FRealSingle InFrictionCoefficient, bool bInUseCCD, FRealSingle InSelfCollisionThickness) { CollisionThickness = InCollisionThickness; FrictionCoefficient = InFrictionCoefficient; bUseCCD = bInUseCCD; SelfCollisionThickness = InSelfCollisionThickness; }
 		void SetBackstopProperties(bool bInEnableBackstop) { bEnableBackstop = bInEnableBackstop; }
-		void SetDampingProperties(FRealSingle InDampingCoefficient) { DampingCoefficient = InDampingCoefficient; }
+		void SetDampingProperties(FRealSingle InDampingCoefficient, FRealSingle InLocalDampingCoefficient = 0.f) { DampingCoefficient = InDampingCoefficient; LocalDampingCoefficient = InLocalDampingCoefficient; }
 		void SetAerodynamicsProperties(const TVec2<FRealSingle>& InDrag, const TVec2<FRealSingle>& InLift, FRealSingle InAirDensity, const FVec3& InWindVelocity) { Drag = InDrag; Lift = InLift; InAirDensity = AirDensity; WindVelocity = InWindVelocity; }
 		void SetPressureProperties(const TVec2<FRealSingle>& InPressure) { Pressure = InPressure; }
 		void SetGravityProperties(FRealSingle InGravityScale, bool bInIsGravityOverridden, const FVec3& InGravityOverride) { GravityScale = InGravityScale; bIsGravityOverridden = bInIsGravityOverridden; GravityOverride = InGravityOverride; }
diff --git a/Engine/Plugins/Experimental/ChaosCloth/Source/Chaos/Public/ChaosCloth/ChaosClothingSimulationInteractor.h b/Engine/Plugins/Experimental/ChaosCloth/Source/Chaos/Public/ChaosCloth/ChaosClothingSimulationInteractor.h
--- a/Engine/Plugins/Experimental/ChaosCloth/Source/Chaos/Public/ChaosCloth/ChaosClothingSimulationInteractor.h
+++ b/Engine/Plugins/Experimental/ChaosCloth/Source/Chaos/Public/ChaosCloth/ChaosClothingSimulationInteractor.h
@@ -43,7 +43,7 @@ public:
 	void SetBackstop(bool bEnabled = true);
 
 	UFUNCTION(BlueprintCallable, Category = ClothingSimulation, Meta = (Keywords = "Chaos Coefficient"))
-	void SetDamping(float DampingCoefficient = 0.01f);
+	void SetDamping(float DampingCoefficient = 0.01f, float LocalDampingCoefficient = 0.f);
 
 	// Deprecated. This function cannot set different Low and High values for the Drag and Lift weight maps. Use SetWind instead.
 	UFUNCTION(BlueprintCallable, Category = ClothingSimulation, Meta = (Keywords = "Chaos Wind Drag Lift Coefficient"))
