Chaos Cloth: Add debug draw of spring constraints with color based on parallel graph coloring.

Added new debug draw for edge springs. Updated bend springs to use the same debug draw since they are both just FPBDSpringConstraints.

diff --git a/Engine/Plugins/Experimental/ChaosCloth/Source/Chaos/Private/ChaosCloth/ChaosClothingSimulation.cpp b/Engine/Plugins/Experimental/ChaosCloth/Source/Chaos/Private/ChaosCloth/ChaosClothingSimulation.cpp
--- a/Engine/Plugins/Experimental/ChaosCloth/Source/Chaos/Private/ChaosCloth/ChaosClothingSimulation.cpp
+++ b/Engine/Plugins/Experimental/ChaosCloth/Source/Chaos/Private/ChaosCloth/ChaosClothingSimulation.cpp
@@ -103,6 +103,7 @@ namespace ClothingSimulationCVar
 	TAutoConsoleVariable<bool> DebugDrawBackstopDistances   (TEXT("p.ChaosCloth.DebugDrawBackstopDistances"   ), false, TEXT("Whether to debug draw the Chaos Cloth backstop distances"), ECVF_Cheat);
 	TAutoConsoleVariable<bool> DebugDrawMaxDistances        (TEXT("p.ChaosCloth.DebugDrawMaxDistances"        ), false, TEXT("Whether to debug draw the Chaos Cloth max distances"), ECVF_Cheat);
 	TAutoConsoleVariable<bool> DebugDrawAnimDrive           (TEXT("p.ChaosCloth.DebugDrawAnimDrive"           ), false, TEXT("Whether to debug draw the Chaos Cloth anim drive"), ECVF_Cheat);
+	TAutoConsoleVariable<bool> DebugDrawEdgeConstraint      (TEXT("p.ChaosCloth.DebugDrawEdgeConstraint"      ), false, TEXT("Whether to debug draw the Chaos Cloth edge constraint"), ECVF_Cheat);
 	TAutoConsoleVariable<bool> DebugDrawBendingConstraint   (TEXT("p.ChaosCloth.DebugDrawBendingConstraint"   ), false, TEXT("Whether to debug draw the Chaos Cloth bending constraint"), ECVF_Cheat);
 	TAutoConsoleVariable<bool> DebugDrawLongRangeConstraint (TEXT("p.ChaosCloth.DebugDrawLongRangeConstraint" ), false, TEXT("Whether to debug draw the Chaos Cloth long range constraint (aka tether constraint)"), ECVF_Cheat);
 	TAutoConsoleVariable<bool> DebugDrawWindForces          (TEXT("p.ChaosCloth.DebugDrawWindForces"          ), false, TEXT("Whether to debug draw the Chaos Cloth wind forces"), ECVF_Cheat);
@@ -606,6 +607,7 @@ void FClothingSimulation::Simulate(IClothingSimulationContext* InContext)
 	if (ClothingSimulationCVar::DebugDrawBackstopDistances   .GetValueOnAnyThread()) { DebugDrawBackstopDistances   (); }
 	if (ClothingSimulationCVar::DebugDrawMaxDistances        .GetValueOnAnyThread()) { DebugDrawMaxDistances        (); }
 	if (ClothingSimulationCVar::DebugDrawAnimDrive           .GetValueOnAnyThread()) { DebugDrawAnimDrive           (); }
+	if (ClothingSimulationCVar::DebugDrawEdgeConstraint      .GetValueOnAnyThread()) { DebugDrawEdgeConstraint      (); }
 	if (ClothingSimulationCVar::DebugDrawBendingConstraint   .GetValueOnAnyThread()) { DebugDrawBendingConstraint   (); }
 	if (ClothingSimulationCVar::DebugDrawLongRangeConstraint .GetValueOnAnyThread()) { DebugDrawLongRangeConstraint (); }
 	if (ClothingSimulationCVar::DebugDrawWindForces          .GetValueOnAnyThread()) { DebugDrawWindAndPressureForces(); }
@@ -1827,6 +1829,74 @@ void FClothingSimulation::DebugDrawAnimDrive(FPrimitiveDrawInterface* PDI) const
 	}
 }
 
+static void DebugDrawSpringConstraintColors(FPrimitiveDrawInterface* PDI, const TConstArrayView<Softs::FSolverVec3>& Positions, const FVec3& LocalSpaceLocation, const Softs::FPBDSpringConstraints* const SpringConstraints)
+{
+	check(SpringConstraints);
+
+	const TArray<TVec2<int32>>& Constraints = SpringConstraints->GetConstraints();
+	const TArray<int32>& ConstraintsPerColorStartIndex = SpringConstraints->GetConstraintsPerColorStartIndex();
+	if (ConstraintsPerColorStartIndex.Num() > 1)
+	{
+		const int32 ConstraintColorNum = ConstraintsPerColorStartIndex.Num() - 1;
+		const uint8 HueOffset = 196 / ConstraintColorNum;
+		auto ConstraintColor =
+			[HueOffset](int32 ColorIndex)->FLinearColor
+		{
+			return FLinearColor::MakeFromHSV8(ColorIndex * HueOffset, 255, 255);
+		};
+
+		for (int32 ConstraintColorIndex = 0; ConstraintColorIndex < ConstraintColorNum; ++ConstraintColorIndex)
+		{
+			const int32 ColorStart = ConstraintsPerColorStartIndex[ConstraintColorIndex];
+			const int32 ColorEnd = ConstraintsPerColorStartIndex[ConstraintColorIndex + 1];
+			const FLinearColor DrawColor = ConstraintColor(ConstraintColorIndex);
+			for (int32 ConstraintIndex = ColorStart; ConstraintIndex < ColorEnd; ++ConstraintIndex)
+			{
+				// Draw line
+				const TVec2<int32>& Constraint = Constraints[ConstraintIndex];
+				const FVec3 Pos0 = FVec3(Positions[Constraint[0]]) + LocalSpaceLocation;
+				const FVec3 Pos1 = FVec3(Positions[Constraint[1]]) + LocalSpaceLocation;
+				DrawLine(PDI, Pos0, Pos1, DrawColor);
+			}
+		}
+	}
+	else
+	{
+		for (const TVec2<int32>& Constraint : Constraints)
+		{
+			// Draw line
+			const FVec3 Pos0 = FVec3(Positions[Constraint[0]]) + LocalSpaceLocation;
+			const FVec3 Pos1 = FVec3(Positions[Constraint[1]]) + LocalSpaceLocation;
+
+			DrawLine(PDI, Pos0, Pos1, FLinearColor::Black);
+		}
+	}
+}
+
+void FClothingSimulation::DebugDrawEdgeConstraint(FPrimitiveDrawInterface* PDI) const
+{
+	const FVec3& LocalSpaceLocation = Solver->GetLocalSpaceLocation();
+
+	for (const FClothingSimulationCloth* const Cloth : Solver->GetCloths())
+	{
+		const int32 Offset = Cloth->GetOffset(Solver.Get());
+		if (Offset == INDEX_NONE)
+		{
+			continue;
+		}
+
+		// Draw constraints
+		const FClothConstraints& ClothConstraints = Solver->GetClothConstraints(Offset);
+
+		const TConstArrayView<Softs::FSolverVec3> Positions = Cloth->GetParticlePositions(Solver.Get());
+
+		if (const Softs::FPBDSpringConstraints* const EdgeConstraints = ClothConstraints.GetEdgeConstraints().Get())
+		{
+			DebugDrawSpringConstraintColors(PDI, Positions, LocalSpaceLocation, EdgeConstraints);
+		}
+	}
+}
+
 void FClothingSimulation::DebugDrawBendingConstraint(FPrimitiveDrawInterface* PDI) const
 {
 	const FVec3& LocalSpaceLocation = Solver->GetLocalSpaceLocation();
@@ -1846,15 +1916,7 @@ void FClothingSimulation::DebugDrawBendingConstraint(FPrimitiveDrawInterface* PD
 
 		if (const Softs::FPBDSpringConstraints* const BendingConstraints = ClothConstraints.GetBendingConstraints().Get())
 		{
-			const TArray<TVec2<int32>>& Constraints = BendingConstraints->GetConstraints();
-			for (const TVec2<int32>& Constraint : Constraints)
-			{
-				// Draw line
-				const FVec3 Pos0 = FVec3(Positions[Constraint[0]]) + LocalSpaceLocation;
-				const FVec3 Pos1 = FVec3(Positions[Constraint[1]]) + LocalSpaceLocation;
-
-				DrawLine(PDI, Pos0, Pos1, FLinearColor::Black);
-			}
+			DebugDrawSpringConstraintColors(PDI, Positions, LocalSpaceLocation, BendingConstraints);
 		}
 
 		if (const Softs::FPBDBendingConstraints* const BendingConstraints = ClothConstraints.GetBendingElementConstraints().Get())
diff --git a/Engine/Plugins/Experimental/ChaosCloth/Source/Chaos/Public/ChaosCloth/ChaosClothingSimulation.h b/Engine/Plugins/Experimental/ChaosCloth/Source/Chaos/Public/ChaosCloth/ChaosClothingSimulation.h
--- a/Engine/Plugins/Experimental/ChaosCloth/Source/Chaos/Public/ChaosCloth/ChaosClothingSimulation.h
+++ b/Engine/Plugins/Experimental/ChaosCloth/Source/Chaos/Public/ChaosCloth/ChaosClothingSimulation.h
@@ -113,6 +113,7 @@ namespace Chaos
 		CHAOSCLOTH_API void DebugDrawBackstopDistances(FPrimitiveDrawInterface* PDI = nullptr) const;
 		CHAOSCLOTH_API void DebugDrawMaxDistances(FPrimitiveDrawInterface* PDI = nullptr) const;
 		CHAOSCLOTH_API void DebugDrawAnimDrive(FPrimitiveDrawInterface* PDI = nullptr) const;
+		CHAOSCLOTH_API void DebugDrawEdgeConstraint(FPrimitiveDrawInterface* PDI = nullptr) const;
 		CHAOSCLOTH_API void DebugDrawBendingConstraint(FPrimitiveDrawInterface* PDI = nullptr) const;
 		CHAOSCLOTH_API void DebugDrawLongRangeConstraint(FPrimitiveDrawInterface* PDI = nullptr) const;
 		CHAOSCLOTH_API void DebugDrawWindAndPressureForces(FPrimitiveDrawInterface* PDI = nullptr) const;
diff --git a/Engine/Plugins/Experimental/ChaosClothEditor/Source/Chaos/Private/ChaosClothEditor/ChaosSimulationEditorExtender.cpp b/Engine/Plugins/Experimental/ChaosClothEditor/Source/Chaos/Private/ChaosClothEditor/ChaosSimulationEditorExtender.cpp
--- a/Engine/Plugins/Experimental/ChaosClothEditor/Source/Chaos/Private/ChaosClothEditor/ChaosSimulationEditorExtender.cpp
+++ b/Engine/Plugins/Experimental/ChaosClothEditor/Source/Chaos/Private/ChaosClothEditor/ChaosSimulationEditorExtender.cpp
@@ -78,6 +78,7 @@ const FVisualizationOption FVisualizationOption::OptionData[] =
 	FVisualizationOption(&Chaos::FClothingSimulation::DebugDrawMaxDistances        , TEXT("p.ChaosClothEditor.DebugDrawMaxDistances"        ), LOCTEXT("ChaosVisName_MaxDistance"         , "Max Distances"              ), LOCTEXT("ChaosVisName_MaxDistance_ToolTip"         , "Draws the current max distances for the sim particles as a line along its normal"), /*bDisablesSimulation =*/true),
 	FVisualizationOption(&Chaos::FClothingSimulation::DebugDrawMaxDistanceValues   , TEXT("p.ChaosClothEditor.DebugDrawMaxDistanceValues"   ), LOCTEXT("ChaosVisName_MaxDistanceValue"    , "Max Distances As Numbers"   ), LOCTEXT("ChaosVisName_MaxDistanceValue_ToolTip"    , "Draws the current max distances as numbers")),
 	FVisualizationOption(&Chaos::FClothingSimulation::DebugDrawAnimDrive           , TEXT("p.ChaosClothEditor.DebugDrawAnimDrive"           ), LOCTEXT("ChaosVisName_AnimDrive"           , "Anim Drive"                 ), LOCTEXT("ChaosVisName_AnimDrive_Tooltip"           , "Draws the current skinned reference mesh for the simulation which anim drive will attempt to reach if enabled")),
+	FVisualizationOption(&Chaos::FClothingSimulation::DebugDrawEdgeConstraint      , TEXT("p.ChaosClothEditor.DebugDrawEdgeConstraint"      ), LOCTEXT("ChaosVisName_EdgeConstraint"      , "Edge Constraint"            ), LOCTEXT("ChaosVisName_EdgeConstraint_Tooltip"      , "Draws the edge spring constraints")),
 	FVisualizationOption(&Chaos::FClothingSimulation::DebugDrawBendingConstraint   , TEXT("p.ChaosClothEditor.DebugDrawBendingConstraint"   ), LOCTEXT("ChaosVisName_BendingConstraint"   , "Bending Constraint"         ), LOCTEXT("ChaosVisName_BendingConstraint_Tooltip"   , "Draws the bending spring constraints")),
 	FVisualizationOption(&Chaos::FClothingSimulation::DebugDrawLongRangeConstraint , TEXT("p.ChaosClothEditor.DebugDrawLongRangeConstraint" ), LOCTEXT("ChaosVisName_LongRangeConstraint" , "Long Range Constraint"      ), LOCTEXT("ChaosVisName_LongRangeConstraint_Tooltip" , "Draws the long range attachment constraint distances")),
 	FVisualizationOption(&Chaos::FClothingSimulation::DebugDrawWindAndPressureForces, TEXT("p.ChaosClothEditor.DebugDrawWindAndPressureForces"), LOCTEXT("ChaosVisName_WindAndPressureForces", "Wind Aerodynamic And Pressure Forces"), LOCTEXT("ChaosVisName_WindAndPressure_Tooltip", "Draws the Wind drag and lift and pressure forces")),
diff --git a/Engine/Source/Runtime/Experimental/Chaos/Public/Chaos/PBDSpringConstraints.h b/Engine/Source/Runtime/Experimental/Chaos/Public/Chaos/PBDSpringConstraints.h
--- a/Engine/Source/Runtime/Experimental/Chaos/Public/Chaos/PBDSpringConstraints.h
+++ b/Engine/Source/Runtime/Experimental/Chaos/Public/Chaos/PBDSpringConstraints.h
@@ -32,6 +32,8 @@ public:
 
 	void Apply(FSolverParticles& Particles, const FSolverReal Dt) const;
 
+	const TArray<int32>& GetConstraintsPerColorStartIndex() const { return ConstraintsPerColorStartIndex; }
+
 private:
 	void InitColor(const FSolverParticles& InParticles);
 	void ApplyHelper(FSolverParticles& Particles, const FSolverReal Dt, const int32 ConstraintIndex, const FSolverReal ExpStiffnessValue) const;
