Chaos Cloth: Only calculate buckling status for BendingElements (and XPBD Bending Elements) on Init (i.e., once per subframe) rather than every iteration.

By not dynamically updating the buckling status per iteration, early partial results (e.g., first iteration where just gravity + tethers + stretch has run) cannot introduce incorrect buckling where cloth can get stuck in an unnatural buckled state due to tether artifacts. This does mean that cloth will buckle lagged by one subframe, but in practice, this produces better results.

diff --git a/Engine/Plugins/Experimental/ChaosCloth/Source/Chaos/Private/ChaosCloth/ChaosClothConstraints.cpp b/Engine/Plugins/Experimental/ChaosCloth/Source/Chaos/Private/ChaosCloth/ChaosClothConstraints.cpp
--- a/Engine/Plugins/Experimental/ChaosCloth/Source/Chaos/Private/ChaosCloth/ChaosClothConstraints.cpp
+++ b/Engine/Plugins/Experimental/ChaosCloth/Source/Chaos/Private/ChaosCloth/ChaosClothConstraints.cpp
@@ -158,8 +158,9 @@ void FClothConstraints::CreateRules()
 	if (BendingElementConstraints)
 	{
 		ConstraintInits[ConstraintInitIndex++] =
-			[this](Softs::FSolverParticles& /*Particles*/, const Softs::FSolverReal Dt)
+			[this](Softs::FSolverParticles& Particles, const Softs::FSolverReal Dt)
 			{
+				BendingElementConstraints->Init(Particles);
 				BendingElementConstraints->ApplyProperties(Dt, Evolution->GetIterations());
 			};
 		ConstraintRules[ConstraintRuleIndex++] =
@@ -171,9 +172,9 @@ void FClothConstraints::CreateRules()
 	if (XBendingElementConstraints)
 	{
 		ConstraintInits[ConstraintInitIndex++] =
-			[this](Softs::FSolverParticles& /*Particles*/, const Softs::FSolverReal Dt)
+			[this](Softs::FSolverParticles& Particles, const Softs::FSolverReal Dt)
 		{
-			XBendingElementConstraints->Init();
+			XBendingElementConstraints->Init(Particles);
 			XBendingElementConstraints->ApplyProperties(Dt, Evolution->GetIterations());
 		};
 		ConstraintRules[ConstraintRuleIndex++] =
diff --git a/Engine/Plugins/Experimental/ChaosCloth/Source/Chaos/Private/ChaosCloth/ChaosClothingSimulation.cpp b/Engine/Plugins/Experimental/ChaosCloth/Source/Chaos/Private/ChaosCloth/ChaosClothingSimulation.cpp
--- a/Engine/Plugins/Experimental/ChaosCloth/Source/Chaos/Private/ChaosCloth/ChaosClothingSimulation.cpp
+++ b/Engine/Plugins/Experimental/ChaosCloth/Source/Chaos/Private/ChaosCloth/ChaosClothingSimulation.cpp
@@ -1922,19 +1922,15 @@ void FClothingSimulation::DebugDrawBendingConstraint(FPrimitiveDrawInterface* PD
 		if (const Softs::FPBDBendingConstraints* const BendingConstraints = ClothConstraints.GetBendingElementConstraints().Get())
 		{
 			const TArray<TVec4<int32>>& Constraints = BendingConstraints->GetConstraints();
-			const TArray<FSolverReal>& RestAngles = BendingConstraints->GetRestAngles();
+			const TArray<bool>& IsBuckled = BendingConstraints->GetIsBuckled();
 
-			// Color constraint edge with red or blue: Red = Buckled, Blue = Not Buckled. This is based on its current configuration, not the values used
-			// in the actual simulation.
+			// Color constraint edge with red or blue: Red = Buckled, Blue = Not Buckled. 
 			for (int32 ConstraintIndex = 0; ConstraintIndex < Constraints.Num(); ++ConstraintIndex)
 			{
 				const Softs::FSolverVec3& P1 = Positions[Constraints[ConstraintIndex][0]];
 				const Softs::FSolverVec3& P2 = Positions[Constraints[ConstraintIndex][1]];
-				const Softs::FSolverVec3& P3 = Positions[Constraints[ConstraintIndex][2]];
-				const Softs::FSolverVec3& P4 = Positions[Constraints[ConstraintIndex][3]];
-				const Softs::FSolverReal Angle = Softs::FPBDBendingConstraintsBase::CalcAngle(P1, P2, P3, P4);
 				
-				const bool bIsBuckled = BendingConstraints->AngleIsBuckled(Angle, RestAngles[ConstraintIndex]);
+				const bool bIsBuckled = IsBuckled[ConstraintIndex];
 
 				const FVec3 Pos0 = FVec3(P1) + LocalSpaceLocation;
 				const FVec3 Pos1 = FVec3(P2) + LocalSpaceLocation;
@@ -1945,19 +1941,15 @@ void FClothingSimulation::DebugDrawBendingConstraint(FPrimitiveDrawInterface* PD
 		if (const Softs::FXPBDBendingConstraints* const BendingConstraints = ClothConstraints.GetXBendingElementConstraints().Get())
 		{
 			const TArray<TVec4<int32>>& Constraints = BendingConstraints->GetConstraints();
-			const TArray<FSolverReal>& RestAngles = BendingConstraints->GetRestAngles();
+			const TArray<bool>& IsBuckled = BendingConstraints->GetIsBuckled();
 
-			// Color constraint edge with red or blue: Red = Buckled, Blue = Not Buckled. This is based on its current configuration, not the values used
-			// in the actual simulation.
+			// Color constraint edge with red or blue: Red = Buckled, Blue = Not Buckled.
 			for (int32 ConstraintIndex = 0; ConstraintIndex < Constraints.Num(); ++ConstraintIndex)
 			{
 				const Softs::FSolverVec3& P1 = Positions[Constraints[ConstraintIndex][0]];
 				const Softs::FSolverVec3& P2 = Positions[Constraints[ConstraintIndex][1]];
-				const Softs::FSolverVec3& P3 = Positions[Constraints[ConstraintIndex][2]];
-				const Softs::FSolverVec3& P4 = Positions[Constraints[ConstraintIndex][3]];
-				const Softs::FSolverReal Angle = Softs::FPBDBendingConstraintsBase::CalcAngle(P1, P2, P3, P4);
 
-				const bool bIsBuckled = BendingConstraints->AngleIsBuckled(Angle, RestAngles[ConstraintIndex]);
+				const bool bIsBuckled = IsBuckled[ConstraintIndex];
 
 				const FVec3 Pos0 = FVec3(P1) + LocalSpaceLocation;
 				const FVec3 Pos1 = FVec3(P2) + LocalSpaceLocation;
diff --git a/Engine/Source/Runtime/Experimental/Chaos/Private/Chaos/PBDBendingConstraints.cpp b/Engine/Source/Runtime/Experimental/Chaos/Private/Chaos/PBDBendingConstraints.cpp
--- a/Engine/Source/Runtime/Experimental/Chaos/Private/Chaos/PBDBendingConstraints.cpp
+++ b/Engine/Source/Runtime/Experimental/Chaos/Private/Chaos/PBDBendingConstraints.cpp
@@ -110,8 +110,8 @@ void FPBDBendingConstraints::Apply(FSolverParticles& Particles, const FSolverRea
 						(ispc::FVector4f*)Particles.GetPAndInvM().GetData(),
 						(ispc::FIntVector4*)&Constraints.GetData()[ColorStart],
 						&RestAngles.GetData()[ColorStart],
+						&IsBuckled.GetData()[ColorStart],
 						ExpStiffnessValue,
-						BucklingRatio,
 						ExpBucklingValue,
 						ColorSize);
 				}
@@ -144,10 +144,10 @@ void FPBDBendingConstraints::Apply(FSolverParticles& Particles, const FSolverRea
 						(ispc::FVector4f*)Particles.GetPAndInvM().GetData(),
 						(ispc::FIntVector4*)&Constraints.GetData()[ColorStart],
 						&RestAngles.GetData()[ColorStart],
+						&IsBuckled.GetData()[ColorStart],
 						StiffnessHasWeightMap,
 						&Stiffness.GetIndices().GetData()[ColorStart],
 						&Stiffness.GetTable().GetData()[0],
-						BucklingRatio,
 						BucklingStiffnessHasWeightMap,
 						&BucklingStiffness.GetIndices().GetData()[ColorStart],
 						&BucklingStiffness.GetTable().GetData()[0],
diff --git a/Engine/Source/Runtime/Experimental/Chaos/Private/Chaos/PBDBendingConstraints.ispc b/Engine/Source/Runtime/Experimental/Chaos/Private/Chaos/PBDBendingConstraints.ispc
--- a/Engine/Source/Runtime/Experimental/Chaos/Private/Chaos/PBDBendingConstraints.ispc
+++ b/Engine/Source/Runtime/Experimental/Chaos/Private/Chaos/PBDBendingConstraints.ispc
@@ -53,8 +53,8 @@ static inline float CalcGradientsAndAngle(const FVector3f& P1, const FVector3f&
 export void ApplyBendingConstraints(uniform FVector4f PandInvM[],
 									const uniform FIntVector4 Constraints[],
 									const uniform float RestAngles[],
+									const uniform bool IsBuckled[],
 									const uniform float Stiffness,
-									const uniform float BucklingRatio,
 									const uniform float BucklingStiffness,
 									const uniform int32 NumConstraints)
 {
@@ -88,8 +88,7 @@ export void ApplyBendingConstraints(uniform FVector4f PandInvM[],
 		const varying float Denom = M1 * VectorSizeSquared(Grad1) + M2 * VectorSizeSquared(Grad2) + M3 * VectorSizeSquared(Grad3) + M4 * VectorSizeSquared(Grad4);
 		
 		const varying float RestAngle = RestAngles[i];
-		const varying bool AngleIsBuckled = FLOAT_PI - abs(Angle) < BucklingRatio * (FLOAT_PI - abs(RestAngle));
-		const varying float StiffnessValue = select(AngleIsBuckled, BucklingStiffness, Stiffness);
+		const varying float StiffnessValue = select(IsBuckled[i], BucklingStiffness, Stiffness);
 
 		const varying float Delta = clamp(StiffnessValue * (Angle - RestAngle), -SingleStepAngleLimit, SingleStepAngleLimit);
 
@@ -117,10 +116,10 @@ export void ApplyBendingConstraints(uniform FVector4f PandInvM[],
 export void ApplyBendingConstraintsWithMaps(uniform FVector4f PandInvM[],
 											const uniform FIntVector4 Constraints[],
 											const uniform float RestAngles[],
+											const uniform bool IsBuckled[],
 											const uniform bool StiffnessHasMap,
 											const uniform uint8 StiffnessIndices[],
 											const uniform float StiffnessTable[],
-											const uniform float BucklingRatio,
 											const uniform bool BucklingStiffnessHasMap,
 											const uniform uint8 BucklingStiffnessIndices[],
 											const uniform float BucklingStiffnessTable[],
@@ -163,8 +162,7 @@ export void ApplyBendingConstraintsWithMaps(uniform FVector4f PandInvM[],
 		const varying float BucklingStiffness = BucklingStiffnessTable[BucklingStiffnessIndex];
 
 		const varying float RestAngle = RestAngles[i];
-		const varying bool AngleIsBuckled = FLOAT_PI - abs(Angle) < BucklingRatio * (FLOAT_PI - abs(RestAngle));
-		const varying float StiffnessValue = select(AngleIsBuckled, BucklingStiffness, Stiffness);
+		const varying float StiffnessValue = select(IsBuckled[i], BucklingStiffness, Stiffness);
 
 		const varying float Delta = clamp(StiffnessValue * (Angle - RestAngle), -SingleStepAngleLimit, SingleStepAngleLimit);
 
diff --git a/Engine/Source/Runtime/Experimental/Chaos/Private/Chaos/XPBDBendingConstraints.cpp b/Engine/Source/Runtime/Experimental/Chaos/Private/Chaos/XPBDBendingConstraints.cpp
--- a/Engine/Source/Runtime/Experimental/Chaos/Private/Chaos/XPBDBendingConstraints.cpp
+++ b/Engine/Source/Runtime/Experimental/Chaos/Private/Chaos/XPBDBendingConstraints.cpp
@@ -75,14 +75,14 @@ void FXPBDBendingConstraints::ApplyHelper(FSolverParticles& Particles, const FSo
 	const int32 i2 = Constraint[1];
 	const int32 i3 = Constraint[2];
 	const int32 i4 = Constraint[3];
-	const FSolverReal Angle = CalcAngle(Particles.P(i1), Particles.P(i2), Particles.P(i3), Particles.P(i4));
 
-	const FSolverReal BiphasicStiffnessValue = AngleIsBuckled(Angle, RestAngles[ConstraintIndex]) ? BucklingValue : StiffnessValue;
+	const FSolverReal BiphasicStiffnessValue = IsBuckled[ConstraintIndex] ? BucklingValue : StiffnessValue;
 	if (BiphasicStiffnessValue < XPBDBendMinStiffness || (Particles.InvM(i1) == (FSolverReal)0. && Particles.InvM(i2) == (FSolverReal)0. && Particles.InvM(i3) == (FSolverReal)0. && Particles.InvM(i4) == (FSolverReal)0.))
 	{
 		return;
 	}
 
+	const FSolverReal Angle = CalcAngle(Particles.P(i1), Particles.P(i2), Particles.P(i3), Particles.P(i4));
 	const TStaticArray<FSolverVec3, 4> Grads = Base::GetGradients(Particles, ConstraintIndex);
 
 	FSolverReal& Lambda = Lambdas[ConstraintIndex];
@@ -129,11 +129,11 @@ void FXPBDBendingConstraints::Apply(FSolverParticles& Particles, const FSolverRe
 						(ispc::FVector4f*)Particles.GetPAndInvM().GetData(),
 						(ispc::FIntVector4*)&Constraints.GetData()[ColorStart],
 						&RestAngles.GetData()[ColorStart],
+						&IsBuckled.GetData()[ColorStart],
 						&Lambdas.GetData()[ColorStart],
 						Dt,
 						XPBDBendMinStiffness,
 						ExpStiffnessValue,
-						BucklingRatio,
 						ExpBucklingValue,
 						ColorSize);
 				}
@@ -166,13 +166,13 @@ void FXPBDBendingConstraints::Apply(FSolverParticles& Particles, const FSolverRe
 						(ispc::FVector4f*)Particles.GetPAndInvM().GetData(),
 						(ispc::FIntVector4*)&Constraints.GetData()[ColorStart],
 						&RestAngles.GetData()[ColorStart],
+						&IsBuckled.GetData()[ColorStart],
 						&Lambdas.GetData()[ColorStart],
 						Dt,
 						XPBDBendMinStiffness,
 						StiffnessHasWeightMap,
 						&Stiffness.GetIndices().GetData()[ColorStart],
 						&Stiffness.GetTable().GetData()[0],
-						BucklingRatio,
 						BucklingStiffnessHasWeightMap,
 						&BucklingStiffness.GetIndices().GetData()[ColorStart],
 						&BucklingStiffness.GetTable().GetData()[0],
diff --git a/Engine/Source/Runtime/Experimental/Chaos/Private/Chaos/XPBDBendingConstraints.ispc b/Engine/Source/Runtime/Experimental/Chaos/Private/Chaos/XPBDBendingConstraints.ispc
--- a/Engine/Source/Runtime/Experimental/Chaos/Private/Chaos/XPBDBendingConstraints.ispc
+++ b/Engine/Source/Runtime/Experimental/Chaos/Private/Chaos/XPBDBendingConstraints.ispc
@@ -51,11 +51,11 @@ static inline float CalcGradientsAndAngle(const FVector3f& P1, const FVector3f&
 export void ApplyXPBDBendingConstraints(uniform FVector4f PandInvM[],
 									const uniform FIntVector4 Constraints[],
 									const uniform float RestAngles[],
+									const uniform bool IsBuckled[],
 									uniform float Lambdas[],
 									const uniform float Dt,
 									const uniform float MinStiffness,
 									const uniform float Stiffness,
-									const uniform float BucklingRatio,
 									const uniform float BucklingStiffness,
 									const uniform int32 NumConstraints)
 {
@@ -86,8 +86,7 @@ export void ApplyXPBDBendingConstraints(uniform FVector4f PandInvM[],
 		const varying float Angle = CalcGradientsAndAngle(P1, P2, P3, P4, Grad1, Grad2, Grad3, Grad4);
 		
 		const varying float RestAngle = RestAngles[i];
-		const varying bool AngleIsBuckled = FLOAT_PI - abs(Angle) < BucklingRatio * (FLOAT_PI - abs(RestAngle));
-		const varying float StiffnessValue = select(AngleIsBuckled, BucklingStiffness, Stiffness);
+		const varying float StiffnessValue = select(IsBuckled[i], BucklingStiffness, Stiffness);
 		
 		const varying float Lambda = Lambdas[i];
 		const varying float Alpha  = 1.f / (StiffnessValue * Dt * Dt);
@@ -119,13 +118,13 @@ export void ApplyXPBDBendingConstraints(uniform FVector4f PandInvM[],
 export void ApplyXPBDBendingConstraintsWithMaps(uniform FVector4f PandInvM[],
 											const uniform FIntVector4 Constraints[],
 											const uniform float RestAngles[],
+											const uniform bool IsBuckled[],
 											uniform float Lambdas[],
 											const uniform float Dt,
 											const uniform float MinStiffness,
 											const uniform bool StiffnessHasMap,
 											const uniform uint8 StiffnessIndices[],
 											const uniform float StiffnessTable[],
-											const uniform float BucklingRatio,
 											const uniform bool BucklingStiffnessHasMap,
 											const uniform uint8 BucklingStiffnessIndices[],
 											const uniform float BucklingStiffnessTable[],
@@ -166,8 +165,7 @@ export void ApplyXPBDBendingConstraintsWithMaps(uniform FVector4f PandInvM[],
 		const varying float BucklingStiffness = BucklingStiffnessTable[BucklingStiffnessIndex];
 		
 		const varying float RestAngle = RestAngles[i];
-		const varying bool AngleIsBuckled = FLOAT_PI - abs(Angle) < BucklingRatio * (FLOAT_PI - abs(RestAngle));
-		const varying float StiffnessValue = select(AngleIsBuckled, BucklingStiffness, Stiffness);
+		const varying float StiffnessValue = select(IsBuckled[i], BucklingStiffness, Stiffness);
 		
 		const varying float Lambda = Lambdas[i];
 		const varying float Alpha  = 1.f / (StiffnessValue * Dt * Dt);
diff --git a/Engine/Source/Runtime/Experimental/Chaos/Public/Chaos/PBDBendingConstraintsBase.h b/Engine/Source/Runtime/Experimental/Chaos/Public/Chaos/PBDBendingConstraintsBase.h
--- a/Engine/Source/Runtime/Experimental/Chaos/Public/Chaos/PBDBendingConstraintsBase.h
+++ b/Engine/Source/Runtime/Experimental/Chaos/Public/Chaos/PBDBendingConstraintsBase.h
@@ -98,7 +98,7 @@ public:
 		const FSolverReal Angle = CalcAngle(P1, P2, P3, P4);
 		const FSolverReal Denom = (InParticles.InvM(i1) * Grads[0].SizeSquared() + InParticles.InvM(i2) * Grads[1].SizeSquared() + InParticles.InvM(i3) * Grads[2].SizeSquared() + InParticles.InvM(i4) * Grads[3].SizeSquared());
 
-		const FSolverReal StiffnessValue = AngleIsBuckled(Angle, RestAngles[i]) ? ExpBucklingValue : ExpStiffnessValue;
+		const FSolverReal StiffnessValue = IsBuckled[i] ? ExpBucklingValue : ExpStiffnessValue;
 
 		constexpr FSolverReal SingleStepAngleLimit = (FSolverReal)(UE_PI * .25f); // this constraint is very non-linear. taking large steps is not accurate
 		const FSolverReal Delta = FMath::Clamp(StiffnessValue * (Angle - RestAngles[i]), -SingleStepAngleLimit, SingleStepAngleLimit);
@@ -139,8 +139,28 @@ public:
 		return UE_PI - FMath::Abs(Angle) < BucklingRatio * (UE_PI - FMath::Abs(RestAngle));
 	}
 
+	void Init(const FSolverParticles& InParticles)
+	{
+		IsBuckled.SetNumUninitialized(Constraints.Num());
+		for (int32 ConstraintIndex = 0; ConstraintIndex < Constraints.Num(); ++ConstraintIndex)
+		{
+			const TVec4<int32>& Constraint = Constraints[ConstraintIndex];
+			const int32 i1 = Constraint[0];
+			const int32 i2 = Constraint[1];
+			const int32 i3 = Constraint[2];
+			const int32 i4 = Constraint[3];
+			const FSolverVec3& P1 = InParticles.X(i1);
+			const FSolverVec3& P2 = InParticles.X(i2);
+			const FSolverVec3& P3 = InParticles.X(i3);
+			const FSolverVec3& P4 = InParticles.X(i4);
+			const FSolverReal Angle = CalcAngle(P1, P2, P3, P4);
+			IsBuckled[ConstraintIndex] = AngleIsBuckled(Angle, RestAngles[ConstraintIndex]);
+		}
+	}
+
 	const TArray<FSolverReal>& GetRestAngles() const { return RestAngles; }
 	const TArray<TVec4<int32>>& GetConstraints() const { return Constraints; }
+	const TArray<bool>& GetIsBuckled() const { return IsBuckled; }
 
 private:
 	template<class TNum>
@@ -221,6 +241,7 @@ protected:
 	FPBDStiffness BucklingStiffness;
 
 	TArray<FSolverReal> RestAngles;
+	TArray<bool> IsBuckled;
 };
 
 }  // End namespace Chaos::Softs
diff --git a/Engine/Source/Runtime/Experimental/Chaos/Public/Chaos/XPBDBendingConstraints.h b/Engine/Source/Runtime/Experimental/Chaos/Public/Chaos/XPBDBendingConstraints.h
--- a/Engine/Source/Runtime/Experimental/Chaos/Public/Chaos/XPBDBendingConstraints.h
+++ b/Engine/Source/Runtime/Experimental/Chaos/Public/Chaos/XPBDBendingConstraints.h
@@ -39,7 +39,12 @@ public:
 
 	virtual ~FXPBDBendingConstraints() override {}
 
-	void Init() const { for (FSolverReal& Lambda : Lambdas) { Lambda = (FSolverReal)0.; } }
+	void Init(const FSolverParticles& InParticles)
+	{ 
+		Lambdas.Reset();
+		Lambdas.AddZeroed(Constraints.Num());
+		FPBDBendingConstraintsBase::Init(InParticles);
+	}
 
 	// Update stiffness values
 	void SetProperties(const FSolverVec2& InStiffness, const FSolverReal InBucklingRatio, const FSolverVec2& InBucklingStiffness)
