Add Level Sets to AggregateGeometry

diff --git a/Engine/Plugins/Experimental/ChaosCloth/Source/Chaos/Private/ChaosCloth/ChaosClothingSimulationCollider.cpp b/Engine/Plugins/Experimental/ChaosCloth/Source/Chaos/Private/ChaosCloth/ChaosClothingSimulationCollider.cpp
--- a/Engine/Plugins/Experimental/ChaosCloth/Source/Chaos/Private/ChaosCloth/ChaosClothingSimulationCollider.cpp
+++ b/Engine/Plugins/Experimental/ChaosCloth/Source/Chaos/Private/ChaosCloth/ChaosClothingSimulationCollider.cpp
@@ -9,6 +9,8 @@
 #include "Chaos/TaperedCylinder.h"
 #include "Chaos/TaperedCapsule.h"
 #include "Chaos/Convex.h"
+#include "Chaos/Levelset.h"
+#include "Chaos/ImplicitObjectTransformed.h"
 #include "ClothingSimulation.h" // For context
 #include "ClothingAsset.h"
 #include "Engine/SkeletalMesh.h"
@@ -40,6 +42,7 @@ void FClothingSimulationCollider::FLODData::Add(
 	FClothingSimulationSolver* Solver,
 	FClothingSimulationCloth* Cloth,
 	const FClothCollisionData& InClothCollisionData,
+	const TArray<FLevelSetCollisionData>& InLevelSetCollisionData,
 	const FReal InScale,
 	const TArray<int32>& UsedBoneIndices)
 {
@@ -70,7 +73,8 @@ void FClothingSimulationCollider::FLODData::Add(
 
 	const uint32 NumConvexes = ClothCollisionData.Convexes.Num();
 	const int32 NumBoxes = ClothCollisionData.Boxes.Num();
-	NumGeometries = NumSpheres + NumCapsules + NumConvexes + NumBoxes;
+	const int32 NumLevelSets = InLevelSetCollisionData.Num();
+	NumGeometries = NumSpheres + NumCapsules + NumConvexes + NumBoxes + NumLevelSets;
 
 	// Retrieve cloth group Id, or use INDEX_NONE if this collider applies to all cloths (when Cloth == nullptr)
 	const uint32 GroupId = Cloth ? Cloth->GetGroupId() : INDEX_NONE;
@@ -290,7 +294,28 @@ void FClothingSimulationCollider::FLODData::Add(
 		}
 	}
 
-	UE_LOG(LogChaosCloth, VeryVerbose, TEXT("Added collisions: %d spheres, %d capsules, %d convexes, %d boxes."), NumSpheres, NumCapsules, NumConvexes, NumBoxes);
+	const int32 LevelSetOffset = BoxOffset + NumBoxes;
+	if (NumLevelSets != 0)
+	{
+		int32* const BoneIndices = Solver->GetCollisionBoneIndices(LevelSetOffset);
+		Softs::FSolverRigidTransform3* const BaseTransforms = Solver->GetCollisionBaseTransforms(LevelSetOffset);
+
+		for (int32 Index = 0; Index < NumLevelSets; ++Index)
+		{
+			// Always initialize the collision particle transforms before setting any geometry as otherwise NaNs gets detected during the bounding box updates
+			BaseTransforms[Index] = Softs::FSolverRigidTransform3::Identity;
+
+			BoneIndices[Index] = GetMappedBoneIndex(UsedBoneIndices, InLevelSetCollisionData[Index].BoneIndex);
+			UE_LOG(LogChaosCloth, VeryVerbose, TEXT("Found collision level set on bone index %d."), BoneIndices[Index]);
+
+			// Setup the collision particle geometry
+			TUniquePtr<FImplicitObject> TransformedLevelSet = MakeUnique<TImplicitObjectTransformed<FReal, 3>>(InLevelSetCollisionData[Index].LevelSet->DeepCopy(), 
+				TRigidTransform<FReal, 3>(InLevelSetCollisionData[Index].Transform));
+			Solver->SetCollisionGeometry(LevelSetOffset, Index, MoveTemp(TransformedLevelSet));
+		}
+	}
+
+	UE_LOG(LogChaosCloth, VeryVerbose, TEXT("Added collisions: %d spheres, %d capsules, %d convexes, %d boxes, % level sets."), NumSpheres, NumCapsules, NumConvexes, NumBoxes, NumLevelSets);
 }
 
 void FClothingSimulationCollider::FLODData::Remove(FClothingSimulationSolver* Solver, FClothingSimulationCloth* Cloth)
@@ -383,7 +408,7 @@ FClothCollisionData FClothingSimulationCollider::GetCollisionData(const FClothin
 	return ClothCollisionData;
 }
 
-void FClothingSimulationCollider::ExtractPhysicsAssetCollision(FClothCollisionData& ClothCollisionData, TArray<int32>& UsedBoneIndices)
+void FClothingSimulationCollider::ExtractPhysicsAssetCollision(FClothCollisionData& ClothCollisionData, TArray<FLevelSetCollisionData>& LevelSetCollisions, TArray<int32>& UsedBoneIndices)
 {
 	ClothCollisionData.Reset();
 	UsedBoneIndices.Reset();
@@ -546,6 +571,15 @@ void FClothingSimulationCollider::ExtractPhysicsAssetCollision(FClothCollisionDa
 				ClothCollisionData.Convexes.Add(Convex);
 			}
 #endif  // #if WITH_CHAOS
+
+			for (const FKLevelSetElem& LevelSetElem : AggGeom.LevelSetElems)
+			{
+				if (LevelSetElem.GetLevelSet().IsValid())
+				{
+					LevelSetCollisions.Add({ LevelSetElem.GetLevelSet(), LevelSetElem.GetTransform(), MappedBoneIndex });
+				}
+			}
+
 		}  // End for PhysAsset->SkeletalBodySetups
 	}  // End if Asset->PhysicsAsset
 }
@@ -607,10 +641,11 @@ void FClothingSimulationCollider::Add(FClothingSimulationSolver* Solver, FClothi
 
 	// Create physics asset collisions, this will affect all LODs, so store it at index 0
 	FClothCollisionData PhysicsAssetCollisionData;
+	TArray<FLevelSetCollisionData> LevelSetCollisions;
 	TArray<int32> UsedBoneIndices;
-	ExtractPhysicsAssetCollision(PhysicsAssetCollisionData, UsedBoneIndices);
+	ExtractPhysicsAssetCollision(PhysicsAssetCollisionData, LevelSetCollisions, UsedBoneIndices);
 
-	LODData[(int32)ECollisionDataType::LODless].Add(Solver, Cloth, PhysicsAssetCollisionData, Scale, UsedBoneIndices);
+	LODData[(int32)ECollisionDataType::LODless].Add(Solver, Cloth, PhysicsAssetCollisionData, LevelSetCollisions, Scale, UsedBoneIndices);
 
 	// Create legacy asset LOD collisions
 	const int32 NumLODs = Asset ? Asset->LodData.Num() : 0;
@@ -629,8 +664,10 @@ void FClothingSimulationCollider::Add(FClothingSimulationSolver* Solver, FClothi
 			AssetCollisionData.SphereConnections.Num(),
 			AssetCollisionData.Convexes.Num());
 
+		const TArray<FLevelSetCollisionData> AssetLevelSetCollisions;		// Shouldn't be any level set collider on legacy asset
+
 		// Add legacy collision
-		LODData[(int32)ECollisionDataType::LODs + Index].Add(Solver, Cloth, AssetCollisionData, Scale, Asset->UsedBoneIndices);
+		LODData[(int32)ECollisionDataType::LODs + Index].Add(Solver, Cloth, AssetCollisionData, AssetLevelSetCollisions, Scale, Asset->UsedBoneIndices);
 	}
 }
 
@@ -656,7 +693,10 @@ void FClothingSimulationCollider::PreUpdate(FClothingSimulationSolver* Solver, F
 	const int32 ExternalCollisionNumGeometries = GetNumGeometries((int32)ECollisionDataType::External);
 	const int32 ExternalCollisionOffset = GetOffset(Solver, Cloth, (int32)ECollisionDataType::External);
 
-	LODData[(int32)ECollisionDataType::External].Add(Solver, Cloth, CollisionData ? *CollisionData : FClothCollisionData());
+	// TODO: Get level sets?
+	const TArray<FLevelSetCollisionData> LevelSetCollisions;
+
+	LODData[(int32)ECollisionDataType::External].Add(Solver, Cloth, CollisionData ? *CollisionData : FClothCollisionData(), LevelSetCollisions);
 	
 	// TODO: Find a better way in case the same number but different collisions are being re-added (hash collision data? Provide user dirty function?)
 	bHasExternalCollisionChanged =
diff --git a/Engine/Plugins/Experimental/ChaosCloth/Source/Chaos/Public/ChaosCloth/ChaosClothingSimulationCollider.h b/Engine/Plugins/Experimental/ChaosCloth/Source/Chaos/Public/ChaosCloth/ChaosClothingSimulationCollider.h
--- a/Engine/Plugins/Experimental/ChaosCloth/Source/Chaos/Public/ChaosCloth/ChaosClothingSimulationCollider.h
+++ b/Engine/Plugins/Experimental/ChaosCloth/Source/Chaos/Public/ChaosCloth/ChaosClothingSimulationCollider.h
@@ -15,6 +15,7 @@ namespace Chaos
 
 	class FClothingSimulationSolver;
 	class FClothingSimulationCloth;
+	class FLevelSet;
 
 	// Collider simulation node
 	class FClothingSimulationCollider final
@@ -71,7 +72,15 @@ namespace Chaos
 		// ---- End of the debugging and visualization functions ----
 
 	private:
-		void ExtractPhysicsAssetCollision(FClothCollisionData& ClothCollisionData, TArray<int32>& UsedBoneIndices);
+
+		struct FLevelSetCollisionData
+		{
+			const TSharedPtr<Chaos::FLevelSet, ESPMode::ThreadSafe> LevelSet;
+			FTransform Transform;
+			int32 BoneIndex;
+		};
+
+		void ExtractPhysicsAssetCollision(FClothCollisionData& ClothCollisionData, TArray<FLevelSetCollisionData>& LevelSetCollisions, TArray<int32>& UsedBoneIndices);
 
 		int32 GetNumGeometries(int32 InSlotIndex) const;
 
@@ -96,6 +105,7 @@ namespace Chaos
 				FClothingSimulationSolver* Solver,
 				FClothingSimulationCloth* Cloth,
 				const FClothCollisionData& InClothCollisionData,
+				const TArray<FLevelSetCollisionData>& InLevelSetCollisionData,
 				const FReal InScale = 1.f,
 				const TArray<int32>& UsedBoneIndices = TArray<int32>());
 			void Remove(FClothingSimulationSolver* Solver, FClothingSimulationCloth* Cloth);
diff --git a/Engine/Source/Editor/Persona/Private/PhysicsAssetRenderUtils.cpp b/Engine/Source/Editor/Persona/Private/PhysicsAssetRenderUtils.cpp
--- a/Engine/Source/Editor/Persona/Private/PhysicsAssetRenderUtils.cpp
+++ b/Engine/Source/Editor/Persona/Private/PhysicsAssetRenderUtils.cpp
@@ -385,6 +385,7 @@ namespace PhysicsAssetRender
 				CollectDrawElements(DrawElements, EAggCollisionShape::Sphyl, AggGeom->SphylElems, BoneTM, Scale, i, RenderSettings, PhysicsAsset, PDI, GetPrimitiveColor, GetPrimitiveMaterial, GetPrimitiveTransform, CreateHitProxy);
 				CollectDrawElements(DrawElements, EAggCollisionShape::Convex, AggGeom->ConvexElems, BoneTM, Scale, i, RenderSettings, PhysicsAsset, PDI, GetPrimitiveColor, GetPrimitiveMaterial, GetPrimitiveTransform, CreateHitProxy);
 				CollectDrawElements(DrawElements, EAggCollisionShape::TaperedCapsule, AggGeom->TaperedCapsuleElems, BoneTM, Scale, i, RenderSettings, PhysicsAsset, PDI, GetPrimitiveColor, GetPrimitiveMaterial, GetPrimitiveTransform, CreateHitProxy);
+				CollectDrawElements(DrawElements, EAggCollisionShape::LevelSet, AggGeom->LevelSetElems, BoneTM, Scale, i, RenderSettings, PhysicsAsset, PDI, GetPrimitiveColor, GetPrimitiveMaterial, GetPrimitiveTransform, CreateHitProxy);
 
 				// Sort elements
 				DrawElements.Sort([](const DrawElement& DE1, const DrawElement& DE2) {
diff --git a/Engine/Source/Editor/PhysicsAssetEditor/Private/PhysicsAssetEditorEditMode.cpp b/Engine/Source/Editor/PhysicsAssetEditor/Private/PhysicsAssetEditorEditMode.cpp
--- a/Engine/Source/Editor/PhysicsAssetEditor/Private/PhysicsAssetEditorEditMode.cpp
+++ b/Engine/Source/Editor/PhysicsAssetEditor/Private/PhysicsAssetEditorEditMode.cpp
@@ -395,6 +395,11 @@ bool FPhysicsAssetEditorEditMode::InputDelta(FEditorViewportClient* InViewportCl
 						AggGeom->TaperedCapsuleElems[SelectedObject.PrimitiveIndex].SetTransform(SelectedObject.ManipulateTM * AggGeom->TaperedCapsuleElems[SelectedObject.PrimitiveIndex].GetTransform());
 						SelectedObject.ManipulateTM.SetIdentity();
 					}
+					else if (SelectedObject.PrimitiveType == EAggCollisionShape::LevelSet)
+					{
+						AggGeom->LevelSetElems[SelectedObject.PrimitiveIndex].SetTransform(SelectedObject.ManipulateTM * AggGeom->LevelSetElems[SelectedObject.PrimitiveIndex].GetTransform());
+						SelectedObject.ManipulateTM.SetIdentity();
+					}
 				}
 				else if (InViewportClient->GetWidgetMode() == UE::Widget::WM_Scale)
 				{
@@ -1008,6 +1013,11 @@ void FPhysicsAssetEditorEditMode::ModifyPrimitiveSize(int32 BodyIndex, EAggColli
 		check(AggGeom->TaperedCapsuleElems.IsValidIndex(PrimIndex));
 		AggGeom->TaperedCapsuleElems[PrimIndex].ScaleElem(DeltaSize, MinPrimSize);
 	}
+	else if (PrimType == EAggCollisionShape::LevelSet)
+	{
+		check(AggGeom->LevelSetElems.IsValidIndex(PrimIndex));
+		AggGeom->LevelSetElems[PrimIndex].ScaleElem(DeltaSize, MinPrimSize);
+	}
 }
 
 void FPhysicsAssetEditorEditMode::HitNothing(FEditorViewportClient* InViewportClient)
diff --git a/Engine/Source/Editor/PhysicsAssetEditor/Private/PhysicsAssetEditorSkeletalMeshComponent.cpp b/Engine/Source/Editor/PhysicsAssetEditor/Private/PhysicsAssetEditorSkeletalMeshComponent.cpp
--- a/Engine/Source/Editor/PhysicsAssetEditor/Private/PhysicsAssetEditorSkeletalMeshComponent.cpp
+++ b/Engine/Source/Editor/PhysicsAssetEditor/Private/PhysicsAssetEditorSkeletalMeshComponent.cpp
@@ -217,6 +217,12 @@ FTransform UPhysicsAssetEditorSkeletalMeshComponent::GetPrimitiveTransform(const
 		PrimTM.ScaleTranslation(Scale3D);
 		return PrimTM * BoneTM;
 	}
+	else if (PrimType == EAggCollisionShape::LevelSet)
+	{
+		FTransform PrimTM = ManTM * SharedBodySetup->AggGeom.LevelSetElems[PrimIndex].GetTransform();
+		PrimTM.ScaleTranslation(Scale3D);
+		return PrimTM * BoneTM;
+	}
 
 	// Should never reach here
 	check(0);
diff --git a/Engine/Source/Runtime/Engine/Classes/PhysicsEngine/AggregateGeom.h b/Engine/Source/Runtime/Engine/Classes/PhysicsEngine/AggregateGeom.h
--- a/Engine/Source/Runtime/Engine/Classes/PhysicsEngine/AggregateGeom.h
+++ b/Engine/Source/Runtime/Engine/Classes/PhysicsEngine/AggregateGeom.h
@@ -5,6 +5,7 @@
 #include "CoreMinimal.h"
 #include "UObject/ObjectMacros.h"
 #include "PhysicsEngine/ConvexElem.h"
+#include "PhysicsEngine/LevelSetElem.h"
 #include "PhysicsEngine/BoxElem.h"
 #include "PhysicsEngine/SphereElem.h"
 #include "PhysicsEngine/SphylElem.h"
@@ -34,6 +35,9 @@ struct ENGINE_API FKAggregateGeom
 	UPROPERTY(EditAnywhere, editfixedsize, Category = "Aggregate Geometry", meta = (DisplayName = "Tapered Capsules"))
 	TArray<FKTaperedCapsuleElem> TaperedCapsuleElems;
 
+	UPROPERTY(EditAnywhere, editfixedsize, Category = "Aggregate Geometry", meta = (DisplayName = "Level Sets"))
+	TArray<FKLevelSetElem> LevelSetElems;
+
 	class FKConvexGeomRenderInfo* RenderInfo;
 
 	FKAggregateGeom()
@@ -56,7 +60,7 @@ struct ENGINE_API FKAggregateGeom
 
 	int32 GetElementCount() const
 	{
-		return SphereElems.Num() + SphylElems.Num() + BoxElems.Num() + ConvexElems.Num() + TaperedCapsuleElems.Num();
+		return SphereElems.Num() + SphylElems.Num() + BoxElems.Num() + ConvexElems.Num() + TaperedCapsuleElems.Num() + LevelSetElems.Num();
 	}
 
 	int32 GetElementCount(EAggCollisionShape::Type Type) const;
@@ -75,6 +79,8 @@ struct ENGINE_API FKAggregateGeom
 			if (ensure(ConvexElems.IsValidIndex(Index))) { return &ConvexElems[Index]; }
 		case EAggCollisionShape::TaperedCapsule:
 			if (ensure(TaperedCapsuleElems.IsValidIndex(Index))) { return &TaperedCapsuleElems[Index]; }
+		case EAggCollisionShape::LevelSet:
+			if (ensure(LevelSetElems.IsValidIndex(Index))) { return &LevelSetElems[Index]; }
 		default:
 			ensure(false);
 			return nullptr;
@@ -93,6 +99,8 @@ struct ENGINE_API FKAggregateGeom
 		if (Index < ConvexElems.Num()) { return &ConvexElems[Index]; }
 		Index -= ConvexElems.Num();
 		if (Index < TaperedCapsuleElems.Num()) { return &TaperedCapsuleElems[Index]; }
+		Index -= TaperedCapsuleElems.Num();
+		if (Index < LevelSetElems.Num()) { return &LevelSetElems[Index]; }
 		ensure(false);
 		return nullptr;
 	}
@@ -109,6 +117,8 @@ struct ENGINE_API FKAggregateGeom
 		if (Index < ConvexElems.Num()) { return &ConvexElems[Index]; }
 		Index -= ConvexElems.Num();
 		if (Index < TaperedCapsuleElems.Num()) { return &TaperedCapsuleElems[Index]; }
+		Index -= TaperedCapsuleElems.Num();
+		if (Index < LevelSetElems.Num()) { return &LevelSetElems[Index]; }
 		ensure(false);
 		return nullptr;
 	}
@@ -120,6 +130,7 @@ struct ENGINE_API FKAggregateGeom
 		SphylElems.Empty();
 		SphereElems.Empty();
 		TaperedCapsuleElems.Empty();
+		LevelSetElems.Empty();
 
 		FreeRenderInfo();
 	}
diff --git a/Engine/Source/Runtime/Engine/Classes/PhysicsEngine/LevelSetElem.h b/Engine/Source/Runtime/Engine/Classes/PhysicsEngine/LevelSetElem.h
new file mode 100644
--- /dev/null
+++ b/Engine/Source/Runtime/Engine/Classes/PhysicsEngine/LevelSetElem.h
@@ -0,0 +1,92 @@
+// Copyright Epic Games, Inc. All Rights Reserved.
+
+#pragma once
+
+#include "CoreMinimal.h"
+#include "UObject/ObjectMacros.h"
+#include "PhysicsEngine/ShapeElem.h"
+#include "LevelSetElem.generated.h"
+
+namespace Chaos
+{
+	class FLevelSet;
+}
+
+USTRUCT()
+struct FKLevelSetElem : public FKShapeElem
+{
+	GENERATED_USTRUCT_BODY()
+
+	ENGINE_API FKLevelSetElem() :
+		FKShapeElem(EAggCollisionShape::LevelSet)
+	{}
+
+	ENGINE_API FKLevelSetElem(const FKLevelSetElem& Other)
+	{
+		CloneElem(Other);
+	}
+
+	ENGINE_API const FKLevelSetElem& operator=(const FKLevelSetElem& Other)
+	{
+		CloneElem(Other);
+		return *this;
+	}
+
+	ENGINE_API void BuildLevelSet(const FVector& GridOrigin, const TArray<double>& GridValues, const FIntVector& GridDims, float GridCellSize);
+
+	ENGINE_API FTransform GetTransform() const
+	{
+		return Transform;
+	};
+
+	ENGINE_API void SetTransform(const FTransform& InTransform)
+	{
+		ensure(InTransform.IsValid());
+		Transform = InTransform;
+	}
+
+	ENGINE_API void ScaleElem(FVector DeltaSize, float MinSize);
+
+	// Draw helpers
+
+	/** Get geometry of all cells where the level set function is less than or equal to zero */
+	ENGINE_API void GetInteriorGridCells(TArray<FBox>& CellBoxes) const;
+
+	/** Get geometry of all cell faces where level set function changes sign */
+	ENGINE_API void GetZeroIsosurfaceGridCellFaces(TArray<FVector3f>& Vertices, TArray<FIntVector>& Tris) const;
+
+	// Draw functions
+	ENGINE_API virtual void DrawElemWire(class FPrimitiveDrawInterface* PDI, const FTransform& ElemTM, const float Scale, const FColor Color) const override;
+	ENGINE_API virtual void DrawElemSolid(class FPrimitiveDrawInterface* PDI, const FTransform& ElemTM, const float Scale, const FMaterialRenderProxy* MaterialRenderProxy) const override;
+	ENGINE_API void GetElemSolid(const FTransform& ElemTM, const FVector& Scale3D, const FMaterialRenderProxy* MaterialRenderProxy, int32 ViewIndex, class FMeshElementCollector& Collector) const;
+
+	ENGINE_API FBox CalcAABB(const FTransform& BoneTM, const FVector& Scale3D) const;
+
+	bool Serialize(FArchive& Ar);
+
+	const TSharedPtr<Chaos::FLevelSet, ESPMode::ThreadSafe> GetLevelSet() const
+	{
+		return LevelSet;
+	}
+
+private:
+
+	/** Transform of this element */
+	UPROPERTY()
+	FTransform Transform;
+
+	TSharedPtr<Chaos::FLevelSet, ESPMode::ThreadSafe> LevelSet;
+
+	/** Helper function to safely copy instances of this shape*/
+	void CloneElem(const FKLevelSetElem& Other);
+};
+
+/* Enable our own serialization function to handle FLevelSet */
+template<>
+struct TStructOpsTypeTraits<FKLevelSetElem> : public TStructOpsTypeTraitsBase2<FKLevelSetElem>
+{
+	enum
+	{
+		WithSerializer = true
+	};
+};
diff --git a/Engine/Source/Runtime/Engine/Classes/PhysicsEngine/ShapeElem.h b/Engine/Source/Runtime/Engine/Classes/PhysicsEngine/ShapeElem.h
--- a/Engine/Source/Runtime/Engine/Classes/PhysicsEngine/ShapeElem.h
+++ b/Engine/Source/Runtime/Engine/Classes/PhysicsEngine/ShapeElem.h
@@ -19,6 +19,7 @@ namespace EAggCollisionShape
 		Sphyl,
 		Convex,
 		TaperedCapsule,
+		LevelSet,
 
 		Unknown
 	};
diff --git a/Engine/Source/Runtime/Engine/Private/Collision/KAggregateGeom.cpp b/Engine/Source/Runtime/Engine/Private/Collision/KAggregateGeom.cpp
--- a/Engine/Source/Runtime/Engine/Private/Collision/KAggregateGeom.cpp
+++ b/Engine/Source/Runtime/Engine/Private/Collision/KAggregateGeom.cpp
@@ -16,6 +16,7 @@
 #include "Engine/Polys.h"
 #include "PhysXIncludes.h"
 #include "Chaos/Convex.h"
+#include "Chaos/Levelset.h"
 #if INTEL_ISPC
 #include "KAggregateGeom.ispc.generated.h"
 #endif
@@ -105,6 +106,11 @@ FBox FKAggregateGeom::CalcAABB(const FTransform& Transform) const
 		Box += TaperedCapsuleElems[i].CalcAABB(BoneTM, ScaleFactor);
 	}
 
+	for (int32 i = 0; i < LevelSetElems.Num(); i++)
+	{
+		Box += LevelSetElems[i].CalcAABB(BoneTM, (FVector)Scale3D);
+	}
+
 	return Box;
 
 }
@@ -494,6 +500,147 @@ FBox FKTaperedCapsuleElem::CalcAABB(const FTransform& BoneTM, float Scale) const
 	return Result;
 }
 
+
+///////////////////////////////////////
+//////// FKLevelSetElem ///////////////
+///////////////////////////////////////
+
+FBox FKLevelSetElem::CalcAABB(const FTransform& BoneTM, const FVector& Scale3D) const
+{
+	FBox Box;
+
+	if (LevelSet.IsValid())
+	{
+		Box = FBox(LevelSet->BoundingBox().Min(), LevelSet->BoundingBox().Max());
+		const FTransform LocalToWorld = FTransform(FQuat::Identity, FVector::ZeroVector, Scale3D) * BoneTM;
+		return Box.TransformBy(Transform * LocalToWorld);
+	}
+
+	return Box;
+}
+
+void FKLevelSetElem::BuildLevelSet(const FVector& GridOrigin, const TArray<double>& GridValues, const FIntVector& GridDims, float GridCellSize)
+{
+	const Chaos::FVec3 Min = GridOrigin;
+	const Chaos::FVec3 Max = GridOrigin + Chaos::FVec3(GridCellSize * (FVector)GridDims);
+	const Chaos::TVec3<int32> ChaosDim(GridDims[0], GridDims[1], GridDims[2]);
+	
+	Chaos::TUniformGrid<Chaos::FReal, 3> ChaosGrid(Min, Max, ChaosDim);
+	Chaos::TArrayND<Chaos::FReal, 3> Phi( Chaos::TVec3<int32>(GridDims[0], GridDims[1], GridDims[2]), GridValues );
+
+	LevelSet = TSharedPtr<Chaos::FLevelSet>(new Chaos::FLevelSet(MoveTemp(ChaosGrid), MoveTemp(Phi), 0));
+}
+
+void FKLevelSetElem::GetInteriorGridCells( TArray<FBox>& CellBoxes ) const
+{
+	if (LevelSet.IsValid())
+	{
+		const Chaos::TUniformGrid<Chaos::FReal, 3>& Grid = LevelSet->GetGrid();
+		const Chaos::TVector<int32, 3> Cells = Grid.Counts();
+		for (int i = 0; i < Cells.X; ++i)
+		{
+			for (int j = 0; j < Cells.Y; ++j)
+			{
+				for (int k = 0; k < Cells.Z; ++k)
+				{
+					const double Value = LevelSet->GetPhiArray()(i, j, k);
+
+					if (Value <= 0.0)
+					{
+						const FVector3d CellMin = Grid.MinCorner() + Grid.Dx() * FVector3d(i, j, k);
+						const FVector3d CellMax = CellMin + Grid.Dx() * FVector3d(1, 1, 1);
+						CellBoxes.Emplace( FBox(CellMin, CellMax) );
+					}
+				}
+			}
+		}
+	}
+}
+
+
+void FKLevelSetElem::GetZeroIsosurfaceGridCellFaces(TArray<FVector3f>& Vertices, TArray<FIntVector>& Tris ) const
+{
+	if (LevelSet.IsValid())
+	{
+		const Chaos::TUniformGrid<Chaos::FReal, 3>& Grid = LevelSet->GetGrid();
+		const Chaos::TVector<int32, 3> Cells = Grid.Counts();
+		for (int i = 0; i < Cells.X-1; ++i)
+		{
+			for (int j = 0; j < Cells.Y-1; ++j)
+			{
+				for (int k = 0; k < Cells.Z-1; ++k)
+				{
+					const double Sign = FMath::Sign(LevelSet->GetPhiArray()(i, j, k));
+					const double SignNextI = FMath::Sign(LevelSet->GetPhiArray()(i+1, j, k));
+					const double SignNextJ = FMath::Sign(LevelSet->GetPhiArray()(i, j+1, k));
+					const double SignNextK = FMath::Sign(LevelSet->GetPhiArray()(i, j, k+1));
+
+					const FVector3d CellMin = Grid.MinCorner() + Grid.Dx() * FVector3d(i, j, k);
+
+					if (Sign > SignNextI)
+					{
+						const int32 V0 = Vertices.Emplace( CellMin + FVector3d(1, 0, 0) );
+						const int32 V1 = Vertices.Emplace( CellMin + FVector3d(1, 1, 0) );
+						const int32 V2 = Vertices.Emplace( CellMin + FVector3d(1, 1, 1) );
+						const int32 V3 = Vertices.Emplace( CellMin + FVector3d(1, 0, 1) );
+						Tris.Emplace(FIntVector(V0, V1, V2));
+						Tris.Emplace(FIntVector(V2, V3, V0));
+					}
+					else if (Sign < SignNextI)
+					{
+						const int32 V0 = Vertices.Emplace(CellMin + FVector3d(1, 0, 0));
+						const int32 V1 = Vertices.Emplace(CellMin + FVector3d(1, 1, 0));
+						const int32 V2 = Vertices.Emplace(CellMin + FVector3d(1, 1, 1));
+						const int32 V3 = Vertices.Emplace(CellMin + FVector3d(1, 0, 1));
+						Tris.Emplace(FIntVector(V0, V2, V1));
+						Tris.Emplace(FIntVector(V2, V0, V3));
+					}
+
+					
+					if (Sign > SignNextJ)
+					{
+						const int32 V0 = Vertices.Emplace(CellMin + FVector3d(0, 1, 0));
+						const int32 V1 = Vertices.Emplace(CellMin + FVector3d(1, 1, 0));
+						const int32 V2 = Vertices.Emplace(CellMin + FVector3d(1, 1, 1));
+						const int32 V3 = Vertices.Emplace(CellMin + FVector3d(0, 1, 1));
+						Tris.Emplace(FIntVector(V0, V2, V1));
+						Tris.Emplace(FIntVector(V2, V0, V3));
+					}
+					else if (Sign < SignNextJ)
+					{
+						const int32 V0 = Vertices.Emplace(CellMin + FVector3d(0, 1, 0));
+						const int32 V1 = Vertices.Emplace(CellMin + FVector3d(1, 1, 0));
+						const int32 V2 = Vertices.Emplace(CellMin + FVector3d(1, 1, 1));
+						const int32 V3 = Vertices.Emplace(CellMin + FVector3d(0, 1, 1));
+						Tris.Emplace(FIntVector(V0, V1, V2));
+						Tris.Emplace(FIntVector(V2, V3, V0));
+					}
+
+					if (Sign > SignNextK)
+					{
+						const int32 V0 = Vertices.Emplace(CellMin + FVector3d(0, 0, 1));
+						const int32 V1 = Vertices.Emplace(CellMin + FVector3d(1, 0, 1));
+						const int32 V2 = Vertices.Emplace(CellMin + FVector3d(1, 1, 1));
+						const int32 V3 = Vertices.Emplace(CellMin + FVector3d(0, 1, 1));
+						Tris.Emplace(FIntVector(V0, V1, V2));
+						Tris.Emplace(FIntVector(V2, V3, V0));
+					}
+					else if (Sign < SignNextK)
+					{
+						const int32 V0 = Vertices.Emplace(CellMin + FVector3d(0, 0, 1));
+						const int32 V1 = Vertices.Emplace(CellMin + FVector3d(1, 0, 1));
+						const int32 V2 = Vertices.Emplace(CellMin + FVector3d(1, 1, 1));
+						const int32 V3 = Vertices.Emplace(CellMin + FVector3d(0, 1, 1));
+						Tris.Emplace(FIntVector(V0, V2, V1));
+						Tris.Emplace(FIntVector(V2, V0, V3));
+					}
+
+				}
+			}
+		}
+	}
+}
+
 static const float DIST_COMPARE_THRESH = 0.1f;
 static const float DIR_COMPARE_THRESH = 0.0003f; // about 1 degree
 
@@ -809,3 +956,21 @@ FArchive& operator<<(FArchive& Ar,FKConvexElem& Elem)
 	}
 	return Ar;
 }
+
+bool FKLevelSetElem::Serialize(FArchive& Ar)
+{
+	Ar << Transform;
+
+	if (Ar.IsLoading())
+	{
+		LevelSet = TSharedPtr<Chaos::FLevelSet>(new Chaos::FLevelSet());
+	}
+
+	if (LevelSet.IsValid())
+	{
+		LevelSet->Serialize(Ar);
+	}
+
+	return true;
+}
+
diff --git a/Engine/Source/Runtime/Engine/Private/PhysicsEngine/BodySetup.cpp b/Engine/Source/Runtime/Engine/Private/PhysicsEngine/BodySetup.cpp
--- a/Engine/Source/Runtime/Engine/Private/PhysicsEngine/BodySetup.cpp
+++ b/Engine/Source/Runtime/Engine/Private/PhysicsEngine/BodySetup.cpp
@@ -27,6 +27,7 @@
 
 #include "ChaosCheck.h"
 #include "Chaos/Convex.h"
+#include "Chaos/Levelset.h"
 
 #include "PhysXCookHelper.h"
 
@@ -222,6 +223,8 @@ void UBodySetup::AddCollisionFrom(const FKAggregateGeom& FromAggGeom)
 		ConvexElem.ResetChaosConvexMesh();
 #endif
 	}
+
+	AggGeom.LevelSetElems.Append(FromAggGeom.LevelSetElems);
 }
 
 void UBodySetup::GetCookInfo(FCookBodySetupInfo& OutCookInfo, EPhysXMeshCookFlags InCookFlags) const
@@ -1990,6 +1993,8 @@ int32 FKAggregateGeom::GetElementCount(EAggCollisionShape::Type Type) const
 		return SphereElems.Num();
 	case EAggCollisionShape::TaperedCapsule:
 		return TaperedCapsuleElems.Num();
+	case EAggCollisionShape::LevelSet:
+		return LevelSetElems.Num();
 	default:
 		return 0;
 	}
@@ -2626,4 +2631,20 @@ float UBodySetup::GetVolume(const FVector& Scale) const
 	return AggGeom.GetVolume(Scale);
 }
 
+
+/** Helper function to safely copy instances of this shape*/
+void FKLevelSetElem::CloneElem(const FKLevelSetElem& Other)
+{
+	Super::CloneElem(Other);
+	LevelSet = Other.LevelSet;
+	Transform = Other.Transform;
+}
+
+void FKLevelSetElem::ScaleElem(FVector DeltaSize, float MinSize)
+{
+	FTransform ScaledTransform = GetTransform();
+	ScaledTransform.SetScale3D(ScaledTransform.GetScale3D() + DeltaSize);
+	SetTransform(ScaledTransform);
+}
+
 #undef LOCTEXT_NAMESPACE
diff --git a/Engine/Source/Runtime/Engine/Private/PhysicsEngine/PhysDrawing.cpp b/Engine/Source/Runtime/Engine/Private/PhysicsEngine/PhysDrawing.cpp
--- a/Engine/Source/Runtime/Engine/Private/PhysicsEngine/PhysDrawing.cpp
+++ b/Engine/Source/Runtime/Engine/Private/PhysicsEngine/PhysDrawing.cpp
@@ -919,6 +919,64 @@ void FKConvexElem::AddCachedSolidConvexGeom(TArray<FDynamicMeshVertex>& VertexBu
 #endif // WITH_PHYSX
 }
 
+/////////////////////////////////////////////////////////////////////////////////////
+// FKLevelSetElem
+/////////////////////////////////////////////////////////////////////////////////////
+
+void FKLevelSetElem::DrawElemWire(class FPrimitiveDrawInterface* PDI, const FTransform& ElemTM, float Scale, const FColor Color) const
+{
+	TArray<FBox> Boxes;
+
+	// Bounding box
+	Boxes.Add(CalcAABB(FTransform::Identity, FVector::OneVector));
+
+	// Cells with negative Phi values
+	GetInteriorGridCells(Boxes);
+
+	for (const FBox& Box : Boxes)
+	{
+		DrawWireBox(PDI, ElemTM.ToMatrixWithScale(), Box, Color, SDPG_World, 0.f);
+	}
+}
+
+void FKLevelSetElem::DrawElemSolid(class FPrimitiveDrawInterface* PDI, const FTransform& ElemTM, float Scale, const FMaterialRenderProxy* MaterialRenderProxy) const
+{
+	TArray<FVector3f> Vertices;
+	TArray<FIntVector> Tris;
+	GetZeroIsosurfaceGridCellFaces(Vertices, Tris);
+
+	FDynamicMeshBuilder MeshBuilder(PDI->View->GetFeatureLevel());
+	for (const FVector3f& V : Vertices)
+	{
+		MeshBuilder.AddVertex(FDynamicMeshVertex(V));
+	}
+	for (const FIntVector& T : Tris)
+	{
+		MeshBuilder.AddTriangle(T[0], T[1], T[2]);
+	}
+
+	MeshBuilder.Draw(PDI, ElemTM.ToMatrixWithScale(), MaterialRenderProxy, SDPG_World, 0.f);
+}
+
+void FKLevelSetElem::GetElemSolid(const FTransform& ElemTM, const FVector& Scale3D, const FMaterialRenderProxy* MaterialRenderProxy, int32 ViewIndex, class FMeshElementCollector& Collector) const
+{
+	TArray<FVector3f> Vertices;
+	TArray<FIntVector> Tris;
+	GetZeroIsosurfaceGridCellFaces(Vertices, Tris);
+
+	FDynamicMeshBuilder MeshBuilder(Collector.GetFeatureLevel());
+	for (const FVector3f& V : Vertices)
+	{
+		MeshBuilder.AddVertex(FDynamicMeshVertex(V));
+	}
+	for (const FIntVector& T : Tris)
+	{
+		MeshBuilder.AddTriangle(T[0], T[1], T[2]);
+	}
+
+	MeshBuilder.GetMesh(ElemTM.ToMatrixWithScale(), MaterialRenderProxy, SDPG_World, false, false, ViewIndex, Collector);
+}
+
 /////////////////////////////////////////////////////////////////////////////////////
 // FKConvexGeomRenderInfo
 /////////////////////////////////////////////////////////////////////////////////////
@@ -1076,6 +1134,17 @@ void FKAggregateGeom::GetAggGeom(const FTransform& Transform, const FColor Color
 		else
 			TaperedCapsuleElems[i].DrawElemWire(Collector.GetPDI(ViewIndex), ElemTM, Scale3D, Color);
 	}
+
+	for (int32 i = 0; i < LevelSetElems.Num(); i++)
+	{
+		FTransform ElemTM = LevelSetElems[i].GetTransform();
+		ElemTM *= Transform;
+
+		if (bDrawSolid)
+			LevelSetElems[i].GetElemSolid(ElemTM, Scale3D, MatInst, ViewIndex, Collector);
+		else
+			LevelSetElems[i].DrawElemWire(Collector.GetPDI(ViewIndex), ElemTM, 1.f, Color);
+	}
 }
 
 /** Release the RenderInfo (if its there) and safely clean up any resources. Not thread safe, but can be called from any thread (conditionally safe). */
diff --git a/Engine/Source/Runtime/Experimental/Chaos/Private/Chaos/Levelset.cpp b/Engine/Source/Runtime/Experimental/Chaos/Private/Chaos/Levelset.cpp
--- a/Engine/Source/Runtime/Experimental/Chaos/Private/Chaos/Levelset.cpp
+++ b/Engine/Source/Runtime/Experimental/Chaos/Private/Chaos/Levelset.cpp
@@ -166,6 +166,18 @@ FLevelSet::FLevelSet(std::istream& Stream)
 	ComputeNormals();
 }
 
+FLevelSet::FLevelSet(TUniformGrid<FReal, 3>&& Grid, TArrayND<FReal, 3>&& Phi, int32 BandWidth)
+	: FImplicitObject(EImplicitObject::HasBoundingBox, ImplicitObjectType::LevelSet)
+	, MGrid(MoveTemp(Grid))
+	, MPhi(MoveTemp(Phi))
+	, MNormals(MGrid)
+	, MLocalBoundingBox(MGrid.MinCorner(), MGrid.MaxCorner())
+	, MOriginalLocalBoundingBox(MLocalBoundingBox)
+	, MBandWidth(BandWidth)
+{
+	ComputeNormals();
+}
+
 FLevelSet::FLevelSet(FLevelSet&& Other)
     : FImplicitObject(EImplicitObject::HasBoundingBox, ImplicitObjectType::LevelSet)
     , MGrid(MoveTemp(Other.MGrid))
diff --git a/Engine/Source/Runtime/Experimental/Chaos/Public/Chaos/Levelset.h b/Engine/Source/Runtime/Experimental/Chaos/Public/Chaos/Levelset.h
--- a/Engine/Source/Runtime/Experimental/Chaos/Public/Chaos/Levelset.h
+++ b/Engine/Source/Runtime/Experimental/Chaos/Public/Chaos/Levelset.h
@@ -10,6 +10,8 @@
 
 namespace Chaos { class FErrorReporter; }
 
+struct FKLevelSetElem;
+
 namespace Chaos
 {
 
@@ -32,6 +34,7 @@ class CHAOS_API FLevelSet final : public FImplicitObject
 	FLevelSet(FErrorReporter& ErrorReporter, const TUniformGrid<FReal, 3>& InGrid, const FParticles& InParticles, const FTriangleMesh& Mesh, const int32 BandWidth = 0);
 	FLevelSet(FErrorReporter& ErrorReporter, const TUniformGrid<FReal, 3>& InGrid, const FImplicitObject& InObject, const int32 BandWidth = 0, const bool bUseObjectPhi = false);
 	FLevelSet(std::istream& Stream);
+	FLevelSet(TUniformGrid<FReal, 3>&& Grid, TArrayND<FReal, 3>&& Phi, int32 BandWidth);
 	FLevelSet(const FLevelSet& Other) = delete;
 	FLevelSet(FLevelSet&& Other);
 	virtual ~FLevelSet();
@@ -195,6 +198,7 @@ class CHAOS_API FLevelSet final : public FImplicitObject
 private:
 	FLevelSet() : FImplicitObject(EImplicitObject::HasBoundingBox, ImplicitObjectType::LevelSet) {}	//needed for serialization
 	friend FImplicitObject;	//needed for serialization
+	friend ::FKLevelSetElem; //needed for serialization
 };
 
 template <typename T, int d>
