Fixed NonUnity compile for ChaosClothingSimulation.cpp (missing includes that were dragged in by other .cpp files in unity builds).

diff --git a/Engine/Plugins/Experimental/ChaosCloth/Source/Chaos/Private/ChaosCloth/ChaosClothingSimulation.cpp b/Engine/Plugins/Experimental/ChaosCloth/Source/Chaos/Private/ChaosCloth/ChaosClothingSimulation.cpp
--- a/Engine/Plugins/Experimental/ChaosCloth/Source/Chaos/Private/ChaosCloth/ChaosClothingSimulation.cpp
+++ b/Engine/Plugins/Experimental/ChaosCloth/Source/Chaos/Private/ChaosCloth/ChaosClothingSimulation.cpp
@@ -59,6 +59,7 @@
 #endif
 
 #if INTEL_ISPC && !UE_BUILD_SHIPPING  // Include only used for the Ispc command
+#include "Chaos/PBDBendingConstraints.h"
 #include "Chaos/PBDCollisionConstraints.h"
 #include "Chaos/PBDEvolution.h"
 #include "Chaos/PBDSpringConstraints.h"
@@ -68,6 +69,8 @@
 #include "Chaos/PBDSphericalConstraint.h"
 #include "Chaos/PBDLongRangeConstraints.h"
 #include "Chaos/VelocityField.h"
+#include "Chaos/XPBDSpringConstraints.h"
+#include "Chaos/XPBDBendingConstraints.h"
 #endif
 
 #if INTEL_ISPC && !UE_BUILD_SHIPPING
