Fix CIS for chaos cache

diff --git a/Engine/Plugins/Experimental/ChaosCaching/Source/ChaosCaching/Public/Chaos/CacheManagerActor.h b/Engine/Plugins/Experimental/ChaosCaching/Source/ChaosCaching/Public/Chaos/CacheManagerActor.h
--- a/Engine/Plugins/Experimental/ChaosCaching/Source/ChaosCaching/Public/Chaos/CacheManagerActor.h
+++ b/Engine/Plugins/Experimental/ChaosCaching/Source/ChaosCaching/Public/Chaos/CacheManagerActor.h
@@ -45,7 +45,7 @@ struct CHAOSCACHING_API FObservedComponent
 {
 	GENERATED_BODY()
 
-		FObservedComponent()
+	FObservedComponent()
 		: CacheName(NAME_None)
 		, bIsSimulating(true)
 		, Cache(nullptr)
@@ -53,6 +53,19 @@ struct CHAOSCACHING_API FObservedComponent
 	{
 	}
 
+	FObservedComponent(const FObservedComponent& OtherComponent)
+	{
+		CacheName = OtherComponent.CacheName;
+		bIsSimulating = OtherComponent.bIsSimulating;
+		bTriggered = OtherComponent.bTriggered;
+		AbsoluteTime = OtherComponent.AbsoluteTime;
+		TimeSinceTrigger = OtherComponent.TimeSinceTrigger;
+		Cache = OtherComponent.Cache;
+		TickRecord = OtherComponent.TickRecord;
+		BestFitAdapter = OtherComponent.BestFitAdapter;
+		SoftComponentRef = OtherComponent.SoftComponentRef;
+	}
+
 	FObservedComponent& operator=(const FObservedComponent& OtherComponent)
 	{
 		CacheName = OtherComponent.CacheName;
diff --git a/Engine/Plugins/Experimental/ChaosCloth/Source/Chaos/Private/ChaosCloth/SkeletalMeshComponentCacheAdapter.cpp b/Engine/Plugins/Experimental/ChaosCloth/Source/Chaos/Private/ChaosCloth/SkeletalMeshComponentCacheAdapter.cpp
--- a/Engine/Plugins/Experimental/ChaosCloth/Source/Chaos/Private/ChaosCloth/SkeletalMeshComponentCacheAdapter.cpp
+++ b/Engine/Plugins/Experimental/ChaosCloth/Source/Chaos/Private/ChaosCloth/SkeletalMeshComponentCacheAdapter.cpp
@@ -114,7 +114,7 @@ namespace Chaos
 				TArray<float>* PendingPY = EvaluatedResult.Channels.Find(PositionYName);
 				TArray<float>* PendingPZ = EvaluatedResult.Channels.Find(PositionZName);
 					
-				if(PendingVX && PendingVY && PendingVZ && PendingPZ && PendingPY && PendingPZ)
+				if(PendingVX && PendingVY && PendingVZ && PendingPX && PendingPY && PendingPZ)
 				{
 					for(uint32 ParticleIndex = 0; ParticleIndex < NumParticles; ++ParticleIndex)
 					{
@@ -223,7 +223,7 @@ namespace Chaos
 					TArray<float>* PendingPY = EvaluatedResult.Channels.Find(PositionYName);
 					TArray<float>* PendingPZ = EvaluatedResult.Channels.Find(PositionZName);
 					
-					if(PendingVX && PendingVY && PendingVZ && PendingPZ && PendingPY && PendingPZ)
+					if(PendingVX && PendingVY && PendingVZ && PendingPX && PendingPY && PendingPZ)
 					{
 						for(int32 ParticleIndex = 0; ParticleIndex < NumCurves; ++ParticleIndex)
 						{
