Chaos Cloth - Fixed an issue with external collisions not resetting the collision particles index during an LOD change, causing the external collision sources to be added multiple times.

diff --git a/Engine/Plugins/Experimental/ChaosCloth/Source/Chaos/Private/ChaosCloth/ChaosClothingSimulation.cpp b/Engine/Plugins/Experimental/ChaosCloth/Source/Chaos/Private/ChaosCloth/ChaosClothingSimulation.cpp
--- a/Engine/Plugins/Experimental/ChaosCloth/Source/Chaos/Private/ChaosCloth/ChaosClothingSimulation.cpp
+++ b/Engine/Plugins/Experimental/ChaosCloth/Source/Chaos/Private/ChaosCloth/ChaosClothingSimulation.cpp
@@ -671,8 +671,7 @@ void FClothingSimulation::GetSimulationData(
 		const int32 LODIndex = Cloth->GetMesh()->GetLODIndex();
 		if (LODIndex != Cloth->GetLODIndex(Solver.Get()))
 		{
-			Cloth->PreUpdate(Solver.Get());  // Currently doing colliders' update, not really required here, but this could later change (technically PreUpdates are for any non parallel cloth updates)
-			Cloth->Update(Solver.Get());  // LOD switching
+			Solver->Update(FSolverReal(0.));  // Update for LOD switching, but do not simulate
 		}
 
 		if (Cloth->GetOffset(Solver.Get()) == INDEX_NONE || Cloth->GetLODIndex(Solver.Get()) == INDEX_NONE)
diff --git a/Engine/Plugins/Experimental/ChaosCloth/Source/Chaos/Private/ChaosCloth/ChaosClothingSimulationSolver.cpp b/Engine/Plugins/Experimental/ChaosCloth/Source/Chaos/Private/ChaosCloth/ChaosClothingSimulationSolver.cpp
--- a/Engine/Plugins/Experimental/ChaosCloth/Source/Chaos/Private/ChaosCloth/ChaosClothingSimulationSolver.cpp
+++ b/Engine/Plugins/Experimental/ChaosCloth/Source/Chaos/Private/ChaosCloth/ChaosClothingSimulationSolver.cpp
@@ -19,6 +19,7 @@
 
 DECLARE_CYCLE_STAT(TEXT("Chaos Cloth Solver Update"), STAT_ChaosClothSolverUpdate, STATGROUP_ChaosCloth);
 DECLARE_CYCLE_STAT(TEXT("Chaos Cloth Solver Update Cloths"), STAT_ChaosClothSolverUpdateCloths, STATGROUP_ChaosCloth);
+DECLARE_CYCLE_STAT(TEXT("Chaos Cloth Solver Update Solver Fields"), STAT_ChaosClothSolverUpdateSolverFields, STATGROUP_ChaosCloth);
 DECLARE_CYCLE_STAT(TEXT("Chaos Cloth Solver Update Pre Solver Step"), STAT_ChaosClothSolverUpdatePreSolverStep, STATGROUP_ChaosCloth);
 DECLARE_CYCLE_STAT(TEXT("Chaos Cloth Solver Update Solver Step"), STAT_ChaosClothSolverUpdateSolverStep, STATGROUP_ChaosCloth);
 DECLARE_CYCLE_STAT(TEXT("Chaos Cloth Solver Update Post Solver Step"), STAT_ChaosClothSolverUpdatePostSolverStep, STATGROUP_ChaosCloth);
@@ -880,20 +881,6 @@ void FClothingSimulationSolver::Update(Softs::FSolverReal InDeltaTime)
 	TRACE_CPUPROFILER_EVENT_SCOPE(FClothingSimulationSolver_Update);
 	SCOPE_CYCLE_COUNTER(STAT_ChaosClothSolverUpdate);
 
-	if (!bClothSolverUseImprovedTimeStepSmoothing)
-	{
-		// Filter delta time to smoothen time variations and prevent unwanted vibrations
-		// Note: This is now deprecated and replaced by in solver input force timestep smoothing
-		constexpr Softs::FSolverReal DeltaTimeDecay = (Softs::FSolverReal)0.1;
-		const Softs::FSolverReal PrevDeltaTime = DeltaTime;
-		DeltaTime = DeltaTime + (InDeltaTime - DeltaTime) * DeltaTimeDecay;
-	}
-	else
-	{
-		// Update time step
-		DeltaTime = InDeltaTime;
-	}
-
 #if !UE_BUILD_SHIPPING
 	// Introduce artificial hitches for debugging any simulation jitter
 	if (ClothSolverDebugHitchLength && ClothSolverDebugHitchInterval)
@@ -919,9 +906,6 @@ void FClothingSimulationSolver::Update(Softs::FSolverReal InDeltaTime)
 		// Clear external collisions so that they can be re-added
 		CollisionParticlesSize = 0;
 
-		// Compute the solver field forces/velocities for future use in the AddExternalForces
-		UpdateSolverField();
-
 		// Run sequential pre-updates first
 		for (FClothingSimulationCloth* const Cloth : Cloths)
 		{
@@ -952,19 +936,43 @@ void FClothingSimulationSolver::Update(Softs::FSolverReal InDeltaTime)
 		EventPreSolve.Broadcast(DeltaTime);
 	}
 
-	// Advance Sim
+	const bool bAdvanceTimeStep = (InDeltaTime > Softs::FSolverReal(0.));
+	if (bAdvanceTimeStep)
 	{
-		TRACE_CPUPROFILER_EVENT_SCOPE(FClothingSimulationSolver_UpdateSolverStep);
-		SCOPE_CYCLE_COUNTER(STAT_ChaosClothSolverUpdateSolverStep);
-		SCOPE_CYCLE_COUNTER(STAT_ClothInternalSolve);
+		if (!bClothSolverUseImprovedTimeStepSmoothing)
+		{
+			// Filter delta time to smoothen time variations and prevent unwanted vibrations
+			// Note: This is now deprecated and replaced by in solver input force timestep smoothing
+			constexpr Softs::FSolverReal DeltaTimeDecay = (Softs::FSolverReal)0.1;
+			const Softs::FSolverReal PrevDeltaTime = DeltaTime;
+			DeltaTime = DeltaTime + (InDeltaTime - DeltaTime) * DeltaTimeDecay;
+		}
+		else
+		{
+			// Update time step
+			DeltaTime = InDeltaTime;
+		}
+
+		// Compute the solver field forces/velocities for future use in the AddExternalForces
+		{
+			TRACE_CPUPROFILER_EVENT_SCOPE(FClothingSimulationSolver_UpdateSolverFields);
+			SCOPE_CYCLE_COUNTER(STAT_ChaosClothSolverUpdateSolverFields);
 
-		if(bEnableSolver)
+			UpdateSolverField();
+		}
+
+		// Advance Sim
+		if (bEnableSolver)
 		{
+			TRACE_CPUPROFILER_EVENT_SCOPE(FClothingSimulationSolver_UpdateSolverStep);
+			SCOPE_CYCLE_COUNTER(STAT_ChaosClothSolverUpdateSolverStep);
+			SCOPE_CYCLE_COUNTER(STAT_ClothInternalSolve);
+
 			// Update solver time dependent parameters
 			constexpr Softs::FSolverReal SolverFrequency = 60.f;  // TODO: This could become a solver property
 
 			const int32 TimeDependentNumIterations = bClothSolverDisableTimeDependentNumIterations ?
-				NumIterations : 
+				NumIterations :
 				(int32)(SolverFrequency * DeltaTime * (Softs::FSolverReal)NumIterations);
 
 			Evolution->SetIterations(FMath::Clamp(TimeDependentNumIterations, 1, MaxNumIterations));
@@ -976,10 +984,10 @@ void FClothingSimulationSolver::Update(Softs::FSolverReal InDeltaTime)
 			{
 				Evolution->AdvanceOneTimeStep(SubstepDeltaTime, bClothSolverUseImprovedTimeStepSmoothing);
 			}
-		}
 
-		Time = Evolution->GetTime();
-		UE_LOG(LogChaosCloth, VeryVerbose, TEXT("DeltaTime: %.6f, Time = %.6f"), DeltaTime, Time);
+			Time = Evolution->GetTime();
+			UE_LOG(LogChaosCloth, VeryVerbose, TEXT("DeltaTime: %.6f, Time = %.6f"), DeltaTime, Time);
+		}
 	}
 
 	// Post solver step, update normals, ...etc
