Chaos Cloth:
- Test for the Fields validity inside the solver force function instead of in the cloth update function, so that the field solver update can be done after the cloth update, and only when the solver is running.
- Limit the solver field update to the active particles for the current LOD.
- Reset the field's outputs to avoid it affecting the cloth once it is no longer active.

diff --git a/Engine/Plugins/Experimental/ChaosCloth/Source/Chaos/Private/ChaosCloth/ChaosClothingSimulationSolver.cpp b/Engine/Plugins/Experimental/ChaosCloth/Source/Chaos/Private/ChaosCloth/ChaosClothingSimulationSolver.cpp
--- a/Engine/Plugins/Experimental/ChaosCloth/Source/Chaos/Private/ChaosCloth/ChaosClothingSimulationSolver.cpp
+++ b/Engine/Plugins/Experimental/ChaosCloth/Source/Chaos/Private/ChaosCloth/ChaosClothingSimulationSolver.cpp
@@ -715,9 +715,6 @@ void FClothingSimulationSolver::AddExternalForces(uint32 GroupId, bool bUseLegac
 {
 	if (Evolution)
 	{
-		const bool bHasVelocityField = !PerSolverField.GetOutputResults(EFieldCommandOutputType::LinearVelocity).IsEmpty();
-		const bool bHasForceField = !PerSolverField.GetOutputResults(EFieldCommandOutputType::LinearForce).IsEmpty();
-
 		const FVec3& AngularDisplacement = FictitiousAngularDisplacements[GroupId];
 		const bool bHasFictitiousForces = !AngularDisplacement.IsNearlyZero();
 
@@ -725,19 +722,19 @@ void FClothingSimulationSolver::AddExternalForces(uint32 GroupId, bool bUseLegac
 		const FVec3 LegacyWindVelocity = WindVelocity * LegacyWindMultiplier;
 
 		Evolution->GetForceFunction(GroupId) =
-			[this, bHasVelocityField, bHasForceField, bHasFictitiousForces, bUseLegacyWind, LegacyWindVelocity, AngularDisplacement](Softs::FSolverParticles& Particles, const FReal Dt, const int32 Index)
+			[this, bHasFictitiousForces, bUseLegacyWind, LegacyWindVelocity, AngularDisplacement](Softs::FSolverParticles& Particles, const FReal Dt, const int32 Index)
 			{
 				FVec3 Forces((FReal)0.);
 
-				if (bHasVelocityField)
+				const TArray<FVector>& LinearVelocities = PerSolverField.GetOutputResults(EFieldCommandOutputType::LinearVelocity);
+				if (!LinearVelocities.IsEmpty())
 				{
-					const TArray<FVector>& LinearVelocities = PerSolverField.GetOutputResults(EFieldCommandOutputType::LinearVelocity);
 					Forces += LinearVelocities[Index] * Particles.M(Index) / Dt;
 				}
 
-				if (bHasForceField)
+				const TArray<FVector>& LinearForces = PerSolverField.GetOutputResults(EFieldCommandOutputType::LinearForce);
+				if (!LinearForces.IsEmpty())
 				{
-					const TArray<FVector>& LinearForces = PerSolverField.GetOutputResults(EFieldCommandOutputType::LinearForce);
 					Forces += LinearForces[Index];
 				}
 
@@ -863,17 +860,28 @@ void FClothingSimulationSolver::UpdateSolverField()
 		TArray<FFieldContextIndex>& SampleIndices = PerSolverField.GetSampleIndices();
 
 		const uint32 NumParticles = Evolution->Particles().Size();
+		const uint32 NumActiveParticles = Evolution->ParticlesActiveView().GetActiveSize();
 
 		SamplePositions.SetNum(NumParticles, false);
-		SampleIndices.SetNum(NumParticles, false);
+		SampleIndices.SetNum(NumActiveParticles, false);
+
+		int32 SampleIndex = 0;
+		Evolution->ParticlesActiveView().SequentialFor(
+			[this, &SamplePositions, &SampleIndices, &SampleIndex](Softs::FSolverParticles& Particles, int32 ParticleIndex)
+			{
+				SamplePositions[ParticleIndex] = FVector(Particles.X(ParticleIndex)) + LocalSpaceLocation;
+				SampleIndices[SampleIndex] = FFieldContextIndex(ParticleIndex, SampleIndex);
+				++SampleIndex;
+			});
 
-		for (uint32 Index = 0; Index < NumParticles; ++Index)
-		{
-			SamplePositions[Index] = FVector(Evolution->Particles().X(Index)) + LocalSpaceLocation;
-			SampleIndices[Index] = FFieldContextIndex(Index, Index);
-		}
 		PerSolverField.ComputeFieldLinearImpulse(GetTime());
 	}
+	else
+	{
+		// Reset the outputs once the field isn't being processed anymore
+		PerSolverField.GetOutputResults(EFieldCommandOutputType::LinearVelocity).Reset();
+		PerSolverField.GetOutputResults(EFieldCommandOutputType::LinearForce).Reset();
+	}
 }
 
 void FClothingSimulationSolver::Update(Softs::FSolverReal InDeltaTime)
@@ -953,7 +961,7 @@ void FClothingSimulationSolver::Update(Softs::FSolverReal InDeltaTime)
 			DeltaTime = InDeltaTime;
 		}
 
-		// Compute the solver field forces/velocities for future use in the AddExternalForces
+		// Compute the solver field forces/velocities for future use in the solver force function
 		{
 			TRACE_CPUPROFILER_EVENT_SCOPE(FClothingSimulationSolver_UpdateSolverFields);
 			SCOPE_CYCLE_COUNTER(STAT_ChaosClothSolverUpdateSolverFields);
diff --git a/Engine/Source/Runtime/Experimental/Chaos/Public/Chaos/PBDActiveView.h b/Engine/Source/Runtime/Experimental/Chaos/Public/Chaos/PBDActiveView.h
--- a/Engine/Source/Runtime/Experimental/Chaos/Public/Chaos/PBDActiveView.h
+++ b/Engine/Source/Runtime/Experimental/Chaos/Public/Chaos/PBDActiveView.h
@@ -50,11 +50,14 @@ namespace Chaos
 		// Return whether there is any active range in the view.
 		bool HasActiveRange() const;
 
+		// Return the total number of active items from all active ranges.
+		int32 GetActiveSize() const;
+
 		// Return a list of pair (offset, range) of all active ranges.
 		TArray<TVector<int32, 2>, TInlineAllocator<8>> GetActiveRanges() const;
 
 		// Return internal ranges.
-		TConstArrayView<int32> GetRanges() const { return Ranges;  }
+		TConstArrayView<int32> GetRanges() const { return Ranges; }
 
 	private:
 		TItemsType& Items;
@@ -221,6 +224,26 @@ namespace Chaos
 		return false;
 	}
 
+	template <class TItemsType>
+	int32 TPBDActiveView<TItemsType>::GetActiveSize() const
+	{
+		int32 ActiveSize = 0;
+		int32 Offset = 0;
+		for (int32 Range : Ranges)
+		{
+			if (Range > 0)
+			{
+				ActiveSize += Range - Offset;
+				Offset = Range;
+			}
+			else
+			{
+				Offset = -Range;
+			}
+		}
+		return ActiveSize;
+	}
+
 	template <class TItemsType>
 	TArray<TVector<int32, 2>, TInlineAllocator<8>> TPBDActiveView<TItemsType>::GetActiveRanges() const
 	{
