Chaos Cloth - Fix the velocity calculations when MaxPhysicsDelta kicks in.
- Add velocity debug draw.
- Draw aerodynamic forces in different colors depending on whether the particle is kinematic/dynamic.

#[rb Michael.Forot]
#[preflight 62bd89183f0d6beee282f026]
#[okforgithub public]
#[jira none]
[Edigrate CL 20895640]

diff --git a/Engine/Plugins/Experimental/ChaosCloth/Source/Chaos/Private/ChaosCloth/ChaosClothingSimulation.cpp b/Engine/Plugins/Experimental/ChaosCloth/Source/Chaos/Private/ChaosCloth/ChaosClothingSimulation.cpp
--- a/Engine/Plugins/Experimental/ChaosCloth/Source/Chaos/Private/ChaosCloth/ChaosClothingSimulation.cpp
+++ b/Engine/Plugins/Experimental/ChaosCloth/Source/Chaos/Private/ChaosCloth/ChaosClothingSimulation.cpp
@@ -99,6 +99,7 @@ namespace ClothingSimulationCVar
 	TAutoConsoleVariable<bool> DebugDrawAnimMeshWired       (TEXT("p.ChaosCloth.DebugDrawAnimMeshWired"       ), false, TEXT("Whether to debug draw the animated/kinematic Cloth wireframe meshes"), ECVF_Cheat);
 	TAutoConsoleVariable<bool> DebugDrawAnimNormals         (TEXT("p.ChaosCloth.DebugDrawAmimNormals"         ), false, TEXT("Whether to debug draw the animated/kinematic Cloth normals"), ECVF_Cheat);
 	TAutoConsoleVariable<bool> DebugDrawPointNormals        (TEXT("p.ChaosCloth.DebugDrawPointNormals"        ), false, TEXT("Whether to debug draw the Chaos Cloth point normals"), ECVF_Cheat);
+	TAutoConsoleVariable<bool> DebugDrawPointVelocities     (TEXT("p.ChaosCloth.DebugDrawPointVelocities"     ), false, TEXT("Whether to debug draw the Chaos Cloth point velocities"), ECVF_Cheat);
 	TAutoConsoleVariable<bool> DebugDrawFaceNormals         (TEXT("p.ChaosCloth.DebugDrawFaceNormals"         ), false, TEXT("Whether to debug draw the Chaos Cloth face normals"), ECVF_Cheat);
 	TAutoConsoleVariable<bool> DebugDrawInversedFaceNormals (TEXT("p.ChaosCloth.DebugDrawInversedFaceNormals" ), false, TEXT("Whether to debug draw the Chaos Cloth inversed face normals"), ECVF_Cheat);
 	TAutoConsoleVariable<bool> DebugDrawCollision           (TEXT("p.ChaosCloth.DebugDrawCollision"           ), false, TEXT("Whether to debug draw the Chaos Cloth collisions"), ECVF_Cheat);
@@ -561,6 +562,7 @@ void FClothingSimulation::Simulate(IClothingSimulationContext* InContext)
 		Solver->SetWindVelocity(Context->WindVelocity, Context->WindAdaption);
 		Solver->SetGravity(bUseGravityOverride ? GravityOverride : Context->WorldGravity);
 		Solver->EnableClothGravityOverride(!bUseGravityOverride);  // Disable all cloth gravity overrides when the interactor takes over
+		Solver->SetVelocityScale((FReal)Context->VelocityScale);
 
 		// Check teleport modes
 		for (const TUniquePtr<FClothingSimulationCloth>& Cloth : Cloths)
@@ -610,6 +612,7 @@ void FClothingSimulation::Simulate(IClothingSimulationContext* InContext)
 	if (ClothingSimulationCVar::DebugDrawGravity             .GetValueOnAnyThread()) { DebugDrawGravity             (); }
 	if (ClothingSimulationCVar::DebugDrawPhysMeshWired       .GetValueOnAnyThread()) { DebugDrawPhysMeshWired       (); }
 	if (ClothingSimulationCVar::DebugDrawAnimMeshWired       .GetValueOnAnyThread()) { DebugDrawAnimMeshWired       (); }
+	if (ClothingSimulationCVar::DebugDrawPointVelocities     .GetValueOnAnyThread()) { DebugDrawPointVelocities     (); }
 	if (ClothingSimulationCVar::DebugDrawAnimNormals         .GetValueOnAnyThread()) { DebugDrawAnimNormals         (); }
 	if (ClothingSimulationCVar::DebugDrawPointNormals        .GetValueOnAnyThread()) { DebugDrawPointNormals        (); }
 	if (ClothingSimulationCVar::DebugDrawCollision           .GetValueOnAnyThread()) { DebugDrawCollision           (); }
@@ -1478,6 +1481,39 @@ void FClothingSimulation::DebugDrawPointNormals(FPrimitiveDrawInterface* PDI) co
 	}
 }
 
+void FClothingSimulation::DebugDrawPointVelocities(FPrimitiveDrawInterface* PDI) const
+{
+	check(Solver);
+	const FVec3& LocalSpaceLocation = Solver->GetLocalSpaceLocation();
+
+	for (const FClothingSimulationCloth* const Cloth : Solver->GetCloths())
+	{
+		const int32 Offset = Cloth->GetOffset(Solver.Get());
+		if (Offset == INDEX_NONE)
+		{
+			continue;
+		}
+
+		const TConstArrayView<Softs::FSolverVec3> Positions = Cloth->GetParticlePositions(Solver.Get());
+		const TConstArrayView<Softs::FSolverVec3> Velocities = Cloth->GetParticleVelocities(Solver.Get());
+		check(Velocities.Num() == Positions.Num());
+
+		const TConstArrayView<Softs::FSolverReal> InvMasses = Cloth->GetParticleInvMasses(Solver.Get());
+		check(InvMasses.Num() == Positions.Num());
+
+		for (int32 Index = 0; Index < Positions.Num(); ++Index)
+		{
+			constexpr FReal DefaultFPS = (FReal)60.;   // TODO: A CVAR would be nice for this
+			const bool bIsKinematic = (InvMasses[Index] == 0.f);
+
+			const FVec3 Pos0 = LocalSpaceLocation + Positions[Index];
+			const FVec3 Pos1 = Pos0 + FVec3(Velocities[Index]) / DefaultFPS;  // Velocity per frame if running at DefaultFPS
+
+			DrawLine(PDI, Pos0, Pos1, bIsKinematic ? FLinearColor::Black : FLinearColor::Yellow);
+		}
+	}
+}
+
 void FClothingSimulation::DebugDrawInversedPointNormals(FPrimitiveDrawInterface* PDI) const
 {
 	check(Solver);
@@ -2069,6 +2105,8 @@ void FClothingSimulation::DebugDrawWindAndPressureForces(FPrimitiveDrawInterface
 		const TConstArrayView<TVec3<int32>>& Elements = VelocityField.GetElements();
 		const TConstArrayView<Softs::FSolverVec3> Forces = VelocityField.GetForces();
 		const TConstArrayView<Softs::FSolverVec3> Positions = Cloth->GetParticlePositions(Solver.Get());
+		const TConstArrayView<Softs::FSolverReal> InvMasses = Cloth->GetParticleInvMasses(Solver.Get());
+		check(InvMasses.Num() == Positions.Num());
 
 		for (int32 ElementIndex = 0; ElementIndex < Elements.Num(); ++ElementIndex)
 		{
@@ -2077,8 +2115,14 @@ void FClothingSimulation::DebugDrawWindAndPressureForces(FPrimitiveDrawInterface
 				FVec3(Positions[Element.X - Offset]) +
 				FVec3(Positions[Element.Y - Offset]) +
 				FVec3(Positions[Element.Z - Offset])) / (FReal)3.;
+
+			const bool bIsKinematic0 = !InvMasses[Element.X - Offset];
+			const bool bIsKinematic1 = !InvMasses[Element.Y - Offset];
+			const bool bIsKinematic2 = !InvMasses[Element.Z - Offset];
+			const bool bIsKinematic = bIsKinematic0 || bIsKinematic1 || bIsKinematic2;
+
 			const FVec3 Force = FVec3(Forces[ElementIndex]) * ForceLength;
-			DrawLine(PDI, Position, Position + Force, FColor::Green);
+			DrawLine(PDI, Position, Position + Force, bIsKinematic ? FColor::Cyan : FColor::Green);
 		}
 	}
 }
diff --git a/Engine/Plugins/Experimental/ChaosCloth/Source/Chaos/Private/ChaosCloth/ChaosClothingSimulationCloth.cpp b/Engine/Plugins/Experimental/ChaosCloth/Source/Chaos/Private/ChaosCloth/ChaosClothingSimulationCloth.cpp
--- a/Engine/Plugins/Experimental/ChaosCloth/Source/Chaos/Private/ChaosCloth/ChaosClothingSimulationCloth.cpp
+++ b/Engine/Plugins/Experimental/ChaosCloth/Source/Chaos/Private/ChaosCloth/ChaosClothingSimulationCloth.cpp
@@ -881,6 +881,14 @@ TConstArrayView<Softs::FSolverVec3> FClothingSimulationCloth::GetParticleNormals
 	return TConstArrayView<Softs::FSolverVec3>(Solver->GetNormals(GetOffset(Solver, LODIndex)), GetNumParticles(LODIndex));
 }
 
+TConstArrayView<Softs::FSolverVec3> FClothingSimulationCloth::GetParticleVelocities(const FClothingSimulationSolver* Solver) const
+{
+	check(Solver);
+	const int32 LODIndex = LODIndices.FindChecked(Solver);
+	check(GetOffset(Solver, LODIndex) != INDEX_NONE);
+	return TConstArrayView<Softs::FSolverVec3>(Solver->GetParticleVs(GetOffset(Solver, LODIndex)), GetNumParticles(LODIndex));
+}
+
 TConstArrayView<Softs::FSolverReal> FClothingSimulationCloth::GetParticleInvMasses(const FClothingSimulationSolver* Solver) const
 {
 	check(Solver);
diff --git a/Engine/Plugins/Experimental/ChaosCloth/Source/Chaos/Private/ChaosCloth/ChaosClothingSimulationSolver.cpp b/Engine/Plugins/Experimental/ChaosCloth/Source/Chaos/Private/ChaosCloth/ChaosClothingSimulationSolver.cpp
--- a/Engine/Plugins/Experimental/ChaosCloth/Source/Chaos/Private/ChaosCloth/ChaosClothingSimulationSolver.cpp
+++ b/Engine/Plugins/Experimental/ChaosCloth/Source/Chaos/Private/ChaosCloth/ChaosClothingSimulationSolver.cpp
@@ -46,6 +46,7 @@ static bool bClothSolverParallelClothUpdate = true;
 static bool bClothSolverParallelClothPostUpdate = true;
 static bool bClothSolverUseImprovedTimeStepSmoothing = true;
 static bool bClothSolverDisableTimeDependentNumIterations = false;
+static bool bClothSolverUseVelocityScale = true;
 
 #if !UE_BUILD_SHIPPING
 static int32 ClothSolverDebugHitchLength = 0;
@@ -63,6 +64,7 @@ FAutoConsoleVariableRef CVarClothSolverDisableCollision(TEXT("p.ChaosCloth.Solve
 
 FAutoConsoleVariableRef CVarClothSolverUseImprovedTimeStepSmoothing(TEXT("p.ChaosCloth.Solver.UseImprovedTimeStepSmoothing"), bClothSolverUseImprovedTimeStepSmoothing, TEXT("Use the time step smoothing on input forces only rather than on the entire cloth solver, in order to avoid miscalculating velocities."));
 FAutoConsoleVariableRef CVarClothSolverDisableTimeDependentNumIterations(TEXT("p.ChaosCloth.Solver.DisableTimeDependentNumIterations"), bClothSolverDisableTimeDependentNumIterations, TEXT("Make the number of iterations independent from the time step."));
+FAutoConsoleVariableRef CVarClothSolverUseVelocityScale(TEXT("p.ChaosCloth.Solver.UseVelocityScale"), bClothSolverUseVelocityScale, TEXT("Use the velocity scale to compensate for clamping to MaxPhysicsDelta, in order to avoid miscalculating velocities during hitches."));
 
 namespace ClothingSimulationSolverDefault
 {
@@ -87,6 +89,7 @@ namespace ClothingSimulationSolverConstant
 FClothingSimulationSolver::FClothingSimulationSolver()
 	: OldLocalSpaceLocation(0.)
 	, LocalSpaceLocation(0.)
+	, VelocityScale(1.)
 	, Time(0.)
 	, DeltaTime(ClothingSimulationSolverConstant::StartDeltaTime)
 	, NumIterations(ClothingSimulationSolverDefault::NumIterations)
@@ -576,11 +579,13 @@ void FClothingSimulationSolver::SetReferenceVelocityScale(
 	FRigidTransform3 OldRootBoneLocalTransform = OldReferenceSpaceTransform;
 	OldRootBoneLocalTransform.AddToTranslation(-OldLocalSpaceLocation);
 
+	const FReal SolverVelocityScale = bClothSolverUseVelocityScale ? VelocityScale : (FReal)1.;
+
 	// Calculate deltas
 	const FRigidTransform3 DeltaTransform = ReferenceSpaceTransform.GetRelativeTransform(OldReferenceSpaceTransform);
 
 	// Apply linear velocity scale
-	const FVec3 LinearRatio = FVec3(1.) - FVec3(LinearVelocityScale).BoundToBox(FVec3(0.), FVec3(1.));
+	const FVec3 LinearRatio = FVec3(1.) - FVec3(LinearVelocityScale).BoundToBox(FVec3(0.), FVec3(1.)) * SolverVelocityScale;
 	const FVec3 DeltaPosition = LinearRatio * DeltaTransform.GetTranslation();
 
 	// Apply angular velocity scale
@@ -592,7 +597,7 @@ void FClothingSimulationSolver::SetReferenceVelocityScale(
 		DeltaAngle -= (FReal)2. * (FReal)PI;
 	}
 
-	const FReal PartialDeltaAngle = DeltaAngle * FMath::Clamp((FReal)1. - (FReal)AngularVelocityScale, (FReal)0., (FReal)1.);
+	const FReal PartialDeltaAngle = DeltaAngle * FMath::Clamp((FReal)1. - (FReal)AngularVelocityScale * SolverVelocityScale, (FReal)0., (FReal)1.);
 	DeltaRotation = UE::Math::TQuat<FReal>(Axis, PartialDeltaAngle);
 
 	// Transform points back into the previous frame of reference before applying the adjusted deltas
diff --git a/Engine/Plugins/Experimental/ChaosCloth/Source/Chaos/Public/ChaosCloth/ChaosClothingSimulation.h b/Engine/Plugins/Experimental/ChaosCloth/Source/Chaos/Public/ChaosCloth/ChaosClothingSimulation.h
--- a/Engine/Plugins/Experimental/ChaosCloth/Source/Chaos/Public/ChaosCloth/ChaosClothingSimulation.h
+++ b/Engine/Plugins/Experimental/ChaosCloth/Source/Chaos/Public/ChaosCloth/ChaosClothingSimulation.h
@@ -109,6 +109,7 @@ namespace Chaos
 		CHAOSCLOTH_API void DebugDrawAnimMeshWired(FPrimitiveDrawInterface* PDI = nullptr) const;
 		CHAOSCLOTH_API void DebugDrawAnimNormals(FPrimitiveDrawInterface* PDI = nullptr) const;
 		CHAOSCLOTH_API void DebugDrawPointNormals(FPrimitiveDrawInterface* PDI = nullptr) const;
+		CHAOSCLOTH_API void DebugDrawPointVelocities(FPrimitiveDrawInterface* PDI = nullptr) const;
 		UE_DEPRECATED(5.0, "DebugDrawInversedPointNormals is mostly redundant and will be removed, use DebugDrawPointNormals instead")
 		CHAOSCLOTH_API void DebugDrawInversedPointNormals(FPrimitiveDrawInterface* PDI = nullptr) const;
 		CHAOSCLOTH_API void DebugDrawCollision(FPrimitiveDrawInterface* PDI = nullptr) const;
diff --git a/Engine/Plugins/Experimental/ChaosCloth/Source/Chaos/Public/ChaosCloth/ChaosClothingSimulationCloth.h b/Engine/Plugins/Experimental/ChaosCloth/Source/Chaos/Public/ChaosCloth/ChaosClothingSimulationCloth.h
--- a/Engine/Plugins/Experimental/ChaosCloth/Source/Chaos/Public/ChaosCloth/ChaosClothingSimulationCloth.h
+++ b/Engine/Plugins/Experimental/ChaosCloth/Source/Chaos/Public/ChaosCloth/ChaosClothingSimulationCloth.h
@@ -118,9 +118,11 @@ namespace Chaos
 		TConstArrayView<Softs::FSolverVec3> GetAnimationNormals(const FClothingSimulationSolver* Solver) const;
 		// Return the solver's positions for this cloth current LOD, not thread safe, call must be done right after the solver update.
 		TConstArrayView<Softs::FSolverVec3> GetParticlePositions(const FClothingSimulationSolver* Solver) const;
+		// Return the solver's velocities for this cloth current LOD, not thread safe, call must be done right after the solver update.
+		TConstArrayView<Softs::FSolverVec3> GetParticleVelocities(const FClothingSimulationSolver* Solver) const;
 		// Return the solver's normals for this cloth current LOD, not thread safe, call must be done right after the solver update.
 		TConstArrayView<Softs::FSolverVec3> GetParticleNormals(const FClothingSimulationSolver* Solver) const;
-		// Return the solver's normals for this cloth current LOD, not thread safe, call must be done right after the solver update.
+		// Return the solver's inverse masses for this cloth current LOD, not thread safe, call must be done right after the solver update.
 		TConstArrayView<Softs::FSolverReal> GetParticleInvMasses(const FClothingSimulationSolver* Solver) const;
 		// Return the current gravity as applied by the solver using the various overrides, not thread safe, call must be done right after the solver update.
 		TVec3<FRealSingle> GetGravity(const FClothingSimulationSolver* Solver) const;
diff --git a/Engine/Plugins/Experimental/ChaosCloth/Source/Chaos/Public/ChaosCloth/ChaosClothingSimulationSolver.h b/Engine/Plugins/Experimental/ChaosCloth/Source/Chaos/Public/ChaosCloth/ChaosClothingSimulationSolver.h
--- a/Engine/Plugins/Experimental/ChaosCloth/Source/Chaos/Public/ChaosCloth/ChaosClothingSimulationSolver.h
+++ b/Engine/Plugins/Experimental/ChaosCloth/Source/Chaos/Public/ChaosCloth/ChaosClothingSimulationSolver.h
@@ -33,6 +33,8 @@ namespace Chaos
 		const FVec3& GetLocalSpaceLocation() const { return LocalSpaceLocation; }
 		void SetLocalSpaceRotation(const FQuat& InLocalSpaceRotation) { LocalSpaceRotation = InLocalSpaceRotation; }
 		const FRotation3& GetLocalSpaceRotation() const { return LocalSpaceRotation; }
+		void SetVelocityScale(FReal InVelocityScale) { VelocityScale = InVelocityScale; }
+		FReal GetVelocityScale() const { return VelocityScale; }
 
 		// Disables all Cloths gravity override mechanism
 		void EnableClothGravityOverride(bool bInIsClothGravityOverrideEnabled) { bIsClothGravityOverrideEnabled = bInIsClothGravityOverrideEnabled; }
@@ -43,9 +45,6 @@ namespace Chaos
 		void SetWindVelocity(const TVec3<FRealSingle>& InWindVelocity, FRealSingle InLegacyWindAdaption = (FRealSingle)0.);
 		const TVec3<FRealSingle>& GetWindVelocity() const { return WindVelocity; }
 
-		UE_DEPRECATED(4.27, "Use SetWindProperties instead.")
-		void SetWindFluidDensity(FRealSingle /*WindFluidDensity*/) {}
-
 		void SetNumIterations(int32 InNumIterations) { NumIterations = InNumIterations; }
 		int32 GetNumIterations() const { return NumIterations; }
 		void SetMaxNumIterations(int32 InMaxNumIterations) { MaxNumIterations = InMaxNumIterations; }
@@ -233,6 +232,7 @@ namespace Chaos
 		FVec3 OldLocalSpaceLocation;  // This is used to translate between world space and simulation space,
 		FVec3 LocalSpaceLocation;     // add this to simulation space coordinates to get world space coordinates, must keep FReal as underlying type for LWC
 		FRotation3 LocalSpaceRotation;
+		FReal VelocityScale;
 
 		// Time stepping
 		FSolverReal Time;
diff --git a/Engine/Plugins/Experimental/ChaosClothEditor/Source/Chaos/Private/ChaosClothEditor/ChaosSimulationEditorExtender.cpp b/Engine/Plugins/Experimental/ChaosClothEditor/Source/Chaos/Private/ChaosClothEditor/ChaosSimulationEditorExtender.cpp
--- a/Engine/Plugins/Experimental/ChaosClothEditor/Source/Chaos/Private/ChaosClothEditor/ChaosSimulationEditorExtender.cpp
+++ b/Engine/Plugins/Experimental/ChaosClothEditor/Source/Chaos/Private/ChaosClothEditor/ChaosSimulationEditorExtender.cpp
@@ -70,6 +70,7 @@ const FVisualizationOption FVisualizationOption::OptionData[] =
 	FVisualizationOption(&Chaos::FClothingSimulation::DebugDrawParticleIndices     , TEXT("p.ChaosClothEditor.DebugDrawParticleIndices"     ), LOCTEXT("ChaosVisName_ParticleIndices"     , "Particle Indices"           ), LOCTEXT("ChaosVisName_ParticleIndices_ToolTip"     , "Draws the particle indices as instantiated by the solver")),
 	FVisualizationOption(&Chaos::FClothingSimulation::DebugDrawElementIndices      , TEXT("p.ChaosClothEditor.DebugDrawElementIndices"      ), LOCTEXT("ChaosVisName_ElementIndices"      , "Element Indices"            ), LOCTEXT("ChaosVisName_ElementIndices_ToolTip"      , "Draws the element's (triangle or other) indices as instantiated by the solver")),
 	FVisualizationOption(&Chaos::FClothingSimulation::DebugDrawPointNormals        , TEXT("p.ChaosClothEditor.DebugDrawPointNormals"        ), LOCTEXT("ChaosVisName_PointNormals"        , "Physical Mesh Normals"      ), LOCTEXT("ChaosVisName_PointNormals_ToolTip"        , "Draws the current point normals for the simulation mesh")),
+	FVisualizationOption(&Chaos::FClothingSimulation::DebugDrawPointVelocities     , TEXT("p.ChaosClothEditor.DebugDrawPointVelocities"     ), LOCTEXT("ChaosVisName_PointVelocities"     , "Point Velocities"           ), LOCTEXT("ChaosVisName_PointVelocities_ToolTip"     , "Draws the current point velocities for the simulation mesh")),
 	FVisualizationOption(&Chaos::FClothingSimulation::DebugDrawAnimNormals         , TEXT("p.ChaosClothEditor.DebugDrawAnimNormals"         ), LOCTEXT("ChaosVisName_AnimNormals"         , "Animated Mesh Normals"      ), LOCTEXT("ChaosVisName_AnimNormals_ToolTip"         , "Draws the current point normals for the animated mesh")),
 	FVisualizationOption(&Chaos::FClothingSimulation::DebugDrawCollision           , TEXT("p.ChaosClothEditor.DebugDrawCollision"           ), LOCTEXT("ChaosVisName_Collision"           , "Collisions"                 ), LOCTEXT("ChaosVisName_Collision_ToolTip"           , "Draws the collision bodies the simulation is currently using")),
 	FVisualizationOption(&Chaos::FClothingSimulation::DebugDrawBackstops           , TEXT("p.ChaosClothEditor.DebugDrawBackstops"           ), LOCTEXT("ChaosVisName_Backstop"            , "Backstops"                  ), LOCTEXT("ChaosVisName_Backstop_ToolTip"            , "Draws the backstop radius and position for each simulation particle")),
diff --git a/Engine/Source/Runtime/ClothingSystemRuntimeCommon/Private/ClothingSimulation.cpp b/Engine/Source/Runtime/ClothingSystemRuntimeCommon/Private/ClothingSimulation.cpp
--- a/Engine/Source/Runtime/ClothingSystemRuntimeCommon/Private/ClothingSimulation.cpp
+++ b/Engine/Source/Runtime/ClothingSystemRuntimeCommon/Private/ClothingSimulation.cpp
@@ -25,6 +25,7 @@ FClothingSimulationContextCommon::FClothingSimulationContextCommon()
 	, WindVelocity(FVector::ZeroVector)
 	, WindAdaption(0.f)
 	, DeltaSeconds(0.f)
+	, VelocityScale(1.f)
 	, TeleportMode(EClothingTeleportMode::None)
 	, MaxDistanceScale(1.f)
 	, PredictedLod(INDEX_NONE)
@@ -165,6 +166,9 @@ void FClothingSimulationContextCommon::FillTeleportMode(const USkeletalMeshCompo
 	TeleportMode = (InDeltaSeconds > InMaxPhysicsDelta * GClothMaxDeltaTimeTeleportMultiplier.GetValueOnGameThread()) ?
 		EClothingTeleportMode::Teleport :
 		InComponent->ClothTeleportMode;
+
+	VelocityScale = (TeleportMode == EClothingTeleportMode::None) ?
+		FMath::Min(InDeltaSeconds, InMaxPhysicsDelta) / InDeltaSeconds : 1.f;
 }
 
 void FClothingSimulationContextCommon::FillMaxDistanceScale(const USkeletalMeshComponent* InComponent)
diff --git a/Engine/Source/Runtime/ClothingSystemRuntimeCommon/Public/ClothingSimulation.h b/Engine/Source/Runtime/ClothingSystemRuntimeCommon/Public/ClothingSimulation.h
--- a/Engine/Source/Runtime/ClothingSystemRuntimeCommon/Public/ClothingSimulation.h
+++ b/Engine/Source/Runtime/ClothingSystemRuntimeCommon/Public/ClothingSimulation.h
@@ -63,6 +63,9 @@ public:
 	// Delta for this tick
 	float DeltaSeconds;
 
+	// Velocity scale to compensate input velocities in case the MaxPhysicsDelta kicks in
+	float VelocityScale;
+
 	// Whether and how we should teleport the simulation this tick
 	EClothingTeleportMode TeleportMode;
 
