Chaos Cloth - Move smoothed timestep out of the solver, leveraging the velocity scale to avoid cloth tugging issue, in order to allow for more drastic change in Dt.

diff --git a/Engine/Plugins/Experimental/ChaosCloth/Source/Chaos/Private/ChaosCloth/ChaosClothingSimulation.cpp b/Engine/Plugins/Experimental/ChaosCloth/Source/Chaos/Private/ChaosCloth/ChaosClothingSimulation.cpp
--- a/Engine/Plugins/Experimental/ChaosCloth/Source/Chaos/Private/ChaosCloth/ChaosClothingSimulation.cpp
+++ b/Engine/Plugins/Experimental/ChaosCloth/Source/Chaos/Private/ChaosCloth/ChaosClothingSimulation.cpp
@@ -82,12 +82,15 @@ bool bChaos_GetSimData_ISPC_Enabled = true;
 FAutoConsoleVariableRef CVarChaosGetSimDataISPCEnabled(TEXT("p.Chaos.GetSimData.ISPC"), bChaos_GetSimData_ISPC_Enabled, TEXT("Whether to use ISPC optimizations when getting simulation data"));
 #endif
 
+
 DECLARE_CYCLE_STAT(TEXT("Chaos Cloth Simulate"), STAT_ChaosClothSimulate, STATGROUP_ChaosCloth);
 DECLARE_CYCLE_STAT(TEXT("Chaos Cloth Create Actor"), STAT_ChaosClothCreateActor, STATGROUP_ChaosCloth);
 DECLARE_CYCLE_STAT(TEXT("Chaos Cloth Get Simulation Data"), STAT_ChaosClothGetSimulationData, STATGROUP_ChaosCloth);
 
 namespace Chaos
 {
+bool bClothUseTimeStepSmoothing = true;
+FAutoConsoleVariableRef CVarClothUseTimeStepSmoothing(TEXT("p.ChaosCloth.UseTimeStepSmoothing"), bClothUseTimeStepSmoothing, TEXT("Use time step smoothing to avoid jitter during drastic changes in time steps."));
 
 #if CHAOS_DEBUG_DRAW
 namespace ClothingSimulationCVar
@@ -544,6 +547,14 @@ void FClothingSimulation::Simulate(IClothingSimulationContext* InContext)
 			return;
 		}
 
+		// Filter delta time to smoothen time variations and prevent unwanted vibrations
+		constexpr Softs::FSolverReal DeltaTimeDecay = (Softs::FSolverReal)0.1;
+		const Softs::FSolverReal DeltaTime = (Softs::FSolverReal)Context->DeltaSeconds;
+		const Softs::FSolverReal PrevDeltaTime = Solver->GetDeltaTime();
+		const Softs::FSolverReal SmoothedDeltaTime = bClothUseTimeStepSmoothing ?
+			PrevDeltaTime + (DeltaTime - PrevDeltaTime) * DeltaTimeDecay :
+			DeltaTime;
+
 		const double StartTime = FPlatformTime::Seconds();
 		const float PrevSimulationTime = SimulationTime;  // Copy the atomic to prevent a re-read
 
@@ -562,7 +573,7 @@ void FClothingSimulation::Simulate(IClothingSimulationContext* InContext)
 		Solver->SetWindVelocity(Context->WindVelocity, Context->WindAdaption);
 		Solver->SetGravity(bUseGravityOverride ? GravityOverride : Context->WorldGravity);
 		Solver->EnableClothGravityOverride(!bUseGravityOverride);  // Disable all cloth gravity overrides when the interactor takes over
-		Solver->SetVelocityScale((FReal)Context->VelocityScale);
+		Solver->SetVelocityScale((FReal)Context->VelocityScale * (FReal)SmoothedDeltaTime / DeltaTime);
 
 		// Check teleport modes
 		for (const TUniquePtr<FClothingSimulationCloth>& Cloth : Cloths)
@@ -583,7 +594,7 @@ void FClothingSimulation::Simulate(IClothingSimulationContext* InContext)
 		// Step the simulation
 		if(Solver->GetEnableSolver() || (Context->CachedPositions.Num() == 0 && Context->CachedVelocities.Num() == 0))
 		{
-			Solver->Update((Softs::FSolverReal)Context->DeltaSeconds);
+			Solver->Update(SmoothedDeltaTime);
 		}
 		else
 		{
diff --git a/Engine/Plugins/Experimental/ChaosCloth/Source/Chaos/Private/ChaosCloth/ChaosClothingSimulationSolver.cpp b/Engine/Plugins/Experimental/ChaosCloth/Source/Chaos/Private/ChaosCloth/ChaosClothingSimulationSolver.cpp
--- a/Engine/Plugins/Experimental/ChaosCloth/Source/Chaos/Private/ChaosCloth/ChaosClothingSimulationSolver.cpp
+++ b/Engine/Plugins/Experimental/ChaosCloth/Source/Chaos/Private/ChaosCloth/ChaosClothingSimulationSolver.cpp
@@ -44,7 +44,6 @@ static int32 ClothSolverMinParallelBatchSize = 1000;
 static bool bClothSolverParallelClothPreUpdate = true;
 static bool bClothSolverParallelClothUpdate = true;
 static bool bClothSolverParallelClothPostUpdate = true;
-static bool bClothSolverUseImprovedTimeStepSmoothing = true;
 static bool bClothSolverDisableTimeDependentNumIterations = false;
 static bool bClothSolverUseVelocityScale = true;
 
@@ -62,7 +61,6 @@ FAutoConsoleVariableRef CVarClothSolverDebugHitchInterval(TEXT("p.ChaosCloth.Sol
 FAutoConsoleVariableRef CVarClothSolverDisableCollision(TEXT("p.ChaosCloth.Solver.DisableCollision"), bClothSolverDisableCollision, TEXT("Disable all collision particles. Needs reset of the simulation (p.ChaosCloth.Reset)."));
 #endif
 
-FAutoConsoleVariableRef CVarClothSolverUseImprovedTimeStepSmoothing(TEXT("p.ChaosCloth.Solver.UseImprovedTimeStepSmoothing"), bClothSolverUseImprovedTimeStepSmoothing, TEXT("Use the time step smoothing on input forces only rather than on the entire cloth solver, in order to avoid miscalculating velocities."));
 FAutoConsoleVariableRef CVarClothSolverDisableTimeDependentNumIterations(TEXT("p.ChaosCloth.Solver.DisableTimeDependentNumIterations"), bClothSolverDisableTimeDependentNumIterations, TEXT("Make the number of iterations independent from the time step."));
 FAutoConsoleVariableRef CVarClothSolverUseVelocityScale(TEXT("p.ChaosCloth.Solver.UseVelocityScale"), bClothSolverUseVelocityScale, TEXT("Use the velocity scale to compensate for clamping to MaxPhysicsDelta, in order to avoid miscalculating velocities during hitches."));
 
@@ -905,6 +903,9 @@ void FClothingSimulationSolver::Update(Softs::FSolverReal InDeltaTime)
 	}
 #endif  // #if !UE_BUILD_SHIPPING
 
+	// Update time step
+	DeltaTime = InDeltaTime;
+
 	// Update Cloths and cloth colliders
 	{
 		TRACE_CPUPROFILER_EVENT_SCOPE(FClothingSimulationSolver_UpdateCloths);
@@ -946,23 +947,9 @@ void FClothingSimulationSolver::Update(Softs::FSolverReal InDeltaTime)
 		EventPreSolve.Broadcast(DeltaTime);
 	}
 
-	const bool bAdvanceTimeStep = (InDeltaTime > Softs::FSolverReal(0.));
+	const bool bAdvanceTimeStep = (DeltaTime > Softs::FSolverReal(0.));
 	if (bAdvanceTimeStep)
 	{
-		if (!bClothSolverUseImprovedTimeStepSmoothing)
-		{
-			// Filter delta time to smoothen time variations and prevent unwanted vibrations
-			// Note: This is now deprecated and replaced by in solver input force timestep smoothing
-			constexpr Softs::FSolverReal DeltaTimeDecay = (Softs::FSolverReal)0.1;
-			const Softs::FSolverReal PrevDeltaTime = DeltaTime;
-			DeltaTime = DeltaTime + (InDeltaTime - DeltaTime) * DeltaTimeDecay;
-		}
-		else
-		{
-			// Update time step
-			DeltaTime = InDeltaTime;
-		}
-
 		// Compute the solver field forces/velocities for future use in the solver force function
 		{
 			TRACE_CPUPROFILER_EVENT_SCOPE(FClothingSimulationSolver_UpdateSolverFields);
@@ -992,7 +979,7 @@ void FClothingSimulationSolver::Update(Softs::FSolverReal InDeltaTime)
 	
 			for (int32 i = 0; i < NumSubsteps; ++i)
 			{
-				Evolution->AdvanceOneTimeStep(SubstepDeltaTime, bClothSolverUseImprovedTimeStepSmoothing);
+				Evolution->AdvanceOneTimeStep(SubstepDeltaTime);
 			}
 
 			Time = Evolution->GetTime();
diff --git a/Engine/Plugins/Experimental/ChaosCloth/Source/Chaos/Public/ChaosCloth/ChaosClothingSimulationSolver.h b/Engine/Plugins/Experimental/ChaosCloth/Source/Chaos/Public/ChaosCloth/ChaosClothingSimulationSolver.h
--- a/Engine/Plugins/Experimental/ChaosCloth/Source/Chaos/Public/ChaosCloth/ChaosClothingSimulationSolver.h
+++ b/Engine/Plugins/Experimental/ChaosCloth/Source/Chaos/Public/ChaosCloth/ChaosClothingSimulationSolver.h
@@ -66,10 +66,12 @@ namespace Chaos
 
 		TConstArrayView<const FClothingSimulationCloth*> GetCloths() const { return Cloths; }
 
-		// Update solver properties before simulation.
-		
+		/** Advance the simulation. */
 		void Update(FSolverReal InDeltaTime);
 
+		/** Return the last delta time used for advancing the simulation. */
+		FSolverReal GetDeltaTime() const { return DeltaTime; }
+
 		/** Set the cached positions onto the particles */
 		void UpdateFromCache(const TArray<FVector>& CachedPositions, const TArray<FVector>& CachedVelocities);
 
diff --git a/Engine/Source/Runtime/Experimental/Chaos/Private/Chaos/PBDEvolution.cpp b/Engine/Source/Runtime/Experimental/Chaos/Private/Chaos/PBDEvolution.cpp
--- a/Engine/Source/Runtime/Experimental/Chaos/Private/Chaos/PBDEvolution.cpp
+++ b/Engine/Source/Runtime/Experimental/Chaos/Private/Chaos/PBDEvolution.cpp
@@ -96,7 +96,6 @@ FPBDEvolution::FPBDEvolution(
 	, MDamping(Damping)
 	, MLocalDamping(LocalDamping)
 	, MTime(0)
-	, MSmoothDt(1.f / 30.f)  // Initialize filtered timestep at 30fps
 {
 	// Add group arrays
 	TArrayCollection::AddArray(&MGroupGravityAccelerations);
@@ -294,7 +293,7 @@ void FPBDEvolution::PreIterationUpdate(
 					}
 
 					// Euler Step Velocity
-					MParticles.V(Index) += MParticles.Acceleration(Index) * MSmoothDt;
+					MParticles.V(Index) += MParticles.Acceleration(Index) * Dt;
 
 					// Damp Velocity Rule
 					if (bDampVelocityRule)
@@ -315,7 +314,7 @@ void FPBDEvolution::PreIterationUpdate(
 	}
 }
 
-void FPBDEvolution::AdvanceOneTimeStep(const FSolverReal Dt, const bool bSmoothDt)
+void FPBDEvolution::AdvanceOneTimeStep(const FSolverReal Dt)
 {
 	TRACE_CPUPROFILER_EVENT_SCOPE(FPBDEvolution_AdvanceOneTimeStep);
 	SCOPE_CYCLE_COUNTER(STAT_ChaosPBDVAdvanceTime);
@@ -323,17 +322,6 @@ void FPBDEvolution::AdvanceOneTimeStep(const FSolverReal Dt, const bool bSmoothD
 	// Advance time
 	MTime += Dt;
 
-	// Filter delta time to smoothen time variations and prevent unwanted vibrations, works best on Forces
-	if (bSmoothDt && CVarChaosPBDEvolutionUseSmoothTimeStep.GetValueOnAnyThread())
-	{
-		constexpr FSolverReal DeltaTimeDecay = (FSolverReal)0.1;
-		MSmoothDt += (Dt - MSmoothDt) * DeltaTimeDecay;
-	}
-	else
-	{
-		MSmoothDt = Dt;
-	}
-
 	// Don't bother with threaded execution if we don't have enough work to make it worth while.
 	const bool bUseSingleThreadedRange = !CVarChaosPBDEvolutionUseNestedParallelFor.GetValueOnAnyThread();
 	const int32 MinParallelBatchSize = !CVarChaosPBDEvolutionParallelIntegrate.GetValueOnAnyThread() ?
diff --git a/Engine/Source/Runtime/Experimental/Chaos/Public/Chaos/PBDEvolution.h b/Engine/Source/Runtime/Experimental/Chaos/Public/Chaos/PBDEvolution.h
--- a/Engine/Source/Runtime/Experimental/Chaos/Public/Chaos/PBDEvolution.h
+++ b/Engine/Source/Runtime/Experimental/Chaos/Public/Chaos/PBDEvolution.h
@@ -29,7 +29,11 @@ class CHAOS_API FPBDEvolution : public TArrayCollection
 	~FPBDEvolution() {}
 
 	// Advance one time step. Filter the input time step if specified.
-	void AdvanceOneTimeStep(const FSolverReal Dt, const bool bSmoothDt = true);
+	UE_DEPRECATED(5.1, "Use AdvanceOneTimeStep(const FSolverReal Dt) instead.")
+	void AdvanceOneTimeStep(const FSolverReal Dt, const bool bSmoothDt) { AdvanceOneTimeStep(Dt); }
+
+	// Advance one time step. Filter the input time step if specified.
+	void AdvanceOneTimeStep(const FSolverReal Dt);
 
 	// Remove all particles, will also reset all rules
 	void ResetParticles();
@@ -216,7 +220,6 @@ private:
 	FSolverReal MDamping;
 	FSolverReal MLocalDamping;
 	FSolverReal MTime;
-	FSolverReal MSmoothDt;
 };
 
 }  // End namespace Chaos::Softs
