Chaos Cloth:
- Turned the ClothingSimulationMesh class into a base class to allow non skeletal mesh based simulations.
- Deprecated direct usage of ClothingSimulationMesh, re-implemented the skeletal mesh dependent code in a new ClothingSimulationSkeletalMesh class, and prepared the code so that virtual pure methods can eventually replace the skeletal mesh dependent implementation and that all remaining skeletal mesh dependencies can be removed in 5.4.
- Moved the Cloth and Collider LODData structure definition out of the .h to their respective .cpp files as to tidy up the two headers.
- Fixed the mesh skinning if the number of influences ever gets greater than 12.

diff --git a/Engine/Plugins/Experimental/ChaosCloth/Source/Chaos/Private/ChaosCloth/ChaosClothingSimulation.cpp b/Engine/Plugins/Experimental/ChaosCloth/Source/Chaos/Private/ChaosCloth/ChaosClothingSimulation.cpp
--- a/Engine/Plugins/Experimental/ChaosCloth/Source/Chaos/Private/ChaosCloth/ChaosClothingSimulation.cpp
+++ b/Engine/Plugins/Experimental/ChaosCloth/Source/Chaos/Private/ChaosCloth/ChaosClothingSimulation.cpp
@@ -8,7 +8,7 @@
 #include "ChaosCloth/ChaosWeightMapTarget.h"
 
 #include "ChaosCloth/ChaosClothingSimulationSolver.h"
-#include "ChaosCloth/ChaosClothingSimulationMesh.h"
+#include "ChaosCloth/ChaosClothingSimulationSkeletalMesh.h"
 #include "ChaosCloth/ChaosClothingSimulationCloth.h"
 #include "ChaosCloth/ChaosClothingSimulationCollider.h"
 
@@ -82,7 +82,6 @@ bool bChaos_GetSimData_ISPC_Enabled = true;
 FAutoConsoleVariableRef CVarChaosGetSimDataISPCEnabled(TEXT("p.Chaos.GetSimData.ISPC"), bChaos_GetSimData_ISPC_Enabled, TEXT("Whether to use ISPC optimizations when getting simulation data"));
 #endif
 
-
 DECLARE_CYCLE_STAT(TEXT("Chaos Cloth Simulate"), STAT_ChaosClothSimulate, STATGROUP_ChaosCloth);
 DECLARE_CYCLE_STAT(TEXT("Chaos Cloth Create Actor"), STAT_ChaosClothCreateActor, STATGROUP_ChaosCloth);
 DECLARE_CYCLE_STAT(TEXT("Chaos Cloth Get Simulation Data"), STAT_ChaosClothGetSimulationData, STATGROUP_ChaosCloth);
@@ -385,7 +384,7 @@ void FClothingSimulation::CreateActor(USkeletalMeshComponent* InOwnerComponent,
 	}
 
 	// Create mesh node
-	const int32 MeshIndex = Meshes.Emplace(MakeUnique<FClothingSimulationMesh>(
+	const int32 MeshIndex = Meshes.Emplace(MakeUnique<FClothingSimulationSkeletalMesh>(
 		Asset,
 		InOwnerComponent));
 
@@ -401,7 +400,9 @@ void FClothingSimulation::CreateActor(USkeletalMeshComponent* InOwnerComponent,
 
 	// Create cloth node
 	const int32 ClothIndex = Cloths.Emplace(MakeUnique<FClothingSimulationCloth>(
+		PRAGMA_DISABLE_DEPRECATION_WARNINGS  // TODO: CHAOS_IS_CLOTHINGSIMULATIONMESH_ABSTRACT
 		Meshes[MeshIndex].Get(),
+		PRAGMA_ENABLE_DEPRECATION_WARNINGS
 		TArray<FClothingSimulationCollider*>({ Colliders[ColliderIndex].Get() }),
 		InSimDataIndex,
 		(FClothingSimulationCloth::EMassMode)ClothConfig->MassMode,
@@ -682,7 +683,9 @@ void FClothingSimulation::GetSimulationData(
 		}
 
 		// If the LOD has changed while the simulation is suspended, the cloth still needs to be updated with the correct LOD data
+PRAGMA_DISABLE_DEPRECATION_WARNINGS  // TODO: CHAOS_IS_CLOTHINGSIMULATIONMESH_ABSTRACT
 		const int32 LODIndex = Cloth->GetMesh()->GetLODIndex();
+PRAGMA_ENABLE_DEPRECATION_WARNINGS
 		if (LODIndex != Cloth->GetLODIndex(Solver.Get()))
 		{
 			Solver->Update(FSolverReal(0.));  // Update for LOD switching, but do not simulate
@@ -751,14 +754,9 @@ void FClothingSimulation::GetSimulationData(
 		}
 
 		// Set the current LOD these data apply to, so that the correct deformer mappings can be applied
-		if (const UClothingAssetCommon* const ClothingAsset = Cloth->GetMesh()->GetAsset())
-		{
-			Data.LODIndex = ClothingAsset->LodMap.Find(LODIndex);  // Store the mesh LOD index, which is different to the cloth asset's LOD index
-		}
-		else
-		{
-			Data.LODIndex = 0;
-		}
+PRAGMA_DISABLE_DEPRECATION_WARNINGS  // TODO: CHAOS_IS_CLOTHINGSIMULATIONMESH_ABSTRACT
+		Data.LODIndex = Cloth->GetMesh()->GetOwnerLODIndex(LODIndex);  // The owner component LOD index can be different to the cloth mesh LOD index
+PRAGMA_ENABLE_DEPRECATION_WARNINGS
 	}
 }
 
@@ -841,7 +839,7 @@ void FClothingSimulation::RefreshClothConfig(const IClothingSimulationContext* I
 	// Recreate all cloths
 	for (TUniquePtr<FClothingSimulationCloth>& Cloth : Cloths)
 	{
-		FClothingSimulationMesh* const Mesh = Cloth->GetMesh();
+		FClothingSimulationSkeletalMesh* const Mesh = static_cast<FClothingSimulationSkeletalMesh*>(Cloth->GetMesh());
 		TArray<FClothingSimulationCollider*> ClothColliders = Cloth->GetColliders();
 		const uint32 GroupId = Cloth->GetGroupId();
 		const UChaosClothConfig* const ClothConfig = Mesh->GetAsset()->GetClothConfig<UChaosClothConfig>();
diff --git a/Engine/Plugins/Experimental/ChaosCloth/Source/Chaos/Private/ChaosCloth/ChaosClothingSimulationCloth.cpp b/Engine/Plugins/Experimental/ChaosCloth/Source/Chaos/Private/ChaosCloth/ChaosClothingSimulationCloth.cpp
--- a/Engine/Plugins/Experimental/ChaosCloth/Source/Chaos/Private/ChaosCloth/ChaosClothingSimulationCloth.cpp
+++ b/Engine/Plugins/Experimental/ChaosCloth/Source/Chaos/Private/ChaosCloth/ChaosClothingSimulationCloth.cpp
@@ -17,6 +17,44 @@ namespace ClothingSimulationClothConsoleVariables
 	TAutoConsoleVariable<bool> CVarLegacyDisablesAccurateWind(TEXT("p.ChaosCloth.LegacyDisablesAccurateWind"), true, TEXT("Whether using the Legacy wind model switches off the accurate wind model, or adds up to it"));
 }
 
+struct FClothingSimulationCloth::FLODData
+{
+	// Input mesh
+	const int32 NumParticles;
+	const TConstArrayView<uint32> Indices;
+	const TArray<TConstArrayView<FRealSingle>> WeightMaps;
+	const TArray<TConstArrayView<TTuple<int32, int32, FRealSingle>>> Tethers;
+
+	// Per Solver data
+	struct FSolverData
+	{
+		int32 Offset;
+		FTriangleMesh TriangleMesh;  // TODO: Triangle Mesh shouldn't really be solver dependent (ie not use an offset)
+	};
+	TMap<FClothingSimulationSolver*, FSolverData> SolverData;
+
+	// Stats
+	int32 NumKinenamicParticles;
+	int32 NumDynammicParticles;
+
+	FLODData(
+		int32 InNumParticles,
+		const TConstArrayView<uint32>& InIndices,
+		const TArray<TConstArrayView<FRealSingle>>& InWeightMaps,
+		const TArray<TConstArrayView<TTuple<int32, int32, FRealSingle>>>& InTethers);
+
+	void Add(FClothingSimulationSolver* Solver, FClothingSimulationCloth* Cloth, int32 LODIndex);
+	void Remove(FClothingSimulationSolver* Solver);
+
+	void Update(FClothingSimulationSolver* Solver, FClothingSimulationCloth* Cloth);
+
+	void Enable(FClothingSimulationSolver* Solver, bool bEnable) const;
+
+	void ResetStartPose(FClothingSimulationSolver* Solver) const;
+
+	void UpdateNormals(FClothingSimulationSolver* Solver) const;
+};
+
 FClothingSimulationCloth::FLODData::FLODData(
 	int32 InNumParticles,
 	const TConstArrayView<uint32>& InIndices,
@@ -50,7 +88,9 @@ void FClothingSimulationCloth::FLODData::Add(FClothingSimulationSolver* Solver,
 	Offset = Solver->AddParticles(NumParticles, Cloth->GroupId);  // TODO: Have a per solver map of offset
 
 	// Update source mesh for this LOD, this is required prior to reset the start pose
+PRAGMA_DISABLE_DEPRECATION_WARNINGS  // TODO: CHAOS_IS_CLOTHINGSIMULATIONMESH_ABSTRACT
 	Cloth->Mesh->Update(Solver, INDEX_NONE, InLODIndex, 0, Offset);
+PRAGMA_ENABLE_DEPRECATION_WARNINGS
 
 	// Reset the particles start pose before setting up mass and constraints
 	ResetStartPose(Solver);
@@ -206,7 +246,9 @@ void FClothingSimulationCloth::FLODData::Add(FClothingSimulationSolver* Solver,
 	{
 		const TConstArrayView<FRealSingle>& TetherScaleMultipliers = WeightMaps[(int32)EChaosWeightMapTarget::TetherScale];
 
+PRAGMA_DISABLE_DEPRECATION_WARNINGS  // TODO: CHAOS_IS_CLOTHINGSIMULATIONMESH_ABSTRACT
 		const Softs::FSolverReal MeshScale = Cloth->Mesh->GetScale();
+PRAGMA_ENABLE_DEPRECATION_WARNINGS
 
 		ClothConstraints.SetLongRangeConstraints(
 			Tethers,
@@ -276,7 +318,9 @@ void FClothingSimulationCloth::FLODData::Update(FClothingSimulationSolver* Solve
 	const int32 Offset = SolverData.FindChecked(Solver).Offset;
 	check(Offset != INDEX_NONE);
 
+PRAGMA_DISABLE_DEPRECATION_WARNINGS  // TODO: CHAOS_IS_CLOTHINGSIMULATIONMESH_ABSTRACT
 	const Softs::FSolverReal MeshScale = Cloth->Mesh->GetScale();
+PRAGMA_ENABLE_DEPRECATION_WARNINGS
 
 	// Update the animatable constraint parameters
 	FClothConstraints& ClothConstraints = Solver->GetClothConstraints(Offset);
@@ -334,7 +378,9 @@ void FClothingSimulationCloth::FLODData::UpdateNormals(FClothingSimulationSolver
 }
 
 FClothingSimulationCloth::FClothingSimulationCloth(
+PRAGMA_DISABLE_DEPRECATION_WARNINGS  // TODO: CHAOS_IS_CLOTHINGSIMULATIONMESH_ABSTRACT
 	FClothingSimulationMesh* InMesh,
+PRAGMA_ENABLE_DEPRECATION_WARNINGS
 	TArray<FClothingSimulationCollider*>&& InColliders,
 	uint32 InGroupId,
 	EMassMode InMassMode,
@@ -442,6 +488,7 @@ FClothingSimulationCloth::~FClothingSimulationCloth()
 {
 }
 
+PRAGMA_DISABLE_DEPRECATION_WARNINGS  // TODO: CHAOS_IS_CLOTHINGSIMULATIONMESH_ABSTRACT
 void FClothingSimulationCloth::SetMesh(FClothingSimulationMesh* InMesh)
 {
 	Mesh = InMesh;
@@ -452,11 +499,11 @@ void FClothingSimulationCloth::SetMesh(FClothingSimulationMesh* InMesh)
 	const bool bUseGeodesicTethers = (TetherMode == ETetherMode::Geodesic);
 	for (int32 Index = 0; Index < NumLODs; ++Index)
 	{
-		LODData.Emplace(
+		LODData.Add(MakeUnique<FLODData>(
 			Mesh->GetNumPoints(Index),
 			Mesh->GetIndices(Index),
 			Mesh->GetWeightMaps(Index),
-			Mesh->GetTethers(Index, bUseGeodesicTethers));
+			Mesh->GetTethers(Index, bUseGeodesicTethers)));
 	}
 
 	// Iterate all known solvers
@@ -468,6 +515,7 @@ void FClothingSimulationCloth::SetMesh(FClothingSimulationMesh* InMesh)
 		Solver->RefreshCloth(this);
 	}
 }
+PRAGMA_ENABLE_DEPRECATION_WARNINGS
 
 void FClothingSimulationCloth::SetColliders(TArray<FClothingSimulationCollider*>&& InColliders)
 {
@@ -573,7 +621,7 @@ void FClothingSimulationCloth::Add(FClothingSimulationSolver* Solver)
 	// Add LODs
 	for (int32 Index = 0; Index < LODData.Num(); ++Index)
 	{
-		LODData[Index].Add(Solver, this, Index);
+		LODData[Index]->Add(Solver, this, Index);
 	}
 
 	// Add Colliders
@@ -593,20 +641,20 @@ void FClothingSimulationCloth::Remove(FClothingSimulationSolver* Solver)
 
 	// Remove solver from maps
 	LODIndices.Remove(Solver);
-	for (FLODData& LODDatum: LODData)
+	for (TUniquePtr<FLODData>& LODDatum: LODData)
 	{
-		LODDatum.Remove(Solver);
+		LODDatum->Remove(Solver);
 	}
 }
 
 int32 FClothingSimulationCloth::GetNumParticles(int32 InLODIndex) const
 {
-	return LODData.IsValidIndex(InLODIndex) ? LODData[InLODIndex].NumParticles : 0;
+	return LODData.IsValidIndex(InLODIndex) ? LODData[InLODIndex]->NumParticles : 0;
 }
 
 int32 FClothingSimulationCloth::GetOffset(const FClothingSimulationSolver* Solver, int32 InLODIndex) const
 {
-	return LODData.IsValidIndex(InLODIndex) ? LODData[InLODIndex].SolverData.FindChecked(Solver).Offset : 0;
+	return LODData.IsValidIndex(InLODIndex) ? LODData[InLODIndex]->SolverData.FindChecked(Solver).Offset : 0;
 }
 
 TVec3<FRealSingle> FClothingSimulationCloth::GetGravity(const FClothingSimulationSolver* Solver) const
@@ -642,26 +690,28 @@ const FTriangleMesh& FClothingSimulationCloth::GetTriangleMesh(const FClothingSi
 {
 	const int32 LODIndex = LODIndices.FindChecked(Solver);
 	static const FTriangleMesh EmptyTriangleMesh;
-	return LODData.IsValidIndex(LODIndex) ? LODData[LODIndex].SolverData.FindChecked(Solver).TriangleMesh : EmptyTriangleMesh;
+	return LODData.IsValidIndex(LODIndex) ? LODData[LODIndex]->SolverData.FindChecked(Solver).TriangleMesh : EmptyTriangleMesh;
 }
 
 const TArray<TConstArrayView<FRealSingle>>& FClothingSimulationCloth::GetWeightMaps(const FClothingSimulationSolver* Solver) const
 {
 	const int32 LODIndex = LODIndices.FindChecked(Solver);
 	static const TArray<TConstArrayView<FRealSingle>> EmptyWeightMaps;
-	return LODData.IsValidIndex(LODIndex) ? LODData[LODIndex].WeightMaps : EmptyWeightMaps;
+	return LODData.IsValidIndex(LODIndex) ? LODData[LODIndex]->WeightMaps : EmptyWeightMaps;
 }
 
 const TArray<TConstArrayView<TTuple<int32, int32, float>>>& FClothingSimulationCloth::GetTethers(const FClothingSimulationSolver* Solver) const
 {
 	const int32 LODIndex = LODIndices.FindChecked(Solver);
 	static const TArray<TConstArrayView<TTuple<int32, int32, float>>> EmptyTethers;
-	return LODData.IsValidIndex(LODIndex) ? LODData[LODIndex].Tethers : EmptyTethers;
+	return LODData.IsValidIndex(LODIndex) ? LODData[LODIndex]->Tethers : EmptyTethers;
 }
 
 int32 FClothingSimulationCloth::GetReferenceBoneIndex() const
 {
+PRAGMA_DISABLE_DEPRECATION_WARNINGS  // TODO: CHAOS_IS_CLOTHINGSIMULATIONMESH_ABSTRACT
 	return Mesh ? Mesh->GetReferenceBoneIndex() : INDEX_NONE;
+PRAGMA_ENABLE_DEPRECATION_WARNINGS
 }
 
 void FClothingSimulationCloth::PreUpdate(FClothingSimulationSolver* Solver)
@@ -699,11 +749,13 @@ void FClothingSimulationCloth::Update(FClothingSimulationSolver* Solver)
 	int32& LODIndex = LODIndices.FindChecked(Solver);  // Must be added to solver first
 
 	const int32 PrevLODIndex = LODIndex;
+PRAGMA_DISABLE_DEPRECATION_WARNINGS  // TODO: CHAOS_IS_CLOTHINGSIMULATIONMESH_ABSTRACT
 	LODIndex = bUseLODIndexOverride && LODData.IsValidIndex(LODIndexOverride) ? LODIndexOverride : Mesh->GetLODIndex();
 
 	// Update reference space transform from the mesh's reference bone transform  TODO: Add override in the style of LODIndexOverride
 	const FRigidTransform3 OldReferenceSpaceTransform = ReferenceSpaceTransform;
 	ReferenceSpaceTransform = Mesh->GetReferenceBoneTransform();
+PRAGMA_ENABLE_DEPRECATION_WARNINGS
 	ReferenceSpaceTransform.SetScale3D(FVec3(1.f));
 
 	// Update Cloth Colliders
@@ -721,7 +773,9 @@ void FClothingSimulationCloth::Update(FClothingSimulationSolver* Solver)
 	const int32 Offset = GetOffset(Solver, LODIndex);
 	check(PrevOffset != INDEX_NONE && Offset != INDEX_NONE);
 
+PRAGMA_DISABLE_DEPRECATION_WARNINGS  // TODO: CHAOS_IS_CLOTHINGSIMULATIONMESH_ABSTRACT
 	Mesh->Update(Solver, PrevLODIndex, LODIndex, PrevOffset, Offset);
+PRAGMA_ENABLE_DEPRECATION_WARNINGS
 
 	// LOD Switching
 	if (LODIndex != PrevLODIndex)
@@ -729,17 +783,19 @@ void FClothingSimulationCloth::Update(FClothingSimulationSolver* Solver)
 		if (PrevLODIndex != INDEX_NONE)
 		{
 			// Disable previous LOD's particles
-			LODData[PrevLODIndex].Enable(Solver, false);
+			LODData[PrevLODIndex]->Enable(Solver, false);
 		}
 		if (LODIndex != INDEX_NONE)
 		{
 			// Enable new LOD's particles
-			LODData[LODIndex].Enable(Solver, true);
-			NumActiveKinematicParticles = LODData[LODIndex].NumKinenamicParticles;
-			NumActiveDynamicParticles = LODData[LODIndex].NumDynammicParticles;
+			LODData[LODIndex]->Enable(Solver, true);
+			NumActiveKinematicParticles = LODData[LODIndex]->NumKinenamicParticles;
+			NumActiveDynamicParticles = LODData[LODIndex]->NumDynammicParticles;
 
 			// Wrap new LOD based on previous LOD if possible (can only do 1 level LOD at a time, and if previous LOD exists)
+PRAGMA_DISABLE_DEPRECATION_WARNINGS  // TODO: CHAOS_IS_CLOTHINGSIMULATIONMESH_ABSTRACT
 			bNeedsReset = bNeedsReset || !Mesh->WrapDeformLOD(
+PRAGMA_ENABLE_DEPRECATION_WARNINGS
 				PrevLODIndex,
 				LODIndex,
 				Solver->GetNormals(PrevOffset),
@@ -750,9 +806,9 @@ void FClothingSimulationCloth::Update(FClothingSimulationSolver* Solver)
 				Solver->GetParticleVs(Offset));
 
 				// Update the wind velocity field for the new LOD mesh
-				const TConstArrayView<FRealSingle>& DragMultipliers = LODData[LODIndex].WeightMaps[(int32)EChaosWeightMapTarget::Drag];
-				const TConstArrayView<FRealSingle>& LiftMultipliers = LODData[LODIndex].WeightMaps[(int32)EChaosWeightMapTarget::Lift];
-				const TConstArrayView<FRealSingle>& PressureMultipliers = LODData[LODIndex].WeightMaps[(int32)EChaosWeightMapTarget::Pressure];
+				const TConstArrayView<FRealSingle>& DragMultipliers = LODData[LODIndex]->WeightMaps[(int32)EChaosWeightMapTarget::Drag];
+				const TConstArrayView<FRealSingle>& LiftMultipliers = LODData[LODIndex]->WeightMaps[(int32)EChaosWeightMapTarget::Lift];
+				const TConstArrayView<FRealSingle>& PressureMultipliers = LODData[LODIndex]->WeightMaps[(int32)EChaosWeightMapTarget::Pressure];
 
 				Solver->SetWindAndPressureGeometry(GroupId, GetTriangleMesh(Solver), DragMultipliers, LiftMultipliers, PressureMultipliers);
 		}
@@ -766,7 +822,7 @@ void FClothingSimulationCloth::Update(FClothingSimulationSolver* Solver)
 	if (LODIndex != INDEX_NONE)
 	{
 		// Update Cloth group parameters  TODO: Cloth groups should exist as their own node object so that they can be used by several cloth objects
-		LODData[LODIndex].Update(Solver, this);
+		LODData[LODIndex]->Update(Solver, this);
 
 		// TODO: Move all groupID updates out of the cloth update to allow to use of the same GroupId with different cloths
 
@@ -781,7 +837,7 @@ void FClothingSimulationCloth::Update(FClothingSimulationSolver* Solver)
 			OutAngularVelocityScale = 1.f;
 
 			// Reset to start pose
-			LODData[LODIndex].ResetStartPose(Solver);
+			LODData[LODIndex]->ResetStartPose(Solver);
 			UE_LOG(LogChaosCloth, VeryVerbose, TEXT("Cloth in group Id %d Needs reset."), GroupId);
 		}
 		else if (bNeedsTeleport)
@@ -845,7 +901,7 @@ void FClothingSimulationCloth::PostUpdate(FClothingSimulationSolver* Solver)
 	if (LODIndex != INDEX_NONE)
 	{
 		// Update normals
-		LODData[LODIndex].UpdateNormals(Solver);
+		LODData[LODIndex]->UpdateNormals(Solver);
 	}
 }
 
diff --git a/Engine/Plugins/Experimental/ChaosCloth/Source/Chaos/Private/ChaosCloth/ChaosClothingSimulationCollider.cpp b/Engine/Plugins/Experimental/ChaosCloth/Source/Chaos/Private/ChaosCloth/ChaosClothingSimulationCollider.cpp
--- a/Engine/Plugins/Experimental/ChaosCloth/Source/Chaos/Private/ChaosCloth/ChaosClothingSimulationCollider.cpp
+++ b/Engine/Plugins/Experimental/ChaosCloth/Source/Chaos/Private/ChaosCloth/ChaosClothingSimulationCollider.cpp
@@ -31,6 +31,35 @@ namespace ClothingSimulationColliderConsoleVariables
 	TAutoConsoleVariable<bool> CVarUseOptimizedTaperedCapsule(TEXT("p.ChaosCloth.UseOptimizedTaperedCapsule"), true, TEXT("Use the optimized TaperedCapsule code instead of using a tapered cylinder and two spheres"));
 }
 
+struct FClothingSimulationCollider::FLODData
+{
+	FClothCollisionData ClothCollisionData;
+	int32 NumGeometries;  // Number of collision bodies
+	TMap<FSolverClothPair, int32> Offsets;  // Solver particle offset
+
+	FLODData() : NumGeometries(0) {}
+
+	void Add(
+		FClothingSimulationSolver* Solver,
+		FClothingSimulationCloth* Cloth,
+		const FClothCollisionData& InClothCollisionData,
+		const TArray<FLevelSetCollisionData>& InLevelSetCollisionData,
+		const FReal InScale = 1.f,
+		const TArray<int32>& UsedBoneIndices = TArray<int32>());
+	void Remove(FClothingSimulationSolver* Solver, FClothingSimulationCloth* Cloth);
+
+	void Update(FClothingSimulationSolver* Solver, FClothingSimulationCloth* Cloth, const FClothingSimulationContextCommon* Context);
+
+	void Enable(FClothingSimulationSolver* Solver, FClothingSimulationCloth* Cloth, bool bEnable);
+
+	void ResetStartPose(FClothingSimulationSolver* Solver, FClothingSimulationCloth* Cloth);
+
+	FORCEINLINE static int32 GetMappedBoneIndex(const TArray<int32>& UsedBoneIndices, int32 BoneIndex)
+	{
+		return UsedBoneIndices.IsValidIndex(BoneIndex) ? UsedBoneIndices[BoneIndex] : INDEX_NONE;
+	}
+};
+
 void FClothingSimulationCollider::FLODData::Add(
 	FClothingSimulationSolver* Solver,
 	FClothingSimulationCloth* Cloth,
@@ -380,23 +409,38 @@ FClothingSimulationCollider::FClothingSimulationCollider(
 {
 	// Prepare LOD array
 	const int32 NumLODs = Asset ? Asset->LodData.Num() : 0;
-	LODData.SetNum((int32)ECollisionDataType::LODs + NumLODs);
+	const int32 NumLODData = (int32)ECollisionDataType::LODs + NumLODs;
+	LODData.Reserve(NumLODData);
+	for (int32 Index = 0; Index < NumLODData; ++Index)
+	{
+		LODData.Add(MakeUnique<FLODData>());
+	}
 }
 
 FClothingSimulationCollider::~FClothingSimulationCollider()
 {
 }
 
+int32 FClothingSimulationCollider::GetNumGeometries() const
+{
+	int32 NumGeometries = 0;
+	for (const TUniquePtr<FLODData>& LODDatum : LODData)
+	{
+		NumGeometries += LODDatum->NumGeometries;
+	}
+	return NumGeometries;
+}
+
 FClothCollisionData FClothingSimulationCollider::GetCollisionData(const FClothingSimulationSolver* Solver, const FClothingSimulationCloth* Cloth) const
 {
 	FClothCollisionData ClothCollisionData;
-	ClothCollisionData.Append(LODData[(int32)ECollisionDataType::LODless].ClothCollisionData);
-	ClothCollisionData.Append(LODData[(int32)ECollisionDataType::External].ClothCollisionData);
+	ClothCollisionData.Append(LODData[(int32)ECollisionDataType::LODless]->ClothCollisionData);
+	ClothCollisionData.Append(LODData[(int32)ECollisionDataType::External]->ClothCollisionData);
 
 	const int32 LODIndex = LODIndices.FindChecked(FSolverClothPair(Solver, Cloth));
 	if (LODIndex >= (int32)ECollisionDataType::LODs)
 	{
-		ClothCollisionData.Append(LODData[LODIndex].ClothCollisionData);
+		ClothCollisionData.Append(LODData[LODIndex]->ClothCollisionData);
 	}
 	return ClothCollisionData;
 }
@@ -578,12 +622,12 @@ void FClothingSimulationCollider::ExtractPhysicsAssetCollision(FClothCollisionDa
 
 int32 FClothingSimulationCollider::GetNumGeometries(int32 InSlotIndex) const
 {
-	return LODData.IsValidIndex(InSlotIndex) ? LODData[InSlotIndex].NumGeometries : 0;
+	return LODData.IsValidIndex(InSlotIndex) ? LODData[InSlotIndex]->NumGeometries : 0;
 }
 
 int32 FClothingSimulationCollider::GetOffset(const FClothingSimulationSolver* Solver, const FClothingSimulationCloth* Cloth, int32 InSlotIndex) const
 {
-	const int32* const Offset = LODData.IsValidIndex(InSlotIndex) ? LODData[InSlotIndex].Offsets.Find(FSolverClothPair(Solver, Cloth)) : nullptr;
+	const int32* const Offset = LODData.IsValidIndex(InSlotIndex) ? LODData[InSlotIndex]->Offsets.Find(FSolverClothPair(Solver, Cloth)) : nullptr;
 	return Offset ? *Offset : INDEX_NONE;
 }
 
@@ -599,8 +643,8 @@ bool FClothingSimulationCollider::GetOffsetAndNumGeometries(const FClothingSimul
 
 	if (LODData.IsValidIndex(SlotIndex))
 	{
-		OutOffset = LODData[SlotIndex].Offsets.FindChecked(FSolverClothPair(Solver, Cloth));
-		OutNumGeometries = LODData[SlotIndex].NumGeometries;
+		OutOffset = LODData[SlotIndex]->Offsets.FindChecked(FSolverClothPair(Solver, Cloth));
+		OutNumGeometries = LODData[SlotIndex]->NumGeometries;
 	}
 
 	return OutOffset != INDEX_NONE && OutNumGeometries > 0;
@@ -636,7 +680,7 @@ void FClothingSimulationCollider::Add(FClothingSimulationSolver* Solver, FClothi
 	TArray<int32> UsedBoneIndices;
 	ExtractPhysicsAssetCollision(PhysicsAssetCollisionData, LevelSetCollisions, UsedBoneIndices);
 
-	LODData[(int32)ECollisionDataType::LODless].Add(Solver, Cloth, PhysicsAssetCollisionData, LevelSetCollisions, Scale, UsedBoneIndices);
+	LODData[(int32)ECollisionDataType::LODless]->Add(Solver, Cloth, PhysicsAssetCollisionData, LevelSetCollisions, Scale, UsedBoneIndices);
 
 	// Create legacy asset LOD collisions
 	const int32 NumLODs = Asset ? Asset->LodData.Num() : 0;
@@ -658,7 +702,7 @@ void FClothingSimulationCollider::Add(FClothingSimulationSolver* Solver, FClothi
 		const TArray<FLevelSetCollisionData> AssetLevelSetCollisions;		// Shouldn't be any level set collider on legacy asset
 
 		// Add legacy collision
-		LODData[(int32)ECollisionDataType::LODs + Index].Add(Solver, Cloth, AssetCollisionData, AssetLevelSetCollisions, Scale, Asset->UsedBoneIndices);
+		LODData[(int32)ECollisionDataType::LODs + Index]->Add(Solver, Cloth, AssetCollisionData, AssetLevelSetCollisions, Scale, Asset->UsedBoneIndices);
 	}
 }
 
@@ -666,9 +710,9 @@ void FClothingSimulationCollider::Remove(FClothingSimulationSolver* Solver, FClo
 {
 	LODIndices.Remove(FSolverClothPair(Solver, Cloth));
 
-	for (FLODData& LODDatum: LODData)
+	for (TUniquePtr<FLODData>& LODDatum: LODData)
 	{
-		LODDatum.Remove(Solver, Cloth);
+		LODDatum->Remove(Solver, Cloth);
 	}
 }
 
@@ -687,7 +731,7 @@ void FClothingSimulationCollider::PreUpdate(FClothingSimulationSolver* Solver, F
 	// TODO: Get level sets?
 	const TArray<FLevelSetCollisionData> LevelSetCollisions;
 
-	LODData[(int32)ECollisionDataType::External].Add(Solver, Cloth, CollisionData ? *CollisionData : FClothCollisionData(), LevelSetCollisions);
+	LODData[(int32)ECollisionDataType::External]->Add(Solver, Cloth, CollisionData ? *CollisionData : FClothCollisionData(), LevelSetCollisions);
 	
 	// TODO: Find a better way in case the same number but different collisions are being re-added (hash collision data? Provide user dirty function?)
 	bHasExternalCollisionChanged =
@@ -705,9 +749,9 @@ void FClothingSimulationCollider::Update(FClothingSimulationSolver* Solver, FClo
 	// Update the collision transforms
 	// Note: All the LODs are updated at once, so that the previous transforms are always correct during LOD switching
 	const FClothingSimulationContextCommon* const Context = SkeletalMeshComponent ? static_cast<const FClothingSimulationContextCommon*>(SkeletalMeshComponent->GetClothingSimulationContext()) : nullptr;
-	for (FLODData& LODDatum : LODData)
+	for (TUniquePtr<FLODData>& LODDatum : LODData)
 	{
-		LODDatum.Update(Solver, Cloth, Context);
+		LODDatum->Update(Solver, Cloth, Context);
 	}
 
 	// Update current LOD index
@@ -723,10 +767,10 @@ void FClothingSimulationCollider::Update(FClothingSimulationSolver* Solver, FClo
 	if (bHasExternalCollisionChanged)
 	{
 		// Enable non LOD collisions at first initialization
-		LODData[(int32)ECollisionDataType::External].Enable(Solver, Cloth, true);
+		LODData[(int32)ECollisionDataType::External]->Enable(Solver, Cloth, true);
 
 		// Update initial state for collisions
-		LODData[(int32)ECollisionDataType::External].ResetStartPose(Solver, Cloth);
+		LODData[(int32)ECollisionDataType::External]->ResetStartPose(Solver, Cloth);
 	}
 
 	if (LODIndex != PrevLODIndex)
@@ -734,25 +778,25 @@ void FClothingSimulationCollider::Update(FClothingSimulationSolver* Solver, FClo
 		if (PrevLODIndex == INDEX_NONE)  // First run
 		{
 			// Enable non LOD collisions at first initialization
-			LODData[(int32)ECollisionDataType::LODless].Enable(Solver, Cloth, true);
+			LODData[(int32)ECollisionDataType::LODless]->Enable(Solver, Cloth, true);
 
 			// Update initial state for collisions
-			LODData[(int32)ECollisionDataType::LODless].ResetStartPose(Solver, Cloth);
+			LODData[(int32)ECollisionDataType::LODless]->ResetStartPose(Solver, Cloth);
 		}
 		else if (PrevLODIndex >= (int32)ECollisionDataType::LODs)
 		{
 			// Disable previous LOD
-			LODData[PrevLODIndex].Enable(Solver, Cloth, false);
+			LODData[PrevLODIndex]->Enable(Solver, Cloth, false);
 		}
 		if (LODIndex >= (int32)ECollisionDataType::LODs)
 		{
 			// Enable new LOD
-			LODData[LODIndex].Enable(Solver, Cloth, true);
+			LODData[LODIndex]->Enable(Solver, Cloth, true);
 
 			if (PrevLODIndex == INDEX_NONE)
 			{
 				// Update initial state for collisions
-				LODData[LODIndex].ResetStartPose(Solver, Cloth);
+				LODData[LODIndex]->ResetStartPose(Solver, Cloth);
 			}
 		}
 	}
@@ -760,9 +804,9 @@ void FClothingSimulationCollider::Update(FClothingSimulationSolver* Solver, FClo
 
 void FClothingSimulationCollider::ResetStartPose(FClothingSimulationSolver* Solver, FClothingSimulationCloth* Cloth)
 {
-	for (FLODData& LODDatum : LODData)
+	for (TUniquePtr<FLODData>& LODDatum : LODData)
 	{
-		LODDatum.ResetStartPose(Solver, Cloth);
+		LODDatum->ResetStartPose(Solver, Cloth);
 	}
 }
 
diff --git a/Engine/Plugins/Experimental/ChaosCloth/Source/Chaos/Private/ChaosCloth/ChaosClothingSimulationMesh.cpp b/Engine/Plugins/Experimental/ChaosCloth/Source/Chaos/Private/ChaosCloth/ChaosClothingSimulationMesh.cpp
--- a/Engine/Plugins/Experimental/ChaosCloth/Source/Chaos/Private/ChaosCloth/ChaosClothingSimulationMesh.cpp
+++ b/Engine/Plugins/Experimental/ChaosCloth/Source/Chaos/Private/ChaosCloth/ChaosClothingSimulationMesh.cpp
@@ -3,11 +3,14 @@
 #include "ChaosCloth/ChaosClothingSimulationSolver.h"
 #include "ChaosCloth/ChaosWeightMapTarget.h"
 #include "ChaosCloth/ChaosClothPrivate.h"
-#include "Chaos/PBDSoftsSolverParticles.h" // Needed for PAndInvM in WrapDeformLOD
 #include "ClothingSimulation.h"
+#if !defined(CHAOS_IS_CLOTHINGSIMULATIONMESH_ABSTRACT) || !CHAOS_IS_CLOTHINGSIMULATIONMESH_ABSTRACT
+#include "Components/SkeletalMeshComponent.h"
 #include "ClothingAsset.h"
+#else
+#include "Components/SkinnedMeshComponent.h"
+#endif
 #include "Containers/ArrayView.h"
-#include "Components/SkeletalMeshComponent.h"
 #include "Async/ParallelFor.h"
 #if INTEL_ISPC
 #include "ChaosClothingSimulationMesh.ispc.generated.h"
@@ -31,14 +34,27 @@ DECLARE_CYCLE_STAT(TEXT("Chaos Cloth Wrap Deform Cloth LOD"), STAT_ChaosClothWra
 namespace Chaos
 {
 
+#if !defined(CHAOS_IS_CLOTHINGSIMULATIONMESH_ABSTRACT) || !CHAOS_IS_CLOTHINGSIMULATIONMESH_ABSTRACT
+
+PRAGMA_DISABLE_DEPRECATION_WARNINGS  // TODO: CHAOS_IS_CLOTHINGSIMULATIONMESH_ABSTRACT
+
+FClothingSimulationMesh::FClothingSimulationMesh(const USkinnedMeshComponent* InSkinnedMeshComponent)
+	: Asset(nullptr)
+	, SkinnedMeshComponent(InSkinnedMeshComponent)
+{
+}
+
 FClothingSimulationMesh::FClothingSimulationMesh(const UClothingAssetCommon* InAsset, const USkeletalMeshComponent* InSkeletalMeshComponent)
 	: Asset(InAsset)
-	, SkeletalMeshComponent(InSkeletalMeshComponent)
+	, SkinnedMeshComponent(InSkeletalMeshComponent)
 {
 }
 
-FClothingSimulationMesh::~FClothingSimulationMesh()
+FClothingSimulationMesh::~FClothingSimulationMesh() = default;
+
+const USkeletalMeshComponent* FClothingSimulationMesh::GetSkeletalMeshComponent() const
 {
+	return Cast<USkeletalMeshComponent>(SkinnedMeshComponent);
 }
 
 int32 FClothingSimulationMesh::GetNumLODs() const
@@ -50,9 +66,9 @@ int32 FClothingSimulationMesh::GetLODIndex() const
 {
 	int32 LODIndex = INDEX_NONE;
 
-	if (Asset && SkeletalMeshComponent)
+	if (Asset && SkinnedMeshComponent)
 	{
-		const int32 MeshLODIndex = SkeletalMeshComponent->GetPredictedLODLevel();
+		const int32 MeshLODIndex = SkinnedMeshComponent->GetPredictedLODLevel();
 		if (Asset->LodMap.IsValidIndex(MeshLODIndex))
 		{
 			const int32 ClothLODIndex = Asset->LodMap[MeshLODIndex];
@@ -65,16 +81,47 @@ int32 FClothingSimulationMesh::GetLODIndex() const
 	return LODIndex;
 }
 
+bool FClothingSimulationMesh::IsValidLODIndex(int32 LODIndex) const
+{
+	return Asset && Asset->LodData.IsValidIndex(LODIndex);
+}
+
+int32 FClothingSimulationMesh::GetOwnerLODIndex(int32 LODIndex) const
+{
+	const int32 OwnerLODIndex = Asset ? Asset->LodMap.Find(LODIndex) : INDEX_NONE;
+	return OwnerLODIndex != INDEX_NONE ? OwnerLODIndex : 0;  // Safer to return the default LOD 0 than INDEX_NONE in this case
+}
+
 int32 FClothingSimulationMesh::GetNumPoints(int32 LODIndex) const
 {
-	return (Asset && Asset->LodData.IsValidIndex(LODIndex)) ?
-		Asset->LodData[LODIndex].PhysicalMeshData.Vertices.Num() :
-		0;
+	return IsValidLODIndex(LODIndex) ? Asset->LodData[LODIndex].PhysicalMeshData.Vertices.Num() : 0;
+}
+
+TConstArrayView<FVector3f> FClothingSimulationMesh::GetPositions(int32 LODIndex) const
+{
+	if (IsValidLODIndex(LODIndex))
+	{
+		const FClothLODDataCommon& ClothLODData = Asset->LodData[LODIndex];
+		const FClothPhysicalMeshData& ClothPhysicalMeshData = ClothLODData.PhysicalMeshData;
+		return TConstArrayView<FVector3f>(ClothPhysicalMeshData.Vertices);
+	}
+	return TConstArrayView<FVector3f>();
+}
+
+TConstArrayView<FVector3f> FClothingSimulationMesh::GetNormals(int32 LODIndex) const
+{
+	if (IsValidLODIndex(LODIndex))
+	{
+		const FClothLODDataCommon& ClothLODData = Asset->LodData[LODIndex];
+		const FClothPhysicalMeshData& ClothPhysicalMeshData = ClothLODData.PhysicalMeshData;
+		return TConstArrayView<FVector3f>(ClothPhysicalMeshData.Normals);
+	}
+	return TConstArrayView<FVector3f>();
 }
 
 TConstArrayView<uint32> FClothingSimulationMesh::GetIndices(int32 LODIndex) const
 {
-	return (Asset && Asset->LodData.IsValidIndex(LODIndex)) ?
+	return IsValidLODIndex(LODIndex) ?
 		TConstArrayView<uint32>(Asset->LodData[LODIndex].PhysicalMeshData.Indices) :
 		TConstArrayView<uint32>();
 }
@@ -82,7 +129,7 @@ TConstArrayView<uint32> FClothingSimulationMesh::GetIndices(int32 LODIndex) cons
 TArray<TConstArrayView<FRealSingle>> FClothingSimulationMesh::GetWeightMaps(int32 LODIndex) const
 {
 	TArray<TConstArrayView<FRealSingle>> WeightMaps;
-	if (Asset && Asset->LodData.IsValidIndex(LODIndex))
+	if (IsValidLODIndex(LODIndex))
 	{
 		const FClothLODDataCommon& ClothLODData = Asset->LodData[LODIndex];
 		const FClothPhysicalMeshData& ClothPhysicalMeshData = ClothLODData.PhysicalMeshData;
@@ -107,7 +154,7 @@ TArray<TConstArrayView<FRealSingle>> FClothingSimulationMesh::GetWeightMaps(int3
 TArray<TConstArrayView<TTuple<int32, int32, float>>> FClothingSimulationMesh::GetTethers(int32 LODIndex, bool bUseGeodesicTethers) const
 {
 	TArray<TConstArrayView<TTuple<int32, int32, float>>> Tethers;
-	if (Asset && Asset->LodData.IsValidIndex(LODIndex))
+	if (IsValidLODIndex(LODIndex))
 	{
 		const FClothLODDataCommon& ClothLODData = Asset->LodData[LODIndex];
 		const FClothPhysicalMeshData& ClothPhysicalMeshData = ClothLODData.PhysicalMeshData;
@@ -130,33 +177,80 @@ int32 FClothingSimulationMesh::GetReferenceBoneIndex() const
 
 FRigidTransform3 FClothingSimulationMesh::GetReferenceBoneTransform() const
 {
-	if (SkeletalMeshComponent)
+	if (const FClothingSimulationContextCommon* const Context = GetContext())
 	{
-		if (const FClothingSimulationContextCommon* const Context =
-			static_cast<const FClothingSimulationContextCommon*>(SkeletalMeshComponent->GetClothingSimulationContext()))
-		{
-			const int32 ReferenceBoneIndex = GetReferenceBoneIndex();
-			const TArray<FTransform>& BoneTransforms = Context->BoneTransforms;
+		const int32 ReferenceBoneIndex = GetReferenceBoneIndex();
+		const TArray<FTransform>& BoneTransforms = Context->BoneTransforms;
 
-			return BoneTransforms.IsValidIndex(ReferenceBoneIndex) ?
-				BoneTransforms[ReferenceBoneIndex] * Context->ComponentToWorld :
-				Context->ComponentToWorld;
-		}
+		return BoneTransforms.IsValidIndex(ReferenceBoneIndex) ?
+			BoneTransforms[ReferenceBoneIndex] * Context->ComponentToWorld :
+			Context->ComponentToWorld;
 	}
 	return FRigidTransform3::Identity;
 }
 
-Softs::FSolverReal FClothingSimulationMesh::GetScale() const
+const FTransform& FClothingSimulationMesh::GetComponentToWorldTransform() const
+{
+	return GetContext() ? GetContext()->ComponentToWorld : FTransform::Identity;
+}
+
+const TArray<FMatrix44f>& FClothingSimulationMesh::GetRefToLocalMatrices() const
 {
-	if (SkeletalMeshComponent)
+	static TArray<FMatrix44f> EmptyArray;
+	return GetContext() ? GetContext()->RefToLocals : EmptyArray;
+}
+
+TConstArrayView<int32> FClothingSimulationMesh::GetBoneMap() const
+{
+	return Asset ? Asset->UsedBoneIndices : TConstArrayView<int32>();
+}
+
+TConstArrayView<FClothVertBoneData> FClothingSimulationMesh::GetBoneData(int32 LODIndex) const
+{
+	if (IsValidLODIndex(LODIndex))
 	{
-		if (const FClothingSimulationContextCommon* const Context =
-			static_cast<const FClothingSimulationContextCommon*>(SkeletalMeshComponent->GetClothingSimulationContext()))
-		{
-			return (Softs::FSolverReal)Context->ComponentToWorld.GetScale3D().GetMax();
-		}
+		const FClothLODDataCommon& ClothLODData = Asset->LodData[LODIndex];
+		const FClothPhysicalMeshData& ClothPhysicalMeshData = ClothLODData.PhysicalMeshData;
+		UE_CLOG(ClothPhysicalMeshData.MaxBoneWeights > 12, LogChaosCloth, Warning, TEXT("The cloth physics mesh skinning code can't cope with more than 12 bone influences."));
+		return ClothPhysicalMeshData.BoneData;
 	}
-	return (Softs::FSolverReal)1.;
+	static TArray<FClothVertBoneData> EmptyArray;
+	return EmptyArray;
+}
+
+const TArray<FMeshToMeshVertData>& FClothingSimulationMesh::GetTransitionUpSkinData(int32 LODIndex) const
+{
+	static TArray<FMeshToMeshVertData> EmptyArray;
+	return IsValidLODIndex(LODIndex) ? Asset->LodData[LODIndex].TransitionUpSkinData : EmptyArray;
+}
+
+const TArray<FMeshToMeshVertData>& FClothingSimulationMesh::GetTransitionDownSkinData(int32 LODIndex) const
+{
+	static TArray<FMeshToMeshVertData> EmptyArray;
+	return IsValidLODIndex(LODIndex) ? Asset->LodData[LODIndex].TransitionDownSkinData : EmptyArray;
+}
+
+const FClothingSimulationContextCommon* FClothingSimulationMesh::GetContext() const
+{
+	return GetSkeletalMeshComponent() ?
+		static_cast<const FClothingSimulationContextCommon*>(GetSkeletalMeshComponent()->GetClothingSimulationContext()) :
+		nullptr;
+}
+
+PRAGMA_ENABLE_DEPRECATION_WARNINGS
+
+#else
+FClothingSimulationMesh::FClothingSimulationMesh(const USkinnedMeshComponent* InSkinnedMeshComponent)
+	: SkinnedMeshComponent(InSkinnedMeshComponent)
+{
+}
+
+FClothingSimulationMesh::~FClothingSimulationMesh() = default;
+#endif
+
+Softs::FSolverReal FClothingSimulationMesh::GetScale() const
+{
+	return GetComponentToWorldTransform().GetScale3D().GetMax();
 }
 
 bool FClothingSimulationMesh::WrapDeformLOD(
@@ -170,16 +264,15 @@ bool FClothingSimulationMesh::WrapDeformLOD(
 	SCOPE_CYCLE_COUNTER(STAT_ChaosClothWrapDeformMesh);
 
 	const int32 NumLODsPassed = FMath::Abs(LODIndex - PrevLODIndex);
-	if (NumLODsPassed != 1 || !Asset || !Asset->LodData.IsValidIndex(PrevLODIndex) || !Asset->LodData.IsValidIndex(LODIndex))
+	if (NumLODsPassed != 1 || !IsValidLODIndex(PrevLODIndex) || !IsValidLODIndex(LODIndex))
 	{
 		return false;
 	}
 
-	const FClothLODDataCommon& LODData = Asset->LodData[LODIndex];
-	const int32 NumPoints = LODData.PhysicalMeshData.Vertices.Num();
+	const int32 NumPoints = GetNumPoints(LODIndex);
 	const TArray<FMeshToMeshVertData>& SkinData = (PrevLODIndex < LODIndex) ?
-		LODData.TransitionUpSkinData :
-		LODData.TransitionDownSkinData;
+		GetTransitionUpSkinData(LODIndex) :
+		GetTransitionDownSkinData(LODIndex);
 
 	for (int32 Index = 0; Index < NumPoints; ++Index)  // TODO: Profile for parallel for
 	{
@@ -212,16 +305,15 @@ bool FClothingSimulationMesh::WrapDeformLOD(
 	SCOPE_CYCLE_COUNTER(STAT_ChaosClothWrapDeformClothLOD);
 
 	const int32 NumLODsPassed = FMath::Abs(LODIndex - PrevLODIndex);
-	if (NumLODsPassed != 1 || !Asset || !Asset->LodData.IsValidIndex(PrevLODIndex) || !Asset->LodData.IsValidIndex(LODIndex))
+	if (NumLODsPassed != 1 || !IsValidLODIndex(PrevLODIndex) || !IsValidLODIndex(LODIndex))
 	{
 		return false;
 	}
 
-	const FClothLODDataCommon& LODData = Asset->LodData[LODIndex];
-	const int32 NumPoints = LODData.PhysicalMeshData.Vertices.Num();
+	const int32 NumPoints = GetNumPoints(LODIndex);
 	const TArray<FMeshToMeshVertData>& SkinData = (PrevLODIndex < LODIndex) ?
-		LODData.TransitionUpSkinData :
-		LODData.TransitionDownSkinData;
+		GetTransitionUpSkinData(LODIndex) :
+		GetTransitionDownSkinData(LODIndex);
 
 	for (int32 Index = 0; Index < NumPoints; ++Index)  // TODO: Profile for parallel for
 	{
@@ -258,32 +350,30 @@ void FClothingSimulationMesh::SkinPhysicsMesh(int32 LODIndex, const FVec3& Local
 	SCOPE_CYCLE_COUNTER(STAT_ChaosClothSkinPhysicsMesh);
 	SCOPE_CYCLE_COUNTER(STAT_ClothSkinPhysMesh);
 
-	check(Asset && Asset->LodData.IsValidIndex(LODIndex));
-	const FClothPhysicalMeshData& PhysicalMeshData = Asset->LodData[LODIndex].PhysicalMeshData;
-	UE_CLOG(PhysicalMeshData.MaxBoneWeights > 12, LogChaosCloth, Warning, TEXT("The cloth physics mesh skinning code can't cope with more than 12 bone influences."));
-
-	const uint32 NumPoints = PhysicalMeshData.Vertices.Num();
-
-	check(SkeletalMeshComponent && SkeletalMeshComponent->GetClothingSimulationContext());
-	const FClothingSimulationContextCommon* const Context = static_cast<const FClothingSimulationContextCommon*>(SkeletalMeshComponent->GetClothingSimulationContext());
-	FTransform ComponentToLocalSpaceReal = Context->ComponentToWorld;
+	FTransform ComponentToLocalSpaceReal = GetComponentToWorldTransform();
 	ComponentToLocalSpaceReal.AddToTranslation(-LocalSpaceLocation);
 	const FTransform3f ComponentToLocalSpace(ComponentToLocalSpaceReal);  // LWC: Now in local space, therefore it is safe to use single precision which is the asset data format
 
-	const int32* const RESTRICT BoneMap = Asset->UsedBoneIndices.GetData();
-	const FMatrix44f* const RESTRICT BoneMatrices = Context->RefToLocals.GetData();
-		
+	const int32* const RESTRICT BoneMap = GetBoneMap().GetData();
+	const FMatrix44f* const RESTRICT RefToLocalMatrices = GetRefToLocalMatrices().GetData();
+
+	check(IsValidLODIndex(LODIndex));
+	const uint32 NumPoints = GetNumPoints(LODIndex);
+	const TConstArrayView<FClothVertBoneData> BoneData = GetBoneData(LODIndex);
+	const TConstArrayView<FVector3f> Positions = GetPositions(LODIndex);
+	const TConstArrayView<FVector3f> Normals = GetNormals(LODIndex);
+
 #if INTEL_ISPC
 	if (bChaos_SkinPhysicsMesh_ISPC_Enabled)
 	{
 		ispc::SkinPhysicsMesh(
 			(ispc::FVector3f*)OutPositions,
 			(ispc::FVector3f*)OutNormals,
-			(ispc::FVector3f*)PhysicalMeshData.Vertices.GetData(),
-			(ispc::FVector3f*)PhysicalMeshData.Normals.GetData(),
-			(ispc::FClothVertBoneData*)PhysicalMeshData.BoneData.GetData(),
+			(ispc::FVector3f*)Positions.GetData(),
+			(ispc::FVector3f*)Normals.GetData(),
+			(ispc::FClothVertBoneData*)BoneData.GetData(),
 			BoneMap,
-			(ispc::FMatrix44f*)BoneMatrices,
+			(ispc::FMatrix44f*)RefToLocalMatrices,
 			(ispc::FTransform3f&)ComponentToLocalSpace,
 			NumPoints);
 	}
@@ -292,36 +382,36 @@ void FClothingSimulationMesh::SkinPhysicsMesh(int32 LODIndex, const FVec3& Local
 	{
 		static const uint32 MinParallelVertices = 500;  // 500 seems to be the lowest threshold still giving gains even on profiled assets that are only using a small number of influences
 
-		ParallelFor(NumPoints, [&PhysicalMeshData, &ComponentToLocalSpace, BoneMap, BoneMatrices, &OutPositions, &OutNormals](uint32 VertIndex)
+		ParallelFor(NumPoints, [&BoneData, &Positions, &Normals, &ComponentToLocalSpace, BoneMap, RefToLocalMatrices, &OutPositions, &OutNormals](uint32 VertIndex)
 		{
-			const uint16* const RESTRICT BoneIndices = PhysicalMeshData.BoneData[VertIndex].BoneIndices;
-			const float* const RESTRICT BoneWeights = PhysicalMeshData.BoneData[VertIndex].BoneWeights;
+			const uint16* const RESTRICT BoneIndices = BoneData[VertIndex].BoneIndices;
+			const float* const RESTRICT BoneWeights = BoneData[VertIndex].BoneWeights;
 	
 			// WARNING - HORRIBLE UNROLLED LOOP + JUMP TABLE BELOW
 			// done this way because this is a pretty tight and performance critical loop. essentially
 			// rather than checking each influence we can just jump into this switch and fall through
 			// everything to compose the final skinned data
-			const FVector3f& RefParticle = PhysicalMeshData.Vertices[VertIndex];
-			const FVector3f& RefNormal = PhysicalMeshData.Normals[VertIndex];
+			const FVector3f& RefParticle = Positions[VertIndex];
+			const FVector3f& RefNormal = Normals[VertIndex];
 
 			FVector3f Position(ForceInitToZero);
 			FVector3f Normal(ForceInitToZero);
-
-			switch (PhysicalMeshData.BoneData[VertIndex].NumInfluences)
+			switch (BoneData[VertIndex].NumInfluences)
 			{
-			case 12: AddInfluence(Position, Normal, RefParticle, RefNormal, BoneMatrices[BoneMap[BoneIndices[11]]], BoneWeights[11]);  // Intentional fall through
-			case 11: AddInfluence(Position, Normal, RefParticle, RefNormal, BoneMatrices[BoneMap[BoneIndices[10]]], BoneWeights[10]);  // Intentional fall through
-			case 10: AddInfluence(Position, Normal, RefParticle, RefNormal, BoneMatrices[BoneMap[BoneIndices[ 9]]], BoneWeights[ 9]);  // Intentional fall through
-			case  9: AddInfluence(Position, Normal, RefParticle, RefNormal, BoneMatrices[BoneMap[BoneIndices[ 8]]], BoneWeights[ 8]);  // Intentional fall through
-			case  8: AddInfluence(Position, Normal, RefParticle, RefNormal, BoneMatrices[BoneMap[BoneIndices[ 7]]], BoneWeights[ 7]);  // Intentional fall through
-			case  7: AddInfluence(Position, Normal, RefParticle, RefNormal, BoneMatrices[BoneMap[BoneIndices[ 6]]], BoneWeights[ 6]);  // Intentional fall through
-			case  6: AddInfluence(Position, Normal, RefParticle, RefNormal, BoneMatrices[BoneMap[BoneIndices[ 5]]], BoneWeights[ 5]);  // Intentional fall through
-			case  5: AddInfluence(Position, Normal, RefParticle, RefNormal, BoneMatrices[BoneMap[BoneIndices[ 4]]], BoneWeights[ 4]);  // Intentional fall through
-			case  4: AddInfluence(Position, Normal, RefParticle, RefNormal, BoneMatrices[BoneMap[BoneIndices[ 3]]], BoneWeights[ 3]);  // Intentional fall through
-			case  3: AddInfluence(Position, Normal, RefParticle, RefNormal, BoneMatrices[BoneMap[BoneIndices[ 2]]], BoneWeights[ 2]);  // Intentional fall through
-			case  2: AddInfluence(Position, Normal, RefParticle, RefNormal, BoneMatrices[BoneMap[BoneIndices[ 1]]], BoneWeights[ 1]);  // Intentional fall through
-			case  1: AddInfluence(Position, Normal, RefParticle, RefNormal, BoneMatrices[BoneMap[BoneIndices[ 0]]], BoneWeights[ 0]);  // Intentional fall through
-			default: break;
+			default:  // Intentional fall through
+			case 12: AddInfluence(Position, Normal, RefParticle, RefNormal, RefToLocalMatrices[BoneMap[BoneIndices[11]]], BoneWeights[11]);  // Intentional fall through
+			case 11: AddInfluence(Position, Normal, RefParticle, RefNormal, RefToLocalMatrices[BoneMap[BoneIndices[10]]], BoneWeights[10]);  // Intentional fall through
+			case 10: AddInfluence(Position, Normal, RefParticle, RefNormal, RefToLocalMatrices[BoneMap[BoneIndices[ 9]]], BoneWeights[ 9]);  // Intentional fall through
+			case  9: AddInfluence(Position, Normal, RefParticle, RefNormal, RefToLocalMatrices[BoneMap[BoneIndices[ 8]]], BoneWeights[ 8]);  // Intentional fall through
+			case  8: AddInfluence(Position, Normal, RefParticle, RefNormal, RefToLocalMatrices[BoneMap[BoneIndices[ 7]]], BoneWeights[ 7]);  // Intentional fall through
+			case  7: AddInfluence(Position, Normal, RefParticle, RefNormal, RefToLocalMatrices[BoneMap[BoneIndices[ 6]]], BoneWeights[ 6]);  // Intentional fall through
+			case  6: AddInfluence(Position, Normal, RefParticle, RefNormal, RefToLocalMatrices[BoneMap[BoneIndices[ 5]]], BoneWeights[ 5]);  // Intentional fall through
+			case  5: AddInfluence(Position, Normal, RefParticle, RefNormal, RefToLocalMatrices[BoneMap[BoneIndices[ 4]]], BoneWeights[ 4]);  // Intentional fall through
+			case  4: AddInfluence(Position, Normal, RefParticle, RefNormal, RefToLocalMatrices[BoneMap[BoneIndices[ 3]]], BoneWeights[ 3]);  // Intentional fall through
+			case  3: AddInfluence(Position, Normal, RefParticle, RefNormal, RefToLocalMatrices[BoneMap[BoneIndices[ 2]]], BoneWeights[ 2]);  // Intentional fall through
+			case  2: AddInfluence(Position, Normal, RefParticle, RefNormal, RefToLocalMatrices[BoneMap[BoneIndices[ 1]]], BoneWeights[ 1]);  // Intentional fall through
+			case  1: AddInfluence(Position, Normal, RefParticle, RefNormal, RefToLocalMatrices[BoneMap[BoneIndices[ 0]]], BoneWeights[ 0]);  // Intentional fall through
+			case  0: break;
 			}
 
 			OutPositions[VertIndex] = FSolverVec3(ComponentToLocalSpace.TransformPosition(Position));
@@ -341,7 +431,7 @@ void FClothingSimulationMesh::Update(
 	check(Solver);
 
 	// Exit if any inputs are missing or not ready, and if the LOD is invalid
-	if (!Asset || !Asset->LodData.IsValidIndex(LODIndex) || !SkeletalMeshComponent || !SkeletalMeshComponent->GetClothingSimulationContext())
+	if (!IsValidLODIndex(LODIndex) || !SkinnedMeshComponent)
 	{
 		return;
 	}
@@ -366,8 +456,7 @@ void FClothingSimulationMesh::Update(
 		if (!bValidWrap)
 		{
 			// The previous LOD is invalid, reset old positions with the new LOD
-			const FClothPhysicalMeshData& PhysicalMeshData = Asset->LodData[LODIndex].PhysicalMeshData;
-			const int32 NumPoints = PhysicalMeshData.Vertices.Num();
+			const int32 NumPoints = GetNumPoints(LODIndex);
 
 			for (int32 Index = 0; Index < NumPoints; ++Index)
 			{
diff --git a/Engine/Plugins/Experimental/ChaosCloth/Source/Chaos/Private/ChaosCloth/ChaosClothingSimulationSkeletalMesh.cpp b/Engine/Plugins/Experimental/ChaosCloth/Source/Chaos/Private/ChaosCloth/ChaosClothingSimulationSkeletalMesh.cpp
new file mode 100644
--- /dev/null
+++ b/Engine/Plugins/Experimental/ChaosCloth/Source/Chaos/Private/ChaosCloth/ChaosClothingSimulationSkeletalMesh.cpp
@@ -0,0 +1,199 @@
+// Copyright Epic Games, Inc. All Rights Reserved.
+#include "ChaosCloth/ChaosClothingSimulationSkeletalMesh.h"
+#include "ChaosCloth/ChaosClothPrivate.h"
+#include "Components/SkeletalMeshComponent.h"
+#include "ChaosWeightMapTarget.h"
+#include "ClothingAsset.h"
+#include "ClothingSimulation.h"
+
+namespace Chaos
+{
+	FClothingSimulationSkeletalMesh::FClothingSimulationSkeletalMesh(const UClothingAssetCommon* InAsset, const USkeletalMeshComponent* InSkeletalMeshComponent)
+PRAGMA_DISABLE_DEPRECATION_WARNINGS  // TODO: CHAOS_IS_CLOTHINGSIMULATIONMESH_ABSTRACT
+		: FClothingSimulationMesh(InSkeletalMeshComponent)
+PRAGMA_ENABLE_DEPRECATION_WARNINGS
+		, Asset(InAsset)
+	{
+	}
+
+	int32 FClothingSimulationSkeletalMesh::GetNumLODs() const
+	{
+		return Asset ? Asset->LodData.Num() : 0;
+	}
+
+	int32 FClothingSimulationSkeletalMesh::GetLODIndex() const
+	{
+		int32 LODIndex = INDEX_NONE;
+
+PRAGMA_DISABLE_DEPRECATION_WARNINGS  // TODO: CHAOS_IS_CLOTHINGSIMULATIONMESH_ABSTRACT
+		if (Asset && GetSkinnedMeshComponent())
+		{
+			const int32 OwnerLODIndex = GetSkinnedMeshComponent()->GetPredictedLODLevel();
+PRAGMA_ENABLE_DEPRECATION_WARNINGS
+			if (Asset->LodMap.IsValidIndex(OwnerLODIndex))
+			{
+				const int32 ClothLODIndex = Asset->LodMap[OwnerLODIndex];
+				if (Asset->LodData.IsValidIndex(ClothLODIndex))
+				{
+					LODIndex = ClothLODIndex;
+				}
+			}
+		}
+		return LODIndex;
+	}
+
+	int32 FClothingSimulationSkeletalMesh::GetOwnerLODIndex(int32 LODIndex) const
+	{
+		const int32 OwnerLODIndex = Asset ? Asset->LodMap.Find(LODIndex) : INDEX_NONE;
+		return OwnerLODIndex != INDEX_NONE ? OwnerLODIndex : 0;  // Safer to return the default LOD 0 than INDEX_NONE in this case
+	}
+
+	bool FClothingSimulationSkeletalMesh::IsValidLODIndex(int32 LODIndex) const
+	{
+		return Asset && Asset->LodData.IsValidIndex(LODIndex);
+	}
+
+	int32 FClothingSimulationSkeletalMesh::GetNumPoints(int32 LODIndex) const
+	{
+		return IsValidLODIndex(LODIndex) ? Asset->LodData[LODIndex].PhysicalMeshData.Vertices.Num() : 0;
+	}
+
+	TConstArrayView<FVector3f> FClothingSimulationSkeletalMesh::GetPositions(int32 LODIndex) const
+	{
+		if (IsValidLODIndex(LODIndex))
+		{
+			const FClothLODDataCommon& ClothLODData = Asset->LodData[LODIndex];
+			const FClothPhysicalMeshData& ClothPhysicalMeshData = ClothLODData.PhysicalMeshData;
+			return ClothPhysicalMeshData.Vertices;
+		}
+		return TConstArrayView<FVector3f>();
+	}
+
+	TConstArrayView<FVector3f> FClothingSimulationSkeletalMesh::GetNormals(int32 LODIndex) const
+	{
+		if (IsValidLODIndex(LODIndex))
+		{
+			const FClothLODDataCommon& ClothLODData = Asset->LodData[LODIndex];
+			const FClothPhysicalMeshData& ClothPhysicalMeshData = ClothLODData.PhysicalMeshData;
+			return ClothPhysicalMeshData.Normals;
+		}
+		return TConstArrayView<FVector3f>();
+	}
+
+	TConstArrayView<uint32> FClothingSimulationSkeletalMesh::GetIndices(int32 LODIndex) const
+	{
+		return IsValidLODIndex(LODIndex) ?
+			TConstArrayView<uint32>(Asset->LodData[LODIndex].PhysicalMeshData.Indices) :
+			TConstArrayView<uint32>();
+	}
+
+	TArray<TConstArrayView<FRealSingle>> FClothingSimulationSkeletalMesh::GetWeightMaps(int32 LODIndex) const
+	{
+		TArray<TConstArrayView<FRealSingle>> WeightMaps;
+		if (IsValidLODIndex(LODIndex))
+		{
+			const FClothLODDataCommon& ClothLODData = Asset->LodData[LODIndex];
+			const FClothPhysicalMeshData& ClothPhysicalMeshData = ClothLODData.PhysicalMeshData;
+
+			const UEnum* const ChaosWeightMapTargetEnum = StaticEnum<EChaosWeightMapTarget>();
+			const int32 NumWeightMaps = (int32)ChaosWeightMapTargetEnum->GetMaxEnumValue() + 1;
+
+			WeightMaps.SetNum(NumWeightMaps);
+
+			for (int32 EnumIndex = 0; EnumIndex < ChaosWeightMapTargetEnum->NumEnums(); ++EnumIndex)
+			{
+				const int32 TargetIndex = (int32)ChaosWeightMapTargetEnum->GetValueByIndex(EnumIndex);
+				if (const FPointWeightMap* const WeightMap = ClothPhysicalMeshData.FindWeightMap(TargetIndex))
+				{
+					WeightMaps[TargetIndex] = WeightMap->Values;
+				}
+			}
+		}
+		return WeightMaps;
+	}
+
+	TArray<TConstArrayView<TTuple<int32, int32, float>>> FClothingSimulationSkeletalMesh::GetTethers(int32 LODIndex, bool bUseGeodesicTethers) const
+	{
+		TArray<TConstArrayView<TTuple<int32, int32, float>>> Tethers;
+		if (IsValidLODIndex(LODIndex))
+		{
+			const FClothLODDataCommon& ClothLODData = Asset->LodData[LODIndex];
+			const FClothPhysicalMeshData& ClothPhysicalMeshData = ClothLODData.PhysicalMeshData;
+			const FClothTetherData& ClothTetherData = bUseGeodesicTethers ? ClothPhysicalMeshData.GeodesicTethers : ClothPhysicalMeshData.EuclideanTethers;
+		
+			const int32 NumTetherBatches = ClothTetherData.Tethers.Num();
+			Tethers.Reserve(NumTetherBatches);
+			for (int32 Index = 0; Index < NumTetherBatches; ++Index)
+			{
+				Tethers.Emplace(TConstArrayView<TTuple<int32, int32, float>>(ClothTetherData.Tethers[Index]));
+			}
+		}
+		return Tethers;
+	}
+
+	int32 FClothingSimulationSkeletalMesh::GetReferenceBoneIndex() const
+	{
+		return Asset ? Asset->ReferenceBoneIndex : INDEX_NONE;
+	}
+
+	FRigidTransform3 FClothingSimulationSkeletalMesh::GetReferenceBoneTransform() const
+	{
+		if (const FClothingSimulationContextCommon* const Context = GetContext())
+		{
+			const int32 ReferenceBoneIndex = GetReferenceBoneIndex();
+			const TArray<FTransform>& BoneTransforms = Context->BoneTransforms;
+
+			return BoneTransforms.IsValidIndex(ReferenceBoneIndex) ?
+				BoneTransforms[ReferenceBoneIndex] * Context->ComponentToWorld :
+				Context->ComponentToWorld;
+		}
+		return FRigidTransform3::Identity;
+	}
+
+	const FTransform& FClothingSimulationSkeletalMesh::GetComponentToWorldTransform() const
+	{
+		return GetContext() ? GetContext()->ComponentToWorld : FTransform::Identity;
+	}
+
+	const TArray<FMatrix44f>& FClothingSimulationSkeletalMesh::GetRefToLocalMatrices() const
+	{
+		static TArray<FMatrix44f> EmptyArray;
+		return GetContext() ? GetContext()->RefToLocals : EmptyArray;
+	}
+
+	TConstArrayView<int32> FClothingSimulationSkeletalMesh::GetBoneMap() const
+	{
+		return Asset ? Asset->UsedBoneIndices : TConstArrayView<int32>();
+	}
+
+	TConstArrayView<FClothVertBoneData> FClothingSimulationSkeletalMesh::GetBoneData(int32 LODIndex) const
+	{
+		if (IsValidLODIndex(LODIndex))
+		{
+			const FClothLODDataCommon& ClothLODData = Asset->LodData[LODIndex];
+			const FClothPhysicalMeshData& ClothPhysicalMeshData = ClothLODData.PhysicalMeshData;
+			UE_CLOG(ClothPhysicalMeshData.MaxBoneWeights > 12, LogChaosCloth, Warning, TEXT("The cloth physics mesh skinning code can't cope with more than 12 bone influences."));
+			return ClothPhysicalMeshData.BoneData;
+		}
+		return TConstArrayView<FClothVertBoneData>();
+	}
+
+	const TArray<FMeshToMeshVertData>& FClothingSimulationSkeletalMesh::GetTransitionUpSkinData(int32 LODIndex) const
+	{
+		static TArray<FMeshToMeshVertData> EmptyArray;
+		return IsValidLODIndex(LODIndex) ? Asset->LodData[LODIndex].TransitionUpSkinData : EmptyArray;
+	}
+
+	const TArray<FMeshToMeshVertData>& FClothingSimulationSkeletalMesh::GetTransitionDownSkinData(int32 LODIndex) const
+	{
+		static TArray<FMeshToMeshVertData> EmptyArray;
+		return IsValidLODIndex(LODIndex) ? Asset->LodData[LODIndex].TransitionDownSkinData : EmptyArray;
+	}
+
+	const FClothingSimulationContextCommon* FClothingSimulationSkeletalMesh::GetContext() const
+	{
+		return GetSkeletalMeshComponent() ?
+			static_cast<const FClothingSimulationContextCommon*>(GetSkeletalMeshComponent()->GetClothingSimulationContext()) :
+			nullptr;
+	}
+}  // End namespace Chaos
diff --git a/Engine/Plugins/Experimental/ChaosCloth/Source/Chaos/Private/ChaosCloth/ChaosClothingSimulationSolver.cpp b/Engine/Plugins/Experimental/ChaosCloth/Source/Chaos/Private/ChaosCloth/ChaosClothingSimulationSolver.cpp
--- a/Engine/Plugins/Experimental/ChaosCloth/Source/Chaos/Private/ChaosCloth/ChaosClothingSimulationSolver.cpp
+++ b/Engine/Plugins/Experimental/ChaosCloth/Source/Chaos/Private/ChaosCloth/ChaosClothingSimulationSolver.cpp
@@ -87,6 +87,7 @@ namespace ClothingSimulationSolverConstant
 FClothingSimulationSolver::FClothingSimulationSolver()
 	: OldLocalSpaceLocation(0.)
 	, LocalSpaceLocation(0.)
+	, LocalSpaceRotation(FRotation3::Identity)
 	, VelocityScale(1.)
 	, Time(0.)
 	, DeltaTime(ClothingSimulationSolverConstant::StartDeltaTime)
diff --git a/Engine/Plugins/Experimental/ChaosCloth/Source/Chaos/Public/ChaosCloth/ChaosClothingSimulation.h b/Engine/Plugins/Experimental/ChaosCloth/Source/Chaos/Public/ChaosCloth/ChaosClothingSimulation.h
--- a/Engine/Plugins/Experimental/ChaosCloth/Source/Chaos/Public/ChaosCloth/ChaosClothingSimulation.h
+++ b/Engine/Plugins/Experimental/ChaosCloth/Source/Chaos/Public/ChaosCloth/ChaosClothingSimulation.h
@@ -143,7 +143,9 @@ namespace Chaos
 	private:
 		// Simulation objects
 		TUniquePtr<FClothingSimulationSolver> Solver;  // Default solver
+PRAGMA_DISABLE_DEPRECATION_WARNINGS  // TODO: CHAOS_IS_CLOTHINGSIMULATIONMESH_ABSTRACT
 		TArray<TUniquePtr<FClothingSimulationMesh>> Meshes;
+PRAGMA_ENABLE_DEPRECATION_WARNINGS
 		TArray<TUniquePtr<FClothingSimulationCloth>> Cloths;
 		TArray<TUniquePtr<FClothingSimulationCollider>> Colliders;
 
diff --git a/Engine/Plugins/Experimental/ChaosCloth/Source/Chaos/Public/ChaosCloth/ChaosClothingSimulationCloth.h b/Engine/Plugins/Experimental/ChaosCloth/Source/Chaos/Public/ChaosCloth/ChaosClothingSimulationCloth.h
--- a/Engine/Plugins/Experimental/ChaosCloth/Source/Chaos/Public/ChaosCloth/ChaosClothingSimulationCloth.h
+++ b/Engine/Plugins/Experimental/ChaosCloth/Source/Chaos/Public/ChaosCloth/ChaosClothingSimulationCloth.h
@@ -14,7 +14,7 @@ namespace Chaos
 	class FClothingSimulationCollider;
 
 	// Cloth simulation node
-	class FClothingSimulationCloth final
+	class CHAOSCLOTH_API FClothingSimulationCloth final
 	{
 	public:
 		enum EMassMode
@@ -27,7 +27,9 @@ namespace Chaos
 		typedef FClothConstraints::ETetherMode ETetherMode;
 
 		FClothingSimulationCloth(
+PRAGMA_DISABLE_DEPRECATION_WARNINGS  // TODO: CHAOS_IS_CLOTHINGSIMULATIONMESH_ABSTRACT
 			FClothingSimulationMesh* InMesh,
+PRAGMA_ENABLE_DEPRECATION_WARNINGS
 			TArray<FClothingSimulationCollider*>&& InColliders,
 			uint32 InGroupId,
 			EMassMode InMassMode,
@@ -75,6 +77,11 @@ namespace Chaos
 			int32 InLODIndexOverride);
 		~FClothingSimulationCloth();
 
+		FClothingSimulationCloth(const FClothingSimulationCloth&) = delete;
+		FClothingSimulationCloth(FClothingSimulationCloth&&) = delete;
+		FClothingSimulationCloth& operator=(const FClothingSimulationCloth&) = delete;
+		FClothingSimulationCloth& operator=(FClothingSimulationCloth&&) = delete;
+
 		uint32 GetGroupId() const { return GroupId; }
 		uint32 GetLODIndex(const FClothingSimulationSolver* Solver) const { return LODIndices.FindChecked(Solver); }
 
@@ -101,8 +108,10 @@ namespace Chaos
 		// ---- End of the animatable property setters ----
 
 		// ---- Node property getters/setters
+		PRAGMA_DISABLE_DEPRECATION_WARNINGS  // TODO: CHAOS_IS_CLOTHINGSIMULATIONMESH_ABSTRACT
 		FClothingSimulationMesh* GetMesh() const { return Mesh; }
 		void SetMesh(FClothingSimulationMesh* InMesh);
+		PRAGMA_ENABLE_DEPRECATION_WARNINGS
 
 		const TArray<FClothingSimulationCollider*>& GetColliders() const { return Colliders; }
 		void SetColliders(TArray<FClothingSimulationCollider*>&& InColliders);
@@ -156,46 +165,12 @@ namespace Chaos
 		int32 GetOffset(const FClothingSimulationSolver* Solver, int32 InLODIndex) const;
 
 	private:
-		struct FLODData
-		{
-			// Input mesh
-			const int32 NumParticles;
-			const TConstArrayView<uint32> Indices;
-			const TArray<TConstArrayView<FRealSingle>> WeightMaps;
-			const TArray<TConstArrayView<TTuple<int32, int32, FRealSingle>>> Tethers;
-
-			// Per Solver data
-			struct FSolverData
-			{
-				int32 Offset;
-				FTriangleMesh TriangleMesh;  // TODO: Triangle Mesh shouldn't really be solver dependent (ie not use an offset)
-			};
-			TMap<FClothingSimulationSolver*, FSolverData> SolverData;
-
-			// Stats
-			int32 NumKinenamicParticles;
-			int32 NumDynammicParticles;
-
-			FLODData(
-				int32 InNumParticles,
-				const TConstArrayView<uint32>& InIndices,
-				const TArray<TConstArrayView<FRealSingle>>& InWeightMaps,
-				const TArray<TConstArrayView<TTuple<int32, int32, FRealSingle>>>& InTethers);
-
-			void Add(FClothingSimulationSolver* Solver, FClothingSimulationCloth* Cloth, int32 LODIndex);
-			void Remove(FClothingSimulationSolver* Solver);
-
-			void Update(FClothingSimulationSolver* Solver, FClothingSimulationCloth* Cloth);
-
-			void Enable(FClothingSimulationSolver* Solver, bool bEnable) const;
-
-			void ResetStartPose(FClothingSimulationSolver* Solver) const;
-
-			void UpdateNormals(FClothingSimulationSolver* Solver) const;
-		};
+		struct FLODData;
 
 		// Cloth parameters
+PRAGMA_DISABLE_DEPRECATION_WARNINGS  // TODO: CHAOS_IS_CLOTHINGSIMULATIONMESH_ABSTRACT
 		FClothingSimulationMesh* Mesh;
+PRAGMA_ENABLE_DEPRECATION_WARNINGS
 		TArray<FClothingSimulationCollider*> Colliders;
 		uint32 GroupId;
 		EMassMode MassMode;
@@ -251,7 +226,7 @@ namespace Chaos
 		FRigidTransform3 ReferenceSpaceTransform;  // TODO: Add override in the style of LODIndexOverride
 
 		// LOD data
-		TArray<FLODData> LODData;
+		TArray<TUniquePtr<FLODData>> LODData;
 		TMap<FClothingSimulationSolver*, int32> LODIndices;
 
 		// Stats
diff --git a/Engine/Plugins/Experimental/ChaosCloth/Source/Chaos/Public/ChaosCloth/ChaosClothingSimulationCollider.h b/Engine/Plugins/Experimental/ChaosCloth/Source/Chaos/Public/ChaosCloth/ChaosClothingSimulationCollider.h
--- a/Engine/Plugins/Experimental/ChaosCloth/Source/Chaos/Public/ChaosCloth/ChaosClothingSimulationCollider.h
+++ b/Engine/Plugins/Experimental/ChaosCloth/Source/Chaos/Public/ChaosCloth/ChaosClothingSimulationCollider.h
@@ -18,7 +18,7 @@ namespace Chaos
 	class FLevelSet;
 
 	// Collider simulation node
-	class FClothingSimulationCollider final
+	class CHAOSCLOTH_API FClothingSimulationCollider final
 	{
 	public:
 		enum class ECollisionDataType : int32
@@ -35,7 +35,12 @@ namespace Chaos
 			int32 InLODIndexOverride);
 		~FClothingSimulationCollider();
 
-		int32 GetNumGeometries() const { int32 NumGeometries = 0; for (const FLODData& LODDatum : LODData) { NumGeometries += LODDatum.NumGeometries; } return NumGeometries; }
+		FClothingSimulationCollider(const FClothingSimulationCollider&) = delete;
+		FClothingSimulationCollider(FClothingSimulationCollider&&) = delete;
+		FClothingSimulationCollider& operator=(const FClothingSimulationCollider&) = delete;
+		FClothingSimulationCollider& operator=(FClothingSimulationCollider&&) = delete;
+
+		int32 GetNumGeometries() const;
 
 		// Return source (untransformed) collision data for LODless, external and active LODs.
 		FClothCollisionData GetCollisionData(const FClothingSimulationSolver* Solver, const FClothingSimulationCloth* Cloth) const;
@@ -72,7 +77,6 @@ namespace Chaos
 		// ---- End of the debugging and visualization functions ----
 
 	private:
-
 		struct FLevelSetCollisionData
 		{
 			const TSharedPtr<Chaos::FLevelSet, ESPMode::ThreadSafe> LevelSet;
@@ -93,34 +97,7 @@ namespace Chaos
 	private:
 		typedef TPair<const FClothingSimulationSolver*, const FClothingSimulationCloth*> FSolverClothPair;
 
-		struct FLODData
-		{
-			FClothCollisionData ClothCollisionData;
-			int32 NumGeometries;  // Number of collision bodies
-			TMap<FSolverClothPair, int32> Offsets;  // Solver particle offset
-
-			FLODData() : NumGeometries(0) {}
-
-			void Add(
-				FClothingSimulationSolver* Solver,
-				FClothingSimulationCloth* Cloth,
-				const FClothCollisionData& InClothCollisionData,
-				const TArray<FLevelSetCollisionData>& InLevelSetCollisionData,
-				const FReal InScale = 1.f,
-				const TArray<int32>& UsedBoneIndices = TArray<int32>());
-			void Remove(FClothingSimulationSolver* Solver, FClothingSimulationCloth* Cloth);
-
-			void Update(FClothingSimulationSolver* Solver, FClothingSimulationCloth* Cloth, const FClothingSimulationContextCommon* Context);
-
-			void Enable(FClothingSimulationSolver* Solver, FClothingSimulationCloth* Cloth, bool bEnable);
-
-			void ResetStartPose(FClothingSimulationSolver* Solver, FClothingSimulationCloth* Cloth);
-
-			FORCEINLINE static int32 GetMappedBoneIndex(const TArray<int32>& UsedBoneIndices, int32 BoneIndex)
-			{
-				return UsedBoneIndices.IsValidIndex(BoneIndex) ? UsedBoneIndices[BoneIndex] : INDEX_NONE;
-			}
-		};
+		struct FLODData;
 
 		const UClothingAssetCommon* Asset;
 		const USkeletalMeshComponent* SkeletalMeshComponent;
@@ -130,7 +107,7 @@ namespace Chaos
 		bool bHasExternalCollisionChanged;
 
 		// Collision primitives
-		TArray<FLODData> LODData;  // Actual LODs start at LODStart
+		TArray<TUniquePtr<FLODData>> LODData;  // Actual LODs start at ECollisionDataType::LODs
 		TMap<FSolverClothPair, int32> LODIndices;
 
 		// Initial scale
diff --git a/Engine/Plugins/Experimental/ChaosCloth/Source/Chaos/Public/ChaosCloth/ChaosClothingSimulationMesh.h b/Engine/Plugins/Experimental/ChaosCloth/Source/Chaos/Public/ChaosCloth/ChaosClothingSimulationMesh.h
--- a/Engine/Plugins/Experimental/ChaosCloth/Source/Chaos/Public/ChaosCloth/ChaosClothingSimulationMesh.h
+++ b/Engine/Plugins/Experimental/ChaosCloth/Source/Chaos/Public/ChaosCloth/ChaosClothingSimulationMesh.h
@@ -4,47 +4,130 @@
 #include "Chaos/Core.h"
 #include "Chaos/PBDSoftsEvolutionFwd.h"
 #include "Containers/ContainersFwd.h"
+#include "ClothVertBoneData.h"
 
+//#define CHAOS_IS_CLOTHINGSIMULATIONMESH_ABSTRACT 1  // For testing only. TODO: Once FClothingSimulationMesh is abstract from 5.4, remove all lines where this is referenced
+
+class USkinnedMeshComponent;
+struct FMeshToMeshVertData;
+struct FClothVertBoneData;
+
+#if !defined(CHAOS_IS_CLOTHINGSIMULATIONMESH_ABSTRACT) || !CHAOS_IS_CLOTHINGSIMULATIONMESH_ABSTRACT
 class USkeletalMeshComponent;
 class UClothingAssetCommon;
+class FClothingSimulationContextCommon;
+#endif
 
 namespace Chaos
 {
 	class FClothingSimulationSolver;
 
 	// Mesh simulation node
-	class FClothingSimulationMesh final
+#if !defined(CHAOS_IS_CLOTHINGSIMULATIONMESH_ABSTRACT) || !CHAOS_IS_CLOTHINGSIMULATIONMESH_ABSTRACT
+	class UE_DEPRECATED(5.2, "This class is becoming an abstract interface, do not use directly. Use FClothingSimulationSkeletalMesh instead.")
+		CHAOSCLOTH_API FClothingSimulationMesh
+#else
+	class CHAOSCLOTH_API FClothingSimulationMesh
+#endif
 	{
 	public:
-		FClothingSimulationMesh(const UClothingAssetCommon* InAsset, const USkeletalMeshComponent* InSkeletalMeshComponent);
-		~FClothingSimulationMesh();
+		FClothingSimulationMesh(const USkinnedMeshComponent* InSkinnedMeshComponent);
+		virtual ~FClothingSimulationMesh();
+
+		FClothingSimulationMesh(const FClothingSimulationMesh&) = delete;
+		FClothingSimulationMesh(FClothingSimulationMesh&&) = delete;
+		FClothingSimulationMesh& operator=(const FClothingSimulationMesh&) = delete;
+		FClothingSimulationMesh& operator=(FClothingSimulationMesh&&) = delete;
 
 		// ---- Node property getters/setters
-		const UClothingAssetCommon* GetAsset() const { return Asset; }
-		const USkeletalMeshComponent* GetSkeletalMeshComponent() const { return SkeletalMeshComponent; }
+		const USkinnedMeshComponent* GetSkinnedMeshComponent() const { return SkinnedMeshComponent; }
 		// ---- End of node property getters/setters
+#if !defined(CHAOS_IS_CLOTHINGSIMULATIONMESH_ABSTRACT) || !CHAOS_IS_CLOTHINGSIMULATIONMESH_ABSTRACT
+		UE_DEPRECATED(5.2, "Use the FClothingSimulationSkeletalMesh constructor instead.")
+		FClothingSimulationMesh(const UClothingAssetCommon* InAsset, const USkeletalMeshComponent* InSkeletalMeshComponent);
+		UE_DEPRECATED(5.2, "Use the FClothingSimulationSkeletalMesh::GetAsset() instead.")
+		virtual const UClothingAssetCommon* GetAsset() const { return Asset; }
+		UE_DEPRECATED(5.2, "Use the FClothingSimulationSkeletalMesh::GetSkeletalMeshComponent() instead.")
+		virtual const USkeletalMeshComponent* GetSkeletalMeshComponent() const;
 
+		virtual int32 GetNumLODs() const;
+		virtual int32 GetLODIndex() const;
+		virtual int32 GetOwnerLODIndex(int32 LODIndex) const;
+		virtual bool IsValidLODIndex(int32 LODIndex) const;
+		virtual int32 GetNumPoints(int32 LODIndex) const;
+		virtual TConstArrayView<FVector3f> GetPositions(int32 LODIndex) const;
+		virtual TConstArrayView<FVector3f> GetNormals(int32 LODIndex) const;
+		virtual TConstArrayView<uint32> GetIndices(int32 LODIndex) const;
+		virtual TArray<TConstArrayView<FRealSingle>> GetWeightMaps(int32 LODIndex) const;
+		virtual TArray<TConstArrayView<TTuple<int32, int32, float>>> GetTethers(int32 LODIndex, bool bUseGeodesicTethers) const;
+		virtual int32 GetReferenceBoneIndex() const;
+		virtual FRigidTransform3 GetReferenceBoneTransform() const;
+		virtual const FTransform& GetComponentToWorldTransform() const;
+		virtual const TArray<FMatrix44f>& GetRefToLocalMatrices() const;
+		virtual TConstArrayView<int32> GetBoneMap() const;
+		virtual TConstArrayView<FClothVertBoneData> GetBoneData(int32 LODIndex) const;
+		virtual const TArray<FMeshToMeshVertData>& GetTransitionUpSkinData(int32 LODIndex) const;
+		virtual const TArray<FMeshToMeshVertData>& GetTransitionDownSkinData(int32 LODIndex) const;
+#else
 		// ---- Cloth interface ----
-		void Update(
-			FClothingSimulationSolver* Solver,
-			int32 PrevLODIndex,
-			int32 LODIndex,
-			int32 PrevOffset,
-			int32 Offset);
+		/* Return the number of LODs on this mesh. */
+		virtual int32 GetNumLODs() const = 0;
+
+		/* Return the cloth mesh LOD Index. The returned value can then be used to switch LODs on the simulation object. */
+		virtual int32 GetLODIndex() const = 0;
+
+		/* Return the owner component LOD Index from the specified cloth mesh LOD index, or 0 if the owner LOD cannot be determined. The mapping between the cloth mesh LOD and the owner component LOD is not necessarily one to one. */
+		virtual int32 GetOwnerLODIndex(int32 LODIndex) const = 0;
+
+		/* Return whether the specified LOD index is valid. */
+		virtual bool IsValidLODIndex(int32 LODIndex) const = 0;
+
+		/* Return the number of points for the specified LOD, or 0 if the LOD is empty or invalid. */
+		virtual int32 GetNumPoints(int32 LODIndex) const = 0;
+
+		/* Return the source mesh positions (pre-skinning). */
+		virtual TConstArrayView<FVector3f> GetPositions(int32 LODIndex) const = 0;
 
-		// Return the LOD Index specified by the input SkeletalMeshComponent. Note that this is not the same as the current LOD index, and that Mesh LOD changes are always driven by the Cloth.
-		int32 GetLODIndex() const;
-		int32 GetNumLODs() const;
-		int32 GetNumPoints(int32 LODIndex) const;
-		TConstArrayView<uint32> GetIndices(int32 LODIndex) const;
-		TArray<TConstArrayView<FRealSingle>> GetWeightMaps(int32 LODIndex) const;
-		TArray<TConstArrayView<TTuple<int32, int32, float>>> GetTethers(int32 LODIndex, bool bUseGeodesicTethers) const;
-		int32 GetReferenceBoneIndex() const;
-		FRigidTransform3 GetReferenceBoneTransform() const;
+		/* Return the source mesh normals (pre-skinning). */
+		virtual TConstArrayView<FVector3f> GetNormals(int32 LODIndex) const = 0;
 
-		// Return this mesh component's scale (the max of the three axis scale values)
-		Softs::FSolverReal GetScale() const;
+		/* Return the specified LOD's triangle indices for this mesh. */
+		virtual TConstArrayView<uint32> GetIndices(int32 LODIndex) const = 0;
 
+		/* Return the specified LOD's weightmap. */
+		virtual TArray<TConstArrayView<FRealSingle>> GetWeightMaps(int32 LODIndex) const = 0;
+
+		/* Return the tethers connections for the long range attachment into convenient parallel friendly batches. */
+		virtual TArray<TConstArrayView<TTuple<int32, int32, float>>> GetTethers(int32 LODIndex, bool bUseGeodesicTethers) const = 0;
+
+		/* Return the bone to treat as the root of the simulation space. */
+		virtual int32 GetReferenceBoneIndex() const = 0;
+
+		/* Return the transform of the bone treated as the root of the simulation space. */
+		virtual FRigidTransform3 GetReferenceBoneTransform() const = 0;
+
+		/* Return the transform of the bone treated as the root of the simulation space. */
+		virtual const FTransform& GetComponentToWorldTransform() const = 0;
+
+		/* Return the skinning matrices. */
+		virtual const TArray<FMatrix44f>& GetRefToLocalMatrices() const = 0;
+
+		/* Return the bone map used to remap the used bones index into the correct skinning matrix index. */
+		virtual TConstArrayView<int32> GetBoneMap() const = 0;
+
+		/* Return the bone data containing bone weights and influences. */
+		virtual TConstArrayView<FClothVertBoneData> GetBoneData(int32 LODIndex) const = 0;
+
+		/* Return the transition up data (PrevLODIndex < LODIndex), for matching shapes during LOD changes. */
+		virtual const TArray<FMeshToMeshVertData>& GetTransitionUpSkinData(int32 LODIndex) const = 0;
+
+		/* Return the transition down data (PrevLODIndex > LODIndex), for matching shapes during LOD changes. */
+		virtual const TArray<FMeshToMeshVertData>& GetTransitionDownSkinData(int32 LODIndex) const = 0;
+#endif
+		/** Return this mesh uniform scale as the maximum of the three axis scale value. */
+		virtual Softs::FSolverReal GetScale() const;
+
+		/** Deform the specified positions to match the shape of the previous LOD. */
 		bool WrapDeformLOD(
 			int32 PrevLODIndex,
 			int32 LODIndex,
@@ -52,6 +135,7 @@ namespace Chaos
 			const Softs::FSolverVec3* Positions,
 			Softs::FSolverVec3* OutPositions) const;
 
+		/** Deform the specified positions and transfer velocities to match the dynamics of the previous LOD. */
 		bool WrapDeformLOD(
 			int32 PrevLODIndex,
 			int32 LODIndex,
@@ -61,6 +145,14 @@ namespace Chaos
 			Softs::FPAndInvM* OutPositionAndInvMs0,
 			Softs::FSolverVec3* OutPositions1,
 			Softs::FSolverVec3* OutVelocities) const;
+
+		/* Update the mesh for the next solver step, doing skinning and matching the shapes during LOD changes. */
+		void Update(
+			FClothingSimulationSolver* Solver,
+			int32 PrevLODIndex,
+			int32 LODIndex,
+			int32 PrevOffset,
+			int32 Offset);
 		// ---- End of the Cloth interface ----
 
 	private:
@@ -71,8 +163,12 @@ namespace Chaos
 			Softs::FSolverVec3* OutNormals) const;
 
 	private:
+#if !defined(CHAOS_IS_CLOTHINGSIMULATIONMESH_ABSTRACT) || !CHAOS_IS_CLOTHINGSIMULATIONMESH_ABSTRACT
+		const FClothingSimulationContextCommon* GetContext() const;
+
 		const UClothingAssetCommon* Asset;
-		const USkeletalMeshComponent* SkeletalMeshComponent;
+#endif
+		const USkinnedMeshComponent* SkinnedMeshComponent;
 	};
 } // namespace Chaos
 
diff --git a/Engine/Plugins/Experimental/ChaosCloth/Source/Chaos/Public/ChaosCloth/ChaosClothingSimulationSkeletalMesh.h b/Engine/Plugins/Experimental/ChaosCloth/Source/Chaos/Public/ChaosCloth/ChaosClothingSimulationSkeletalMesh.h
new file mode 100644
--- /dev/null
+++ b/Engine/Plugins/Experimental/ChaosCloth/Source/Chaos/Public/ChaosCloth/ChaosClothingSimulationSkeletalMesh.h
@@ -0,0 +1,61 @@
+// Copyright Epic Games, Inc. All Rights Reserved.
+#pragma once
+
+#include "ChaosCloth/ChaosClothingSimulationMesh.h"
+
+class FClothingSimulationContextCommon;
+class UClothingAssetCommon;
+
+namespace Chaos
+{
+PRAGMA_DISABLE_DEPRECATION_WARNINGS  // TODO: CHAOS_IS_CLOTHINGSIMULATIONMESH_ABSTRACT
+	class FClothingSimulationSkeletalMesh final : public FClothingSimulationMesh
+PRAGMA_ENABLE_DEPRECATION_WARNINGS
+	{
+	public:
+		FClothingSimulationSkeletalMesh(const UClothingAssetCommon* InAsset, const USkeletalMeshComponent* InSkeletalMeshComponent);
+PRAGMA_DISABLE_DEPRECATION_WARNINGS  // TODO: CHAOS_IS_CLOTHINGSIMULATIONMESH_ABSTRACT
+		virtual ~FClothingSimulationSkeletalMesh() override = default;
+PRAGMA_ENABLE_DEPRECATION_WARNINGS
+
+		FClothingSimulationSkeletalMesh(const FClothingSimulationSkeletalMesh&) = delete;
+		FClothingSimulationSkeletalMesh(FClothingSimulationSkeletalMesh&&) = delete;
+		FClothingSimulationSkeletalMesh& operator=(const FClothingSimulationSkeletalMesh&) = delete;
+		FClothingSimulationSkeletalMesh& operator=(FClothingSimulationSkeletalMesh&&) = delete;
+
+#if !defined(CHAOS_IS_CLOTHINGSIMULATIONMESH_ABSTRACT) || !CHAOS_IS_CLOTHINGSIMULATIONMESH_ABSTRACT
+		virtual const UClothingAssetCommon* GetAsset() const override { return Asset; }
+PRAGMA_DISABLE_DEPRECATION_WARNINGS  // TODO: CHAOS_IS_CLOTHINGSIMULATIONMESH_ABSTRACT
+		virtual const USkeletalMeshComponent* GetSkeletalMeshComponent() const override { return Cast<USkeletalMeshComponent>(GetSkinnedMeshComponent()); }
+PRAGMA_ENABLE_DEPRECATION_WARNINGS
+#else
+		const UClothingAssetCommon* GetAsset() const { return Asset; }
+		const USkeletalMeshComponent* GetSkeletalMeshComponent() const { return Cast<USkeletalMeshComponent>(GetSkinnedMeshComponent()); }
+#endif
+		//~ Begin FClothingSimulationMesh Interface
+		virtual int32 GetNumLODs() const override;
+		virtual int32 GetLODIndex() const override;
+		virtual int32 GetOwnerLODIndex(int32 LODIndex) const override;
+		virtual bool IsValidLODIndex(int32 LODIndex) const override;
+		virtual int32 GetNumPoints(int32 LODIndex) const override;
+		virtual TConstArrayView<FVector3f> GetPositions(int32 LODIndex) const override;
+		virtual TConstArrayView<FVector3f> GetNormals(int32 LODIndex) const override;
+		virtual TConstArrayView<uint32> GetIndices(int32 LODIndex) const override;
+		virtual TArray<TConstArrayView<FRealSingle>> GetWeightMaps(int32 LODIndex) const override;
+		virtual TArray<TConstArrayView<TTuple<int32, int32, float>>> GetTethers(int32 LODIndex, bool bUseGeodesicTethers) const override;
+		virtual int32 GetReferenceBoneIndex() const override;
+		virtual FRigidTransform3 GetReferenceBoneTransform() const override;
+		virtual const FTransform& GetComponentToWorldTransform() const override;
+		virtual const TArray<FMatrix44f>& GetRefToLocalMatrices() const override;
+		virtual TConstArrayView<int32> GetBoneMap() const override;
+		virtual TConstArrayView<FClothVertBoneData> GetBoneData(int32 LODIndex) const override;
+		virtual const TArray<FMeshToMeshVertData>& GetTransitionUpSkinData(int32 LODIndex) const override;
+		virtual const TArray<FMeshToMeshVertData>& GetTransitionDownSkinData(int32 LODIndex) const override;
+		//~ End FClothingSimulationMesh Interface
+
+	private:
+		const FClothingSimulationContextCommon* GetContext() const;
+
+		const UClothingAssetCommon* Asset;
+	};
+}
diff --git a/Engine/Plugins/Experimental/ChaosCloth/Source/Chaos/Public/ChaosCloth/ChaosClothingSimulationSolver.h b/Engine/Plugins/Experimental/ChaosCloth/Source/Chaos/Public/ChaosCloth/ChaosClothingSimulationSolver.h
--- a/Engine/Plugins/Experimental/ChaosCloth/Source/Chaos/Public/ChaosCloth/ChaosClothingSimulationSolver.h
+++ b/Engine/Plugins/Experimental/ChaosCloth/Source/Chaos/Public/ChaosCloth/ChaosClothingSimulationSolver.h
@@ -27,6 +27,10 @@ namespace Chaos
 		FClothingSimulationSolver();
 		~FClothingSimulationSolver();
 		
+		FClothingSimulationSolver(const FClothingSimulationSolver&) = delete;
+		FClothingSimulationSolver(FClothingSimulationSolver&&) = delete;
+		FClothingSimulationSolver& operator=(const FClothingSimulationSolver&) = delete;
+		FClothingSimulationSolver& operator=(FClothingSimulationSolver&&) = delete;
 
 		// ---- Animatable property setters ----
 		void SetLocalSpaceLocation(const FVec3& InLocalSpaceLocation, bool bReset = false);
