Chaos Cloth - Fixed fictitious forces to use the cloth reference bone as rotation center instead of the simulation space location.

diff --git a/Engine/Plugins/Experimental/ChaosCloth/Source/Chaos/Private/ChaosCloth/ChaosClothingSimulationSolver.cpp b/Engine/Plugins/Experimental/ChaosCloth/Source/Chaos/Private/ChaosCloth/ChaosClothingSimulationSolver.cpp
--- a/Engine/Plugins/Experimental/ChaosCloth/Source/Chaos/Private/ChaosCloth/ChaosClothingSimulationSolver.cpp
+++ b/Engine/Plugins/Experimental/ChaosCloth/Source/Chaos/Private/ChaosCloth/ChaosClothingSimulationSolver.cpp
@@ -119,6 +119,7 @@ FClothingSimulationSolver::FClothingSimulationSolver()
 	// Add simulation groups arrays
 	Evolution->AddArray(&PreSimulationTransforms);
 	Evolution->AddArray(&FictitiousAngularDisplacements);
+	Evolution->AddArray(&ReferenceSpaceLocations);
 
 	Evolution->Particles().AddArray(&Normals);
 	Evolution->Particles().AddArray(&OldAnimationPositions);
@@ -608,6 +609,7 @@ void FClothingSimulationSolver::SetReferenceVelocityScale(
 	// Save the reference bone relative angular velocity for calculating the fictitious forces
 	const FVec3 FictitiousAngularDisplacement = ReferenceSpaceTransform.TransformVector(Axis * PartialDeltaAngle * FMath::Min((FReal)2., (FReal)FictitiousAngularScale));  // Clamp to 2x the delta angle
 	FictitiousAngularDisplacements[GroupId] = Softs::FSolverVec3(FictitiousAngularDisplacement);
+	ReferenceSpaceLocations[GroupId] = ReferenceSpaceTransform.GetLocation();
 }
 
 Softs::FSolverReal FClothingSimulationSolver::SetParticleMassPerArea(int32 Offset, int32 Size, const FTriangleMesh& Mesh)
@@ -707,7 +709,7 @@ void FClothingSimulationSolver::SetWindAndPressureProperties(uint32 GroupId, con
 	VelocityField.SetProperties(Drag, Lift, AirDensity, Pressure / ClothingSimulationSolverConstant::WorldScale);
 }
 
-const Softs::FVelocityAndPressureField& FClothingSimulationSolver::GetWindVelocityAndPressureField(uint32 GroupId)
+const Softs::FVelocityAndPressureField& FClothingSimulationSolver::GetWindVelocityAndPressureField(uint32 GroupId) const
 {
 	return Evolution->GetVelocityAndPressureField(GroupId);
 }
@@ -717,13 +719,14 @@ void FClothingSimulationSolver::AddExternalForces(uint32 GroupId, bool bUseLegac
 	if (Evolution)
 	{
 		const FVec3& AngularDisplacement = FictitiousAngularDisplacements[GroupId];
+		const FVec3& ReferenceSpaceLocation = ReferenceSpaceLocations[GroupId];
 		const bool bHasFictitiousForces = !AngularDisplacement.IsNearlyZero();
 
 		static const FReal LegacyWindMultiplier = (FReal)25.;
 		const FVec3 LegacyWindVelocity = WindVelocity * LegacyWindMultiplier;
 
 		Evolution->GetForceFunction(GroupId) =
-			[this, bHasFictitiousForces, bUseLegacyWind, LegacyWindVelocity, AngularDisplacement](Softs::FSolverParticles& Particles, const FReal Dt, const int32 Index)
+			[this, bHasFictitiousForces, bUseLegacyWind, LegacyWindVelocity, AngularDisplacement, ReferenceSpaceLocation](Softs::FSolverParticles& Particles, const FReal Dt, const int32 Index)
 			{
 				FVec3 Forces((FReal)0.);
 
@@ -741,7 +744,7 @@ void FClothingSimulationSolver::AddExternalForces(uint32 GroupId, bool bUseLegac
 
 				if (bHasFictitiousForces)
 				{
-					const FVec3& X = Particles.X(Index);
+					const FVec3 X = Particles.X(Index) - ReferenceSpaceLocation;
 					const FVec3 W = AngularDisplacement / Dt;
 					const FReal& M = Particles.M(Index);
 #if 0
diff --git a/Engine/Plugins/Experimental/ChaosCloth/Source/Chaos/Public/ChaosCloth/ChaosClothingSimulationSolver.h b/Engine/Plugins/Experimental/ChaosCloth/Source/Chaos/Public/ChaosCloth/ChaosClothingSimulationSolver.h
--- a/Engine/Plugins/Experimental/ChaosCloth/Source/Chaos/Public/ChaosCloth/ChaosClothingSimulationSolver.h
+++ b/Engine/Plugins/Experimental/ChaosCloth/Source/Chaos/Public/ChaosCloth/ChaosClothingSimulationSolver.h
@@ -139,7 +139,7 @@ namespace Chaos
 		}
 
 		// Return the wind velocity and pressure field associated with a given group id.
-		const Softs::FVelocityAndPressureField& GetWindVelocityAndPressureField(uint32 GroupId);
+		const Softs::FVelocityAndPressureField& GetWindVelocityAndPressureField(uint32 GroupId) const;
 		UE_DEPRECATED(5.1, "Chaos::Softs::FVelocityField has been renamed FVelocityAndPressureField to match its new behavior.")
 		PRAGMA_DISABLE_DEPRECATION_WARNINGS
 		const Softs::FVelocityField& GetWindVelocityField(uint32 GroupId) { return GetWindVelocityAndPressureField(GroupId); }
@@ -218,6 +218,7 @@ namespace Chaos
 		// Simulation group attributes
 		TArrayCollectionArray<Softs::FSolverRigidTransform3> PreSimulationTransforms;  // Allow a different frame of reference for each cloth groups
 		TArrayCollectionArray<Softs::FSolverVec3> FictitiousAngularDisplacements;  // Relative angular displacement of the reference bone that depends on the fictitious angular scale factor
+		TArrayCollectionArray<Softs::FSolverVec3> ReferenceSpaceLocations;  // Center of rotations for fictitious forces in local coordinate to the simulation space location
 
 		// Particle attributes
 		TArrayCollectionArray<Softs::FSolverVec3> Normals;
