Chaos Cloth - Move visualization code to its own files so that it can be reused by other simulations.

diff --git a/Engine/Plugins/Experimental/ChaosCloth/Source/Chaos/Private/ChaosCloth/ChaosClothVisualization.cpp b/Engine/Plugins/Experimental/ChaosCloth/Source/Chaos/Private/ChaosCloth/ChaosClothVisualization.cpp
new file mode 100644
--- /dev/null
+++ b/Engine/Plugins/Experimental/ChaosCloth/Source/Chaos/Private/ChaosCloth/ChaosClothVisualization.cpp
@@ -0,0 +1,1560 @@
+// Copyright Epic Games, Inc. All Rights Reserved.
+#include "ChaosCloth/ChaosClothVisualization.h"
+#if CHAOS_DEBUG_DRAW
+#include "ChaosCloth/ChaosClothingSimulationSolver.h"
+#include "ChaosCloth/ChaosClothingSimulationMesh.h"
+#include "ChaosCloth/ChaosClothingSimulationCloth.h"
+#include "ChaosCloth/ChaosClothingSimulationCollider.h"
+#include "ChaosCloth/ChaosWeightMapTarget.h"
+#include "Chaos/DebugDrawQueue.h"
+#include "Chaos/Capsule.h"
+#include "Chaos/Convex.h"
+#include "Chaos/Sphere.h"
+#include "Chaos/TaperedCapsule.h"
+#include "Chaos/TaperedCylinder.h"
+#include "Chaos/Triangle.h"
+#include "Chaos/VelocityField.h"
+#include "Chaos/PBDAnimDriveConstraint.h"
+#include "Chaos/PBDBendingConstraints.h"
+#include "Chaos/PBDCollisionSpringConstraints.h"
+#include "Chaos/PBDLongRangeConstraints.h"
+#include "Chaos/PBDSphericalConstraint.h"
+#include "Chaos/PBDSpringConstraints.h"
+#include "Chaos/PBDTriangleMeshCollisions.h"
+#include "Chaos/XPBDBendingConstraints.h"
+#include "Chaos/XPBDSpringConstraints.h"
+#include "DynamicMeshBuilder.h"
+#if WITH_EDITOR
+#include "Materials/Material.h"
+#include "Engine/Canvas.h"  // For draw text
+#include "CanvasItem.h"     //
+#include "Engine/Engine.h"  //
+#endif  // #if WITH_EDITOR
+
+namespace UE::Chaos::Cloth
+{
+	FVisualization::FVisualization(const ::Chaos::FClothingSimulationSolver* InSolver)
+		: Solver(InSolver)
+	{
+#if WITH_EDITOR
+		ClothMaterial = LoadObject<UMaterial>(nullptr, TEXT("/Engine/EditorMaterials/Cloth/CameraLitDoubleSided.CameraLitDoubleSided"), nullptr, LOAD_None, nullptr);  // LOAD_EditorOnly
+		ClothMaterialVertex = LoadObject<UMaterial>(nullptr, TEXT("/Engine/EditorMaterials/WidgetVertexColorMaterial"), nullptr, LOAD_None, nullptr);  // LOAD_EditorOnly
+#endif  // #if WITH_EDITOR
+	}
+
+	FVisualization::~FVisualization() = default;
+
+	void FVisualization::SetSolver(const ::Chaos::FClothingSimulationSolver* InSolver)
+	{
+		Solver = InSolver;
+	}
+
+#if WITH_EDITOR
+	void FVisualization::AddReferencedObjects(FReferenceCollector& Collector)
+	{
+		Collector.AddReferencedObject(ClothMaterial);
+		Collector.AddReferencedObject(ClothMaterialVertex);
+	}
+
+	void FVisualization::DrawPhysMeshShaded(FPrimitiveDrawInterface* PDI) const
+	{
+		using namespace ::Chaos;
+
+		if (!Solver || !ClothMaterial)
+		{
+			return;
+		}
+
+		FDynamicMeshBuilder MeshBuilder(PDI->View->GetFeatureLevel());
+		int32 VertexIndex = 0;
+
+		for (const FClothingSimulationCloth* const Cloth : Solver->GetCloths())
+		{
+			const int32 Offset = Cloth->GetOffset(Solver);
+			if (Offset == INDEX_NONE)
+			{
+				continue;
+			}
+
+			const TConstArrayView<TVec3<int32>> Elements = Cloth->GetTriangleMesh(Solver).GetElements();
+			const TConstArrayView<Softs::FSolverVec3> Positions = Cloth->GetParticlePositions(Solver);
+			const TConstArrayView<Softs::FSolverReal> InvMasses = Cloth->GetParticleInvMasses(Solver);
+			check(InvMasses.Num() == Positions.Num());
+
+			for (int32 ElementIndex = 0; ElementIndex < Elements.Num(); ++ElementIndex, VertexIndex += 3)
+			{
+				const TVec3<int32>& Element = Elements[ElementIndex];
+
+				const FVector3f Pos0(Positions[Element.X - Offset]); // TODO: Triangle Mesh shouldn't really be solver dependent (ie not use an offset)
+				const FVector3f Pos1(Positions[Element.Y - Offset]);
+				const FVector3f Pos2(Positions[Element.Z - Offset]);
+
+				const FVector3f Normal = FVector3f::CrossProduct(Pos2 - Pos0, Pos1 - Pos0).GetSafeNormal();
+				const FVector3f Tangent = ((Pos1 + Pos2) * 0.5f - Pos0).GetSafeNormal();
+
+				const bool bIsKinematic0 = (InvMasses[Element.X - Offset] == (Softs::FSolverReal)0.);
+				const bool bIsKinematic1 = (InvMasses[Element.Y - Offset] == (Softs::FSolverReal)0.);
+				const bool bIsKinematic2 = (InvMasses[Element.Z - Offset] == (Softs::FSolverReal)0.);
+
+				MeshBuilder.AddVertex(FDynamicMeshVertex(Pos0, Tangent, Normal, FVector2f(0.f, 0.f), bIsKinematic0 ? FColor::Purple : FColor::White));
+				MeshBuilder.AddVertex(FDynamicMeshVertex(Pos1, Tangent, Normal, FVector2f(0.f, 1.f), bIsKinematic1 ? FColor::Purple : FColor::White));
+				MeshBuilder.AddVertex(FDynamicMeshVertex(Pos2, Tangent, Normal, FVector2f(1.f, 1.f), bIsKinematic2 ? FColor::Purple : FColor::White));
+				MeshBuilder.AddTriangle(VertexIndex, VertexIndex + 1, VertexIndex + 2);
+			}
+		}
+
+		FMatrix LocalSimSpaceToWorld(FMatrix::Identity);
+		LocalSimSpaceToWorld.SetOrigin(Solver->GetLocalSpaceLocation());
+		MeshBuilder.Draw(PDI, LocalSimSpaceToWorld, ClothMaterial->GetRenderProxy(), SDPG_World, false, false);
+	}
+
+	static void DrawText(FCanvas* Canvas, const FSceneView* SceneView, const FVector& Pos, const FText& Text, const FLinearColor& Color)
+	{
+		FVector2D PixelLocation;
+		if (SceneView->WorldToPixel(Pos, PixelLocation))
+		{
+			FCanvasTextItem TextItem(PixelLocation, Text, GEngine->GetSmallFont(), Color);
+			TextItem.Scale = FVector2D::UnitVector;
+			TextItem.EnableShadow(FLinearColor::Black);
+			TextItem.Draw(Canvas);
+		}
+	}
+
+	void FVisualization::DrawParticleIndices(FCanvas* Canvas, const FSceneView* SceneView) const
+	{
+		using namespace ::Chaos;
+
+		if (!Solver)
+		{
+			return;
+		}
+
+		static const FLinearColor DynamicColor = FColor::White;
+		static const FLinearColor KinematicColor = FColor::Purple;
+
+		const FVec3& LocalSpaceLocation = Solver->GetLocalSpaceLocation();
+
+		for (const FClothingSimulationCloth* const Cloth : Solver->GetCloths())
+		{
+			const int32 Offset = Cloth->GetOffset(Solver);
+			if (Offset == INDEX_NONE)
+			{
+				continue;
+			}
+
+			const TConstArrayView<Softs::FSolverVec3> Positions = Cloth->GetParticlePositions(Solver);
+			const TConstArrayView<Softs::FSolverReal> InvMasses = Cloth->GetParticleInvMasses(Solver);
+			check(InvMasses.Num() == Positions.Num());
+
+			for (int32 Index = 0; Index < Positions.Num(); ++Index)
+			{
+				const FVector Position = LocalSpaceLocation + FVector(Positions[Index]);
+
+				const FText Text = FText::AsNumber(Offset + Index);
+				DrawText(Canvas, SceneView, Position, Text, InvMasses[Index] == (Softs::FSolverReal)0. ? KinematicColor : DynamicColor);
+			}
+		}
+	}
+
+	void FVisualization::DrawElementIndices(FCanvas* Canvas, const FSceneView* SceneView) const
+	{
+		using namespace ::Chaos;
+
+		if (!Solver)
+		{
+			return;
+		}
+
+		static const FLinearColor DynamicColor = FColor::White;
+		static const FLinearColor KinematicColor = FColor::Purple;
+
+		const FVec3& LocalSpaceLocation = Solver->GetLocalSpaceLocation();
+
+		for (const FClothingSimulationCloth* const Cloth : Solver->GetCloths())
+		{
+			const int32 Offset = Cloth->GetOffset(Solver);
+			if (Offset == INDEX_NONE)
+			{
+				continue;
+			}
+
+			const TArray<TVec3<int32>>& Elements = Cloth->GetTriangleMesh(Solver).GetElements();
+			const TConstArrayView<Softs::FSolverVec3> Positions = Cloth->GetParticlePositions(Solver);
+			const TConstArrayView<Softs::FSolverReal> InvMasses = Cloth->GetParticleInvMasses(Solver);
+			check(InvMasses.Num() == Positions.Num());
+
+			for (int32 Index = 0; Index < Elements.Num(); ++Index)
+			{
+				const TVec3<int32>& Element = Elements[Index];
+				const FVector Position = LocalSpaceLocation + FVector(Positions[Element[0]] + Positions[Element[1]] + Positions[Element[2]]) / (FReal)3.;
+
+				const bool bIsKinematic0 = (InvMasses[Element.X - Offset] == (Softs::FSolverReal)0.);
+				const bool bIsKinematic1 = (InvMasses[Element.Y - Offset] == (Softs::FSolverReal)0.);
+				const bool bIsKinematic2 = (InvMasses[Element.Z - Offset] == (Softs::FSolverReal)0.);
+				const FLinearColor& Color = (bIsKinematic0 && bIsKinematic1 && bIsKinematic2) ? KinematicColor : DynamicColor;
+				const FText Text = FText::AsNumber(Index);
+				DrawText(Canvas, SceneView, Position, Text, Color);
+			}
+		}
+	}
+
+	void FVisualization::DrawMaxDistanceValues(FCanvas* Canvas, const FSceneView* SceneView) const
+	{
+		using namespace ::Chaos;
+
+		if (!Solver)
+		{
+			return;
+		}
+
+		static const FLinearColor DynamicColor = FColor::White;
+		static const FLinearColor KinematicColor = FColor::Purple;
+
+		FNumberFormattingOptions NumberFormattingOptions;
+		NumberFormattingOptions.AlwaysSign = false;
+		NumberFormattingOptions.UseGrouping = false;
+		NumberFormattingOptions.RoundingMode = ERoundingMode::HalfFromZero;
+		NumberFormattingOptions.MinimumIntegralDigits = 1;
+		NumberFormattingOptions.MaximumIntegralDigits = 6;
+		NumberFormattingOptions.MinimumFractionalDigits = 2;
+		NumberFormattingOptions.MaximumFractionalDigits = 2;
+
+		const FVec3& LocalSpaceLocation = Solver->GetLocalSpaceLocation();
+
+		for (const FClothingSimulationCloth* const Cloth : Solver->GetCloths())
+		{
+			const int32 Offset = Cloth->GetOffset(Solver);
+			if (Offset == INDEX_NONE)
+			{
+				continue;
+			}
+
+			const TConstArrayView<FRealSingle>& MaxDistances = Cloth->GetWeightMaps(Solver)[(int32)EChaosWeightMapTarget::MaxDistance];
+			if (!MaxDistances.Num())
+			{
+				continue;
+			}
+
+			const TConstArrayView<Softs::FSolverVec3> Positions = Cloth->GetAnimationPositions(Solver);
+			const TConstArrayView<Softs::FSolverReal> InvMasses = Cloth->GetParticleInvMasses(Solver);
+			check(MaxDistances.Num() == Positions.Num());
+			check(MaxDistances.Num() == InvMasses.Num());
+
+			for (int32 Index = 0; Index < MaxDistances.Num(); ++Index)
+			{
+				const FReal MaxDistance = MaxDistances[Index];
+				const FVector Position = LocalSpaceLocation + FVector(Positions[Index]);
+
+				const FText Text = FText::AsNumber(MaxDistance, &NumberFormattingOptions);
+				DrawText(Canvas, SceneView, Position, Text, InvMasses[Index] == (Softs::FSolverReal)0. ? KinematicColor : DynamicColor);
+			}
+		}
+	}
+#endif  // #if WITH_EDITOR
+
+	static void DrawPoint(FPrimitiveDrawInterface* PDI, const FVector& Pos, const FLinearColor& Color, const UMaterial* ClothMaterialVertex, const float Thickness = 1.f)  // Use color or material
+	{
+		using namespace ::Chaos;
+
+		if (!PDI)
+		{
+			FDebugDrawQueue::GetInstance().DrawDebugPoint(Pos, Color.ToFColor(true), false, KINDA_SMALL_NUMBER, SDPG_Foreground, Thickness);
+			return;
+		}
+#if WITH_EDITOR
+		if (ClothMaterialVertex)
+		{
+			const FMatrix& ViewMatrix = PDI->View->ViewMatrices.GetViewMatrix();
+			const FVector XAxis = ViewMatrix.GetColumn(0); // Just using transpose here (orthogonal transform assumed)
+			const FVector YAxis = ViewMatrix.GetColumn(1);
+			DrawDisc(PDI, Pos, XAxis, YAxis, Color.ToFColor(true), 0.5f, 10, ClothMaterialVertex->GetRenderProxy(), SDPG_World);
+		}
+		else
+		{
+			PDI->DrawPoint(Pos, Color, Thickness, SDPG_World);
+		}
+#endif
+	}
+
+	static void DrawLine(FPrimitiveDrawInterface* PDI, const FVector& Pos0, const FVector& Pos1, const FLinearColor& Color)
+	{
+		using namespace ::Chaos;
+
+		if (!PDI)
+		{
+			FDebugDrawQueue::GetInstance().DrawDebugLine(Pos0, Pos1, Color.ToFColor(true), false, KINDA_SMALL_NUMBER, SDPG_Foreground, 0.f);
+			return;
+		}
+#if WITH_EDITOR
+		PDI->DrawLine(Pos0, Pos1, Color, SDPG_World, 0.0f, 0.001f);
+#endif
+	}
+
+	static void DrawArc(FPrimitiveDrawInterface* PDI, const FVector& Base, const FVector& X, const FVector& Y, ::Chaos::FReal MinAngle, float MaxAngle, ::Chaos::FReal Radius, const FLinearColor& Color)
+	{
+		using namespace ::Chaos;
+
+		static const int32 Sections = 10;
+		const FReal AngleStep = FMath::DegreesToRadians((MaxAngle - MinAngle) / (FReal)Sections);
+		FReal CurrentAngle = FMath::DegreesToRadians(MinAngle);
+		FVector LastVertex = Base + Radius * (FMath::Cos(CurrentAngle) * X + FMath::Sin(CurrentAngle) * Y);
+
+		for(int32 i = 0; i < Sections; i++)
+		{
+			CurrentAngle += AngleStep;
+			const FVector ThisVertex = Base + Radius * (FMath::Cos(CurrentAngle) * X + FMath::Sin(CurrentAngle) * Y);
+			DrawLine(PDI, LastVertex, ThisVertex, Color);
+			LastVertex = ThisVertex;
+		}
+	}
+
+	static void DrawSphere(FPrimitiveDrawInterface* PDI, const ::Chaos::TSphere<::Chaos::FReal, 3>& Sphere, const FQuat& Rotation, const FVector& Position, const FLinearColor& Color)
+	{
+		using namespace ::Chaos;
+
+		const FReal Radius = Sphere.GetRadius();
+		const FVec3 Center = Position + Rotation.RotateVector(Sphere.GetCenter());
+		if (!PDI)
+		{
+			FDebugDrawQueue::GetInstance().DrawDebugSphere(Center, Radius, 12, Color.ToFColor(true), false, KINDA_SMALL_NUMBER, SDPG_Foreground, 0.f);
+			return;
+		}
+#if WITH_EDITOR
+		const FTransform Transform(Rotation, Center);
+		DrawWireSphere(PDI, Transform, Color, Radius, 12, SDPG_World, 0.0f, 0.001f, false);
+#endif
+	}
+
+	static void DrawBox(FPrimitiveDrawInterface* PDI, const ::Chaos::FAABB3& Box, const FQuat& Rotation, const FVector& Position, const FLinearColor& Color)
+	{
+		using namespace ::Chaos;
+
+		if (!PDI)
+		{
+			const FVec3 Center = Position + Rotation.RotateVector(Box.GetCenter());
+			FDebugDrawQueue::GetInstance().DrawDebugBox(Center, Box.Extents() * 0.5f, Rotation, Color.ToFColor(true), false, KINDA_SMALL_NUMBER, SDPG_Foreground, 0.f);
+			return;
+		}
+#if WITH_EDITOR
+		const FMatrix BoxToWorld = FTransform(Rotation, Position).ToMatrixNoScale();
+		DrawWireBox(PDI, BoxToWorld, FBox(Box.Min(), Box.Max()), Color, SDPG_World, 0.0f, 0.001f, false);
+#endif
+	}
+
+	static void DrawCapsule(FPrimitiveDrawInterface* PDI, const ::Chaos::FCapsule& Capsule, const FQuat& Rotation, const FVector& Position, const FLinearColor& Color)
+	{
+		using namespace ::Chaos;
+
+		const FReal Radius = Capsule.GetRadius();
+		const FReal HalfHeight = Capsule.GetHeight() * 0.5f + Radius;
+		const FVec3 Center = Position + Rotation.RotateVector(Capsule.GetCenter());
+		if (!PDI)
+		{
+			const FQuat Orientation = FQuat::FindBetweenNormals(FVec3::UpVector, Capsule.GetAxis());
+			FDebugDrawQueue::GetInstance().DrawDebugCapsule(Center, HalfHeight, Radius, Rotation * Orientation, Color.ToFColor(true), false, KINDA_SMALL_NUMBER, SDPG_Foreground, 0.f);
+			return;
+		}
+#if WITH_EDITOR
+		const FVec3 Up = Capsule.GetAxis();
+		FVec3 Forward, Right;
+		Up.FindBestAxisVectors(Forward, Right);
+		const FVector X = Rotation.RotateVector(Forward);
+		const FVector Y = Rotation.RotateVector(Right);
+		const FVector Z = Rotation.RotateVector(Up);
+		DrawWireCapsule(PDI, Center, X, Y, Z, Color, Radius, HalfHeight, 12, SDPG_World, 0.0f, 0.001f, false);
+#endif
+	}
+
+	static void DrawTaperedCylinder(FPrimitiveDrawInterface* PDI, const ::Chaos::FTaperedCylinder& TaperedCylinder, const FQuat& Rotation, const FVector& Position, const FLinearColor& Color)
+	{
+		using namespace ::Chaos;
+
+		const FReal HalfHeight = TaperedCylinder.GetHeight() * 0.5f;
+		const FReal Radius1 = TaperedCylinder.GetRadius1();
+		const FReal Radius2 = TaperedCylinder.GetRadius2();
+		const FVector Position1 = Position + Rotation.RotateVector(TaperedCylinder.GetX1());
+		const FVector Position2 = Position + Rotation.RotateVector(TaperedCylinder.GetX2());
+		const FQuat Q = (Position2 - Position1).ToOrientationQuat();
+		const FVector I = Q.GetRightVector();
+		const FVector J = Q.GetUpVector();
+
+		static const int32 NumSides = 12;
+		static const FReal	AngleDelta = (FReal)2. * (FReal)PI / NumSides;
+		FVector	LastVertex1 = Position1 + I * Radius1;
+		FVector	LastVertex2 = Position2 + I * Radius2;
+
+		for (int32 SideIndex = 1; SideIndex <= NumSides; ++SideIndex)
+		{
+			const FReal Angle = AngleDelta * FReal(SideIndex);
+			const FVector ArcPos = I * FMath::Cos(Angle) + J * FMath::Sin(Angle);
+			const FVector Vertex1 = Position1 + ArcPos * Radius1;
+			const FVector Vertex2 = Position2 + ArcPos * Radius2;
+
+			DrawLine(PDI, LastVertex1, Vertex1, Color);
+			DrawLine(PDI, LastVertex2, Vertex2, Color);
+			DrawLine(PDI, LastVertex1, LastVertex2, Color);
+
+			LastVertex1 = Vertex1;
+			LastVertex2 = Vertex2;
+		}
+	}
+
+	static void DrawConvex(FPrimitiveDrawInterface* PDI, const ::Chaos::FConvex& Convex, const FQuat& Rotation, const FVector& Position, const FLinearColor& Color)
+	{
+		using namespace ::Chaos;
+
+		const TArray<FConvex::FPlaneType>& Planes = Convex.GetFaces();
+		for (int32 PlaneIndex1 = 0; PlaneIndex1 < Planes.Num(); ++PlaneIndex1)
+		{
+			const FConvex::FPlaneType& Plane1 = Planes[PlaneIndex1];
+
+			for (int32 PlaneIndex2 = PlaneIndex1 + 1; PlaneIndex2 < Planes.Num(); ++PlaneIndex2)
+			{
+				const FConvex::FPlaneType& Plane2 = Planes[PlaneIndex2];
+
+				// Find the two surface points that belong to both Plane1 and Plane2
+				uint32 ParticleIndex1 = INDEX_NONE;
+
+				const TArray<FConvex::FVec3Type>& Vertices = Convex.GetVertices();
+				for (int32 ParticleIndex = 0; ParticleIndex < Vertices.Num(); ++ParticleIndex)
+				{
+					const FConvex::FVec3Type& X = Vertices[ParticleIndex];
+
+					if (FMath::Square(Plane1.SignedDistance(X)) < KINDA_SMALL_NUMBER && 
+						FMath::Square(Plane2.SignedDistance(X)) < KINDA_SMALL_NUMBER)
+					{
+						if (ParticleIndex1 != INDEX_NONE)
+						{
+							const FVector X1(Vertices[ParticleIndex1]);
+							const FVector X2(X);
+							const FVector Position1 = Position + Rotation.RotateVector(X1);
+							const FVector Position2 = Position + Rotation.RotateVector(X2);
+							DrawLine(PDI, Position1, Position2, Color);
+							break;
+						}
+						ParticleIndex1 = ParticleIndex;
+					}
+				}
+			}
+		}
+	}
+
+	static void DrawCoordinateSystem(FPrimitiveDrawInterface* PDI, const FQuat& Rotation, const FVector& Position)
+	{
+		const FVector X = Rotation.RotateVector(FVector::ForwardVector) * 10.f;
+		const FVector Y = Rotation.RotateVector(FVector::RightVector) * 10.f;
+		const FVector Z = Rotation.RotateVector(FVector::UpVector) * 10.f;
+
+		DrawLine(PDI, Position, Position + X, FLinearColor::Red);
+		DrawLine(PDI, Position, Position + Y, FLinearColor::Green);
+		DrawLine(PDI, Position, Position + Z, FLinearColor::Blue);
+	}
+
+	void FVisualization::DrawBounds(FPrimitiveDrawInterface* PDI) const
+	{
+		using namespace ::Chaos;
+
+		if (!Solver)
+		{
+			return;
+		}
+
+		// Calculate World space bounds
+		const FBoxSphereBounds Bounds = Solver->CalculateBounds();
+
+		// Draw bounds
+		DrawBox(PDI, FAABB3(-Bounds.BoxExtent, Bounds.BoxExtent), FQuat::Identity, Bounds.Origin, FLinearColor(FColor::Purple));
+		DrawSphere(PDI, TSphere<FReal, 3>(FVector::ZeroVector, Bounds.SphereRadius), FQuat::Identity, Bounds.Origin, FLinearColor(FColor::Orange));
+
+		// Draw individual cloth bounds
+		static const FLinearColor Color = FLinearColor(FColor::Purple).Desaturate(0.5);
+		for (const FClothingSimulationCloth* const Cloth : Solver->GetCloths())
+		{
+			if (Cloth->GetOffset(Solver) == INDEX_NONE)
+			{
+				continue;
+			}
+
+			const FAABB3 BoundingBox = Cloth->CalculateBoundingBox(Solver);
+			DrawBox(PDI, BoundingBox, FQuat::Identity, FVector::ZeroVector, Color);  // TODO: Express bounds in local coordinates for LWC
+		}
+	}
+
+	void FVisualization::DrawGravity(FPrimitiveDrawInterface* PDI) const
+	{
+		using namespace ::Chaos;
+
+		if (!Solver)
+		{
+			return;
+		}
+
+		// Draw gravity
+		for (const FClothingSimulationCloth* const Cloth : Solver->GetCloths())
+		{
+			if (Cloth->GetOffset(Solver) == INDEX_NONE)
+			{
+				continue;
+			}
+
+			const FAABB3 Bounds = Cloth->CalculateBoundingBox(Solver);
+
+			const FVector Pos0 = Bounds.Center();
+			const FVector Pos1 = Pos0 + FVector(Cloth->GetGravity(Solver));
+			DrawLine(PDI, Pos0, Pos1, FLinearColor::Red);
+		}
+	}
+
+	void FVisualization::DrawPhysMeshWired(FPrimitiveDrawInterface* PDI) const
+	{
+		using namespace ::Chaos;
+
+		if (!Solver)
+		{
+			return;
+		}
+
+		static const FLinearColor DynamicColor = FColor::White;
+		static const FLinearColor KinematicColor = FColor::Purple;
+
+		const FVec3& LocalSpaceLocation = Solver->GetLocalSpaceLocation();
+
+		for (const FClothingSimulationCloth* const Cloth : Solver->GetCloths())
+		{
+			const int32 Offset = Cloth->GetOffset(Solver);
+			if (Offset == INDEX_NONE)
+			{
+				continue;
+			}
+
+			const TConstArrayView<TVec3<int32>> Elements = Cloth->GetTriangleMesh(Solver).GetElements();
+			const TConstArrayView<Softs::FSolverVec3> Positions = Cloth->GetParticlePositions(Solver);
+			const TConstArrayView<Softs::FSolverReal> InvMasses = Cloth->GetParticleInvMasses(Solver);
+			check(InvMasses.Num() == Positions.Num());
+
+			for (int32 ElementIndex = 0; ElementIndex < Elements.Num(); ++ElementIndex)
+			{
+				const auto& Element = Elements[ElementIndex];
+
+				const FVector Pos0 = LocalSpaceLocation + FVector(Positions[Element.X - Offset]); // TODO: Triangle Mesh shouldn't really be solver dependent (ie not use an offset)
+				const FVector Pos1 = LocalSpaceLocation + FVector(Positions[Element.Y - Offset]);
+				const FVector Pos2 = LocalSpaceLocation + FVector(Positions[Element.Z - Offset]);
+
+				const bool bIsKinematic0 = (InvMasses[Element.X - Offset] == (Softs::FSolverReal)0.);
+				const bool bIsKinematic1 = (InvMasses[Element.Y - Offset] == (Softs::FSolverReal)0.);
+				const bool bIsKinematic2 = (InvMasses[Element.Z - Offset] == (Softs::FSolverReal)0.);
+
+				DrawLine(PDI, Pos0, Pos1, bIsKinematic0 && bIsKinematic1 ? KinematicColor : DynamicColor);
+				DrawLine(PDI, Pos1, Pos2, bIsKinematic1 && bIsKinematic2 ? KinematicColor : DynamicColor);
+				DrawLine(PDI, Pos2, Pos0, bIsKinematic2 && bIsKinematic0 ? KinematicColor : DynamicColor);
+			}
+		}
+	}
+
+	void FVisualization::DrawAnimMeshWired(FPrimitiveDrawInterface* PDI) const
+	{
+		using namespace ::Chaos;
+
+		if (!Solver)
+		{
+			return;
+		}
+
+		static const FLinearColor KinematicColor = FColor::Purple;
+
+		const FVec3& LocalSpaceLocation = Solver->GetLocalSpaceLocation();
+
+		for (const FClothingSimulationCloth* const Cloth : Solver->GetCloths())
+		{
+			const int32 Offset = Cloth->GetOffset(Solver);
+			if (Offset == INDEX_NONE)
+			{
+				continue;
+			}
+
+			const TConstArrayView<TVec3<int32>> Elements = Cloth->GetTriangleMesh(Solver).GetElements();
+			const TConstArrayView<Softs::FSolverVec3> Positions = Cloth->GetAnimationPositions(Solver);
+
+			for (int32 ElementIndex = 0; ElementIndex < Elements.Num(); ++ElementIndex)
+			{
+				const auto& Element = Elements[ElementIndex];
+
+				const FVector Pos0 = LocalSpaceLocation + FVector(Positions[Element.X - Offset]); // TODO: Triangle Mesh shouldn't really be solver dependent (ie not use an offset)
+				const FVector Pos1 = LocalSpaceLocation + FVector(Positions[Element.Y - Offset]);
+				const FVector Pos2 = LocalSpaceLocation + FVector(Positions[Element.Z - Offset]);
+
+				DrawLine(PDI, Pos0, Pos1, KinematicColor);
+				DrawLine(PDI, Pos1, Pos2, KinematicColor);
+				DrawLine(PDI, Pos2, Pos0, KinematicColor);
+			}
+		}
+	}
+
+	void FVisualization::DrawAnimNormals(FPrimitiveDrawInterface* PDI) const
+	{
+		using namespace ::Chaos;
+
+		if (!Solver)
+		{
+			return;
+		}
+
+		static const FLinearColor KinematicColor = FColor::Magenta;
+		constexpr FReal NormalLength = (FReal)20.;
+
+		const FVec3& LocalSpaceLocation = Solver->GetLocalSpaceLocation();
+
+		for (const FClothingSimulationCloth* const Cloth : Solver->GetCloths())
+		{
+			const int32 Offset = Cloth->GetOffset(Solver);
+			if (Offset == INDEX_NONE)
+			{
+				continue;
+			}
+
+			const TConstArrayView<Softs::FSolverVec3> Positions = Cloth->GetAnimationPositions(Solver);
+			const TConstArrayView<Softs::FSolverVec3> Normals = Cloth->GetAnimationNormals(Solver);
+			check(Normals.Num() == Positions.Num());
+
+			for (int32 Index = 0; Index < Positions.Num(); ++Index)
+			{
+				const FVector Pos0 = LocalSpaceLocation + FVector(Positions[Index]);
+				const FVector Pos1 = Pos0 + FVector(Normals[Index]) * NormalLength;
+
+				DrawLine(PDI, Pos0, Pos1, KinematicColor);
+			}
+		}
+	}
+
+	void FVisualization::DrawPointNormals(FPrimitiveDrawInterface* PDI) const
+	{
+		using namespace ::Chaos;
+
+		if (!Solver)
+		{
+			return;
+		}
+
+		static const FLinearColor DynamicColor = FColor::White;
+		static const FLinearColor KinematicColor = FColor::Purple;
+		constexpr FReal NormalLength = (FReal)20.;
+
+		const FVec3& LocalSpaceLocation = Solver->GetLocalSpaceLocation();
+
+		for (const FClothingSimulationCloth* const Cloth : Solver->GetCloths())
+		{
+			const int32 Offset = Cloth->GetOffset(Solver);
+			if (Offset == INDEX_NONE)
+			{
+				continue;
+			}
+
+			const TConstArrayView<Softs::FSolverVec3> Positions = Cloth->GetParticlePositions(Solver);
+			const TConstArrayView<Softs::FSolverVec3> Normals = Cloth->GetParticleNormals(Solver);
+			const TConstArrayView<Softs::FSolverReal> InvMasses = Cloth->GetParticleInvMasses(Solver);
+			check(Normals.Num() == Positions.Num());
+			check(InvMasses.Num() == Positions.Num());
+
+			for (int32 Index = 0; Index < Positions.Num(); ++Index)
+			{
+				const bool bIsKinematic = (InvMasses[Index] == (Softs::FSolverReal)0.);
+				const FVector Pos0 = LocalSpaceLocation + FVector(Positions[Index]);
+				const FVector Pos1 = Pos0 + FVector(Normals[Index]) * NormalLength;
+
+				DrawLine(PDI, Pos0, Pos1, bIsKinematic ? KinematicColor : DynamicColor);
+			}
+		}
+	}
+
+	void FVisualization::DrawPointVelocities(FPrimitiveDrawInterface* PDI) const
+	{
+		using namespace ::Chaos;
+
+		if (!Solver)
+		{
+			return;
+		}
+
+		const FVec3& LocalSpaceLocation = Solver->GetLocalSpaceLocation();
+
+		for (const FClothingSimulationCloth* const Cloth : Solver->GetCloths())
+		{
+			const int32 Offset = Cloth->GetOffset(Solver);
+			if (Offset == INDEX_NONE)
+			{
+				continue;
+			}
+
+			const TConstArrayView<Softs::FSolverVec3> Positions = Cloth->GetParticlePositions(Solver);
+			const TConstArrayView<Softs::FSolverVec3> Velocities = Cloth->GetParticleVelocities(Solver);
+			check(Velocities.Num() == Positions.Num());
+
+			const TConstArrayView<Softs::FSolverReal> InvMasses = Cloth->GetParticleInvMasses(Solver);
+			check(InvMasses.Num() == Positions.Num());
+
+			for (int32 Index = 0; Index < Positions.Num(); ++Index)
+			{
+				constexpr FReal DefaultFPS = (FReal)60.;   // TODO: A CVAR would be nice for this
+				const bool bIsKinematic = (InvMasses[Index] == 0.f);
+
+				const FVec3 Pos0 = LocalSpaceLocation + Positions[Index];
+				const FVec3 Pos1 = Pos0 + FVec3(Velocities[Index]) / DefaultFPS;  // Velocity per frame if running at DefaultFPS
+
+				DrawLine(PDI, Pos0, Pos1, bIsKinematic ? FLinearColor::Black : FLinearColor::Yellow);
+			}
+		}
+	}
+
+	void FVisualization::DrawCollision(FPrimitiveDrawInterface* PDI) const
+	{
+		using namespace ::Chaos;
+
+		if (!Solver)
+		{
+			return;
+		}
+
+		auto DrawCollision =
+			[this, PDI](const FClothingSimulationCollider* Collider, const FClothingSimulationCloth* Cloth, FClothingSimulationCollider::ECollisionDataType CollisionDataType)
+			{
+				static const FLinearColor GlobalColor(FColor::Cyan);
+				static const FLinearColor DynamicColor(FColor::Orange);
+				static const FLinearColor LODsColor(FColor::Silver);
+				static const FLinearColor CollidedColor(FColor::Red);
+
+				const FLinearColor TypeColor =
+					(CollisionDataType == FClothingSimulationCollider::ECollisionDataType::LODless) ? GlobalColor :
+					(CollisionDataType == FClothingSimulationCollider::ECollisionDataType::External) ? DynamicColor : LODsColor;
+
+				const FVec3& LocalSpaceLocation = Solver->GetLocalSpaceLocation();
+
+				const TConstArrayView<TUniquePtr<FImplicitObject>> CollisionGeometries = Collider->GetCollisionGeometries(Solver, Cloth, CollisionDataType);
+				const TConstArrayView<Softs::FSolverVec3> Translations = Collider->GetCollisionTranslations(Solver, Cloth, CollisionDataType);
+				const TConstArrayView<Softs::FSolverRotation3> Rotations = Collider->GetCollisionRotations(Solver, Cloth, CollisionDataType);
+				const TConstArrayView<bool> CollisionStatus = Collider->GetCollisionStatus(Solver, Cloth, CollisionDataType);
+				check(CollisionGeometries.Num() == Translations.Num());
+				check(CollisionGeometries.Num() == Rotations.Num());
+
+				for (int32 Index = 0; Index < CollisionGeometries.Num(); ++Index)
+				{
+					if (const FImplicitObject* const Object = CollisionGeometries[Index].Get())
+					{
+						const FLinearColor Color = CollisionStatus[Index] ? CollidedColor : TypeColor;
+						const FVec3 Position = LocalSpaceLocation + FVec3(Translations[Index]);
+						const FRotation3 Rotation(Rotations[Index]);
+
+						switch (Object->GetType())
+						{
+						case ImplicitObjectType::Sphere:
+							DrawSphere(PDI, Object->GetObjectChecked<TSphere<FReal, 3>>(), Rotation, Position, Color);
+							break;
+
+						case ImplicitObjectType::Box:
+							DrawBox(PDI, Object->GetObjectChecked<TBox<FReal, 3>>().BoundingBox(), Rotation, Position, Color);
+							break;
+
+						case ImplicitObjectType::Capsule:
+							DrawCapsule(PDI, Object->GetObjectChecked<FCapsule>(), Rotation, Position, Color);
+							break;
+
+						case ImplicitObjectType::Union:  // Union only used as old style tapered capsules
+							for (const TUniquePtr<FImplicitObject>& SubObjectPtr : Object->GetObjectChecked<FImplicitObjectUnion>().GetObjects())
+							{
+								if (const FImplicitObject* const SubObject = SubObjectPtr.Get())
+								{
+									switch (SubObject->GetType())
+									{
+									case ImplicitObjectType::Sphere:
+										DrawSphere(PDI, SubObject->GetObjectChecked<TSphere<FReal, 3>>(), Rotation, Position, Color);
+										break;
+
+									case ImplicitObjectType::TaperedCylinder:
+										DrawTaperedCylinder(PDI, SubObject->GetObjectChecked<FTaperedCylinder>(), Rotation, Position, Color);
+										break;
+
+									default:
+										break;
+									}
+								}
+							}
+							break;
+
+						case ImplicitObjectType::TaperedCapsule:  // New collision tapered capsules implicit type that replaces the union
+							{
+								const FTaperedCapsule& TaperedCapsule = Object->GetObjectChecked<FTaperedCapsule>();
+								const FVec3 X1 = TaperedCapsule.GetX1();
+								const FVec3 X2 = TaperedCapsule.GetX2();
+								const FReal Radius1 = TaperedCapsule.GetRadius1();
+								const FReal Radius2 = TaperedCapsule.GetRadius2();
+								DrawSphere(PDI, TSphere<FReal, 3>(X1, Radius1), Rotation, Position, Color);
+								DrawSphere(PDI, TSphere<FReal, 3>(X2, Radius2), Rotation, Position, Color);
+								DrawTaperedCylinder(PDI, FTaperedCylinder(X1, X2, Radius1, Radius2), Rotation, Position, Color);
+							}
+							break;
+
+						case ImplicitObjectType::Convex:
+							DrawConvex(PDI, Object->GetObjectChecked<FConvex>(), Rotation, Position, Color);
+							break;
+
+						default:
+							DrawCoordinateSystem(PDI, Rotation, Position);  // Draw everything else as a coordinate for now
+							break;
+						}
+					}
+				}
+			};
+
+		// Draw collisions
+		for (const FClothingSimulationCloth* const Cloth : Solver->GetCloths())
+		{
+			for (const FClothingSimulationCollider* const Collider : Cloth->GetColliders())
+			{
+				DrawCollision(Collider, Cloth, FClothingSimulationCollider::ECollisionDataType::LODless);
+				DrawCollision(Collider, Cloth, FClothingSimulationCollider::ECollisionDataType::External);
+				DrawCollision(Collider, Cloth, FClothingSimulationCollider::ECollisionDataType::LODs);
+			}
+		}
+
+		// Draw contacts
+		check(Solver->GetCollisionContacts().Num() == Solver->GetCollisionNormals().Num());
+		constexpr FReal NormalLength = (FReal)10.;
+
+		const FVec3& LocalSpaceLocation = Solver->GetLocalSpaceLocation();
+		for (int32 i = 0; i < Solver->GetCollisionContacts().Num(); ++i)
+		{
+			const FVec3 Pos0 = LocalSpaceLocation + FVec3(Solver->GetCollisionContacts()[i]);
+			const FVec3 Normal = FVec3(Solver->GetCollisionNormals()[i]);
+
+			// Draw contact
+			FVec3 TangentU, TangentV;
+			Normal.FindBestAxisVectors(TangentU, TangentV);
+
+			DrawLine(PDI, Pos0 + TangentU, Pos0 + TangentV, FLinearColor::Black);
+			DrawLine(PDI, Pos0 + TangentU, Pos0 - TangentV, FLinearColor::Black);
+			DrawLine(PDI, Pos0 - TangentU, Pos0 - TangentV, FLinearColor::Black);
+			DrawLine(PDI, Pos0 - TangentU, Pos0 + TangentV, FLinearColor::Black);
+
+			// Draw normal
+			static const FLinearColor Brown(0.1f, 0.05f, 0.f);
+			const FVec3 Pos1 = Pos0 + NormalLength * Normal;
+			DrawLine(PDI, Pos0, Pos1, Brown);
+		}
+	}
+
+	void FVisualization::DrawBackstops(FPrimitiveDrawInterface* PDI) const
+	{
+		using namespace ::Chaos;
+
+		if (!Solver)
+		{
+			return;
+		}
+
+		auto DrawBackstop = [PDI](const FVector& Position, const FVector& Normal, FReal Radius, const FVector& Axis, const FLinearColor& Color)
+		{
+			static const FReal MaxCosAngle = (FReal)0.99;
+			if (FMath::Abs(FVector::DotProduct(Normal, Axis)) < MaxCosAngle)
+			{
+				static const FReal ArcLength = (FReal)5.; // Arch length in cm
+				const FReal ArcAngle = (FReal)360. * ArcLength / FMath::Max((Radius * (FReal)2. * (FReal)PI), ArcLength);
+				DrawArc(PDI, Position, Normal, FVector::CrossProduct(Axis, Normal).GetSafeNormal(), -ArcAngle / (FReal)2., ArcAngle / (FReal)2., Radius, Color);
+			}
+		};
+
+		const FVec3& LocalSpaceLocation = Solver->GetLocalSpaceLocation();
+
+		uint8 ColorSeed = 0;
+
+		for (const FClothingSimulationCloth* const Cloth : Solver->GetCloths())
+		{
+			const int32 Offset = Cloth->GetOffset(Solver);
+			if (Offset == INDEX_NONE)
+			{
+				continue;
+			}
+
+			const FClothConstraints& ClothConstraints = Solver->GetClothConstraints(Offset);
+			if (const Softs::FPBDSphericalBackstopConstraint* const BackstopConstraint = ClothConstraints.GetBackstopConstraints().Get())
+			{
+				const bool bUseLegacyBackstop = BackstopConstraint->UseLegacyBackstop();
+				const TConstArrayView<FRealSingle>& BackstopDistances = Cloth->GetWeightMaps(Solver)[(int32)EChaosWeightMapTarget::BackstopDistance];
+				const TConstArrayView<FRealSingle>& BackstopRadiuses = Cloth->GetWeightMaps(Solver)[(int32)EChaosWeightMapTarget::BackstopRadius];
+				const TConstArrayView<Softs::FSolverVec3> AnimationPositions = Cloth->GetAnimationPositions(Solver);
+				const TConstArrayView<Softs::FSolverVec3> AnimationNormals = Cloth->GetAnimationNormals(Solver);
+				const TConstArrayView<Softs::FSolverVec3> ParticlePositions = Cloth->GetParticlePositions(Solver);
+
+				for (int32 Index = 0; Index < AnimationPositions.Num(); ++Index)
+				{
+					ColorSeed += 157;  // Prime number that gives a good spread of colors without getting too similar as a rand might do.
+					const FLinearColor ColorLight = FLinearColor::MakeFromHSV8(ColorSeed, 160, 128);
+					const FLinearColor ColorDark = FLinearColor::MakeFromHSV8(ColorSeed, 160, 64);
+
+					const FReal BackstopRadius = BackstopRadiuses[Index] * BackstopConstraint->GetSphereRadiiMultiplier();
+					const FReal BackstopDistance = BackstopDistances[Index];
+
+					const FVector AnimationNormal(AnimationNormals[Index]);
+
+					// Draw a line to show the current distance to the sphere
+					const FVector Pos0 = LocalSpaceLocation + FVector(AnimationPositions[Index]);
+					const FVector Pos1 = Pos0 - (bUseLegacyBackstop ? BackstopDistance - BackstopRadius : BackstopDistance) * AnimationNormal;
+					const FVector Pos2 = LocalSpaceLocation + FVector(ParticlePositions[Index]);
+					DrawLine(PDI, Pos1, Pos2, ColorLight);
+
+					// Draw the sphere
+					if (BackstopRadius > 0.f)
+					{
+						const FVector Center = Pos0 - (bUseLegacyBackstop ? BackstopDistance : BackstopRadius + BackstopDistance) * AnimationNormal;
+						DrawBackstop(Center, AnimationNormal, BackstopRadius, FVector::ForwardVector, ColorDark);
+						DrawBackstop(Center, AnimationNormal, BackstopRadius, FVector::UpVector, ColorDark);
+						DrawBackstop(Center, AnimationNormal, BackstopRadius, FVector::RightVector, ColorDark);
+					}
+				}
+			}
+		}
+	}
+
+	void FVisualization::DrawBackstopDistances(FPrimitiveDrawInterface* PDI) const
+	{
+		using namespace ::Chaos;
+
+		if (!Solver)
+		{
+			return;
+		}
+
+		const FVec3& LocalSpaceLocation = Solver->GetLocalSpaceLocation();
+
+		uint8 ColorSeed = 0;
+
+		for (const FClothingSimulationCloth* const Cloth : Solver->GetCloths())
+		{
+			const int32 Offset = Cloth->GetOffset(Solver);
+			if (Offset == INDEX_NONE)
+			{
+				continue;
+			}
+
+			const FClothConstraints& ClothConstraints = Solver->GetClothConstraints(Offset);
+			if (const Softs::FPBDSphericalBackstopConstraint* const BackstopConstraint = ClothConstraints.GetBackstopConstraints().Get())
+			{
+				const bool bUseLegacyBackstop = BackstopConstraint->UseLegacyBackstop();
+				const TConstArrayView<FRealSingle>& BackstopDistances = Cloth->GetWeightMaps(Solver)[(int32)EChaosWeightMapTarget::BackstopDistance];
+				const TConstArrayView<FRealSingle>& BackstopRadiuses = Cloth->GetWeightMaps(Solver)[(int32)EChaosWeightMapTarget::BackstopRadius];
+				const TConstArrayView<Softs::FSolverVec3> AnimationPositions = Cloth->GetAnimationPositions(Solver);
+				const TConstArrayView<Softs::FSolverVec3> AnimationNormals = Cloth->GetAnimationNormals(Solver);
+
+				for (int32 Index = 0; Index < AnimationPositions.Num(); ++Index)
+				{
+					ColorSeed += 157;  // Prime number that gives a good spread of colors without getting too similar as a rand might do.
+					const FLinearColor ColorLight = FLinearColor::MakeFromHSV8(ColorSeed, 160, 128);
+					const FLinearColor ColorDark = FLinearColor::MakeFromHSV8(ColorSeed, 160, 64);
+
+					const FReal BackstopRadius = BackstopRadiuses[Index] * BackstopConstraint->GetSphereRadiiMultiplier();
+					const FReal BackstopDistance = BackstopDistances[Index];
+
+					const FVector AnimationNormal(AnimationNormals[Index]);
+
+					// Draw a line to the sphere boundary
+					const FVector Pos0 = LocalSpaceLocation + FVector(AnimationPositions[Index]);
+					const FVector Pos1 = Pos0 - (bUseLegacyBackstop ? BackstopDistance - BackstopRadius : BackstopDistance) * AnimationNormal;
+					DrawLine(PDI, Pos0, Pos1, ColorDark);
+				}
+			}
+		}
+	}
+
+	void FVisualization::DrawMaxDistances(FPrimitiveDrawInterface* PDI) const
+	{
+		using namespace ::Chaos;
+
+		if (!Solver)
+		{
+			return;
+		}
+
+		// Draw max distances
+		const FVec3& LocalSpaceLocation = Solver->GetLocalSpaceLocation();
+
+		for (const FClothingSimulationCloth* const Cloth : Solver->GetCloths())
+		{
+			const int32 Offset = Cloth->GetOffset(Solver);
+			if (Offset == INDEX_NONE)
+			{
+				continue;
+			}
+
+			const TConstArrayView<FRealSingle>& MaxDistances = Cloth->GetWeightMaps(Solver)[(int32)EChaosWeightMapTarget::MaxDistance];
+			if (!MaxDistances.Num())
+			{
+				continue;
+			}
+
+			const TConstArrayView<Softs::FSolverReal> InvMasses = Cloth->GetParticleInvMasses(Solver);
+			const TConstArrayView<Softs::FSolverVec3> Positions = Cloth->GetAnimationPositions(Solver);
+			const TConstArrayView<Softs::FSolverVec3> Normals = Cloth->GetAnimationNormals(Solver);
+			check(Normals.Num() == Positions.Num());
+			check(MaxDistances.Num() == Positions.Num());
+			check(InvMasses.Num() == Positions.Num());
+
+			for (int32 Index = 0; Index < MaxDistances.Num(); ++Index)
+			{
+				const FReal MaxDistance = (FReal)MaxDistances[Index];
+				const FVector Position = LocalSpaceLocation + FVector(Positions[Index]);
+				if (InvMasses[Index] == (Softs::FSolverReal)0.)
+				{
+#if WITH_EDITOR
+					DrawPoint(PDI, Position, FLinearColor::Red, ClothMaterialVertex);
+#else
+					DrawPoint(nullptr, Position, FLinearColor::Red, nullptr);
+#endif
+				}
+				else
+				{
+					DrawLine(PDI, Position, Position + FVector(Normals[Index]) * MaxDistance, FLinearColor::White);
+				}
+			}
+		}
+	}
+
+	void FVisualization::DrawAnimDrive(FPrimitiveDrawInterface* PDI) const
+	{
+		using namespace ::Chaos;
+
+		if (!Solver)
+		{
+			return;
+		}
+
+		const FVec3& LocalSpaceLocation = Solver->GetLocalSpaceLocation();
+
+		for (const FClothingSimulationCloth* const Cloth : Solver->GetCloths())
+		{
+			const int32 Offset = Cloth->GetOffset(Solver);
+			if (Offset == INDEX_NONE)
+			{
+				continue;
+			}
+
+			const FClothConstraints& ClothConstraints = Solver->GetClothConstraints(Offset);
+			if (const Softs::FPBDAnimDriveConstraint* const AnimDriveConstraint = ClothConstraints.GetAnimDriveConstraints().Get())
+			{
+				const TConstArrayView<FRealSingle>& AnimDriveStiffnessMultipliers = Cloth->GetWeightMaps(Solver)[(int32)EChaosWeightMapTarget::AnimDriveStiffness];
+				const TConstArrayView<Softs::FSolverVec3> AnimationPositions = Cloth->GetAnimationPositions(Solver);
+				const TConstArrayView<Softs::FSolverVec3> ParticlePositions = Cloth->GetParticlePositions(Solver);
+				check(ParticlePositions.Num() == AnimationPositions.Num());
+
+				const FVec2 AnimDriveStiffness = AnimDriveConstraint->GetStiffness();
+				const FRealSingle StiffnessOffset = AnimDriveStiffness[0];
+				const FRealSingle StiffnessRange = AnimDriveStiffness[1] - AnimDriveStiffness[0];
+
+				for (int32 Index = 0; Index < ParticlePositions.Num(); ++Index)
+				{
+					const FRealSingle Stiffness = AnimDriveStiffnessMultipliers.IsValidIndex(Index) ?
+						StiffnessOffset + AnimDriveStiffnessMultipliers[Index] * StiffnessRange :
+						StiffnessOffset;
+
+					const FVector AnimationPosition = LocalSpaceLocation + FVector(AnimationPositions[Index]);
+					const FVector ParticlePosition = LocalSpaceLocation + FVector(ParticlePositions[Index]);
+					DrawLine(PDI, AnimationPosition, ParticlePosition, FLinearColor(FColor::Cyan) * Stiffness);
+				}
+			}
+		}
+	}
+
+	static void DrawSpringConstraintColors(FPrimitiveDrawInterface* PDI, const TConstArrayView<::Chaos::Softs::FSolverVec3>& Positions, const ::Chaos::FVec3& LocalSpaceLocation, const ::Chaos::Softs::FPBDSpringConstraints* const SpringConstraints)
+	{
+		using namespace ::Chaos;
+
+		check(SpringConstraints);
+
+		const TArray<TVec2<int32>>& Constraints = SpringConstraints->GetConstraints();
+		const TArray<int32>& ConstraintsPerColorStartIndex = SpringConstraints->GetConstraintsPerColorStartIndex();
+		if (ConstraintsPerColorStartIndex.Num() > 1)
+		{
+			const int32 ConstraintColorNum = ConstraintsPerColorStartIndex.Num() - 1;
+			const uint8 HueOffset = 196 / ConstraintColorNum;
+			auto ConstraintColor =
+				[HueOffset](int32 ColorIndex)->FLinearColor
+			{
+				return FLinearColor::MakeFromHSV8(ColorIndex * HueOffset, 255, 255);
+			};
+
+			for (int32 ConstraintColorIndex = 0; ConstraintColorIndex < ConstraintColorNum; ++ConstraintColorIndex)
+			{
+				const int32 ColorStart = ConstraintsPerColorStartIndex[ConstraintColorIndex];
+				const int32 ColorEnd = ConstraintsPerColorStartIndex[ConstraintColorIndex + 1];
+				const FLinearColor DrawColor = ConstraintColor(ConstraintColorIndex);
+				for (int32 ConstraintIndex = ColorStart; ConstraintIndex < ColorEnd; ++ConstraintIndex)
+				{
+					// Draw line
+					const TVec2<int32>& Constraint = Constraints[ConstraintIndex];
+					const FVec3 Pos0 = FVec3(Positions[Constraint[0]]) + LocalSpaceLocation;
+					const FVec3 Pos1 = FVec3(Positions[Constraint[1]]) + LocalSpaceLocation;
+					DrawLine(PDI, Pos0, Pos1, DrawColor);
+				}
+			}
+		}
+		else
+		{
+			for (const TVec2<int32>& Constraint : Constraints)
+			{
+				// Draw line
+				const FVec3 Pos0 = FVec3(Positions[Constraint[0]]) + LocalSpaceLocation;
+				const FVec3 Pos1 = FVec3(Positions[Constraint[1]]) + LocalSpaceLocation;
+
+				DrawLine(PDI, Pos0, Pos1, FLinearColor::Black);
+			}
+		}
+	}
+
+	void FVisualization::DrawEdgeConstraint(FPrimitiveDrawInterface* PDI) const
+	{
+		using namespace ::Chaos;
+
+		if (!Solver)
+		{
+			return;
+		}
+
+		const FVec3& LocalSpaceLocation = Solver->GetLocalSpaceLocation();
+
+		for (const FClothingSimulationCloth* const Cloth : Solver->GetCloths())
+		{
+			const int32 Offset = Cloth->GetOffset(Solver);
+			if (Offset == INDEX_NONE)
+			{
+				continue;
+			}
+
+			// Draw constraints
+			const FClothConstraints& ClothConstraints = Solver->GetClothConstraints(Offset);
+
+			const TConstArrayView<Softs::FSolverVec3> Positions = Cloth->GetParticlePositions(Solver);
+
+			if (const Softs::FPBDSpringConstraints* const EdgeConstraints = ClothConstraints.GetEdgeConstraints().Get())
+			{
+				DrawSpringConstraintColors(PDI, Positions, LocalSpaceLocation, EdgeConstraints);
+			}
+		}
+	}
+
+	void FVisualization::DrawBendingConstraint(FPrimitiveDrawInterface* PDI) const
+	{
+		using namespace ::Chaos;
+
+		if (!Solver)
+		{
+			return;
+		}
+
+		const FVec3& LocalSpaceLocation = Solver->GetLocalSpaceLocation();
+
+		for (const FClothingSimulationCloth* const Cloth : Solver->GetCloths())
+		{
+			const int32 Offset = Cloth->GetOffset(Solver);
+			if (Offset == INDEX_NONE)
+			{
+				continue;
+			}
+
+			// Draw constraints
+			const FClothConstraints& ClothConstraints = Solver->GetClothConstraints(Offset);
+
+			const TConstArrayView<Softs::FSolverVec3> Positions = Cloth->GetParticlePositions(Solver);
+
+			if (const Softs::FPBDSpringConstraints* const BendingConstraints = ClothConstraints.GetBendingConstraints().Get())
+			{
+				DrawSpringConstraintColors(PDI, Positions, LocalSpaceLocation, BendingConstraints);
+			}
+
+			if (const Softs::FPBDBendingConstraints* const BendingConstraints = ClothConstraints.GetBendingElementConstraints().Get())
+			{
+				const TArray<TVec4<int32>>& Constraints = BendingConstraints->GetConstraints();
+				const TArray<bool>& IsBuckled = BendingConstraints->GetIsBuckled();
+
+				// Color constraint edge with red or blue: Red = Buckled, Blue = Not Buckled. 
+				for (int32 ConstraintIndex = 0; ConstraintIndex < Constraints.Num(); ++ConstraintIndex)
+				{
+					const Softs::FSolverVec3& P1 = Positions[Constraints[ConstraintIndex][0]];
+					const Softs::FSolverVec3& P2 = Positions[Constraints[ConstraintIndex][1]];
+				
+					const bool bIsBuckled = IsBuckled[ConstraintIndex];
+
+					const FVec3 Pos0 = FVec3(P1) + LocalSpaceLocation;
+					const FVec3 Pos1 = FVec3(P2) + LocalSpaceLocation;
+					DrawLine(PDI, Pos0, Pos1, bIsBuckled ? FLinearColor::Red : FLinearColor::Blue);
+				}
+			}
+
+			if (const Softs::FXPBDBendingConstraints* const BendingConstraints = ClothConstraints.GetXBendingElementConstraints().Get())
+			{
+				const TArray<TVec4<int32>>& Constraints = BendingConstraints->GetConstraints();
+				const TArray<bool>& IsBuckled = BendingConstraints->GetIsBuckled();
+
+				// Color constraint edge with red or blue: Red = Buckled, Blue = Not Buckled.
+				for (int32 ConstraintIndex = 0; ConstraintIndex < Constraints.Num(); ++ConstraintIndex)
+				{
+					const Softs::FSolverVec3& P1 = Positions[Constraints[ConstraintIndex][0]];
+					const Softs::FSolverVec3& P2 = Positions[Constraints[ConstraintIndex][1]];
+
+					const bool bIsBuckled = IsBuckled[ConstraintIndex];
+
+					const FVec3 Pos0 = FVec3(P1) + LocalSpaceLocation;
+					const FVec3 Pos1 = FVec3(P2) + LocalSpaceLocation;
+					DrawLine(PDI, Pos0, Pos1, bIsBuckled ? FLinearColor::Red : FLinearColor::Blue);
+				}
+			}
+		}
+	}
+
+	void FVisualization::DrawLongRangeConstraint(FPrimitiveDrawInterface* PDI) const
+	{
+		using namespace ::Chaos;
+
+		if (!Solver)
+		{
+			return;
+		}
+
+		const FVec3& LocalSpaceLocation = Solver->GetLocalSpaceLocation();
+
+		auto PseudoRandomColor =
+			[](int32 NumColorRotations) -> FLinearColor
+			{
+				static const uint8 Spread = 157;  // Prime number that gives a good spread of colors without getting too similar as a rand might do.
+				uint8 Seed = Spread;
+				for (int32 i = 0; i < NumColorRotations; ++i)
+				{
+					Seed += Spread;
+				}
+				return FLinearColor::MakeFromHSV8(Seed, 160, 128);
+			};
+
+		auto Darken =
+			[](const FLinearColor& Color) -> FLinearColor
+			{
+				FLinearColor ColorHSV = Color.LinearRGBToHSV();
+				ColorHSV.B *= .5f;
+				return ColorHSV.HSVToLinearRGB();
+			};
+
+		int32 ColorOffset = 0;
+
+		for (const FClothingSimulationCloth* const Cloth : Solver->GetCloths())
+		{
+			const int32 Offset = Cloth->GetOffset(Solver);
+			if (Offset == INDEX_NONE)
+			{
+				continue;
+			}
+
+			// Draw constraints
+			const FClothConstraints& ClothConstraints = Solver->GetClothConstraints(Offset);
+			const TConstArrayView<Softs::FSolverVec3> Positions = Cloth->GetParticlePositions(Solver);
+
+			if (const Softs::FPBDLongRangeConstraints* const LongRangeConstraints = ClothConstraints.GetLongRangeConstraints().Get())
+			{
+				const TArray<TConstArrayView<Softs::FPBDLongRangeConstraints::FTether>>& Tethers = LongRangeConstraints->GetTethers();
+
+				for (int32 BatchIndex = 0; BatchIndex < Tethers.Num(); ++BatchIndex)
+				{
+					const FLinearColor Color = PseudoRandomColor(ColorOffset + BatchIndex);
+					const FLinearColor DarkenedColor = Darken(Color);
+
+					const TConstArrayView<Softs::FPBDLongRangeConstraints::FTether>& TetherBatch = Tethers[BatchIndex];
+
+					// Draw tethers
+					for (const Softs::FPBDLongRangeConstraints::FTether& Tether : TetherBatch)
+					{
+						const int32 KinematicIndex = LongRangeConstraints->GetStartIndex(Tether);
+						const int32 DynamicIndex = LongRangeConstraints->GetEndIndex(Tether);
+						const FReal TargetLength = LongRangeConstraints->GetTargetLength(Tether);
+
+						const FVec3 Pos0 = FVec3(Positions[KinematicIndex]) + LocalSpaceLocation;
+						const FVec3 Pos1 = FVec3(Positions[DynamicIndex]) + LocalSpaceLocation;
+
+						DrawLine(PDI, Pos0, Pos1, Color);
+#if WITH_EDITOR
+						DrawPoint(PDI, Pos1, Color, ClothMaterialVertex);
+#else
+						DrawPoint(nullptr, Pos1, Color, nullptr);
+#endif
+						FVec3 Direction = Pos1 - Pos0;
+						const float Length = Direction.SafeNormalize();
+						if (Length > SMALL_NUMBER)
+						{
+							const FVec3 Pos2 = Pos1 + Direction * (TargetLength - Length);
+							DrawLine(PDI, Pos1, Pos2, DarkenedColor);
+						}
+					}
+				}
+
+				// Rotate the colors for each cloth
+				ColorOffset += Tethers.Num();
+			}
+		}
+	}
+
+	void FVisualization::DrawWindAndPressureForces(FPrimitiveDrawInterface* PDI) const
+	{
+		using namespace ::Chaos;
+
+		if (!Solver)
+		{
+			return;
+		}
+
+		constexpr FReal ForceLength = (FReal)10.;
+		const FVec3& LocalSpaceLocation = Solver->GetLocalSpaceLocation();
+
+		for (const FClothingSimulationCloth* const Cloth : Solver->GetCloths())
+		{
+			const int32 Offset = Cloth->GetOffset(Solver);
+			if (Offset == INDEX_NONE)
+			{
+				continue;
+			}
+
+			const Softs::FVelocityAndPressureField& VelocityField = Solver->GetWindVelocityAndPressureField(Cloth->GetGroupId());
+
+			const TConstArrayView<TVec3<int32>>& Elements = VelocityField.GetElements();
+			const TConstArrayView<Softs::FSolverVec3> Forces = VelocityField.GetForces();
+			const TConstArrayView<Softs::FSolverVec3> Positions = Cloth->GetParticlePositions(Solver);
+			const TConstArrayView<Softs::FSolverReal> InvMasses = Cloth->GetParticleInvMasses(Solver);
+			check(InvMasses.Num() == Positions.Num());
+
+			for (int32 ElementIndex = 0; ElementIndex < Elements.Num(); ++ElementIndex)
+			{
+				const TVec3<int32>& Element = Elements[ElementIndex];
+				const FVec3 Position = LocalSpaceLocation + (
+					FVec3(Positions[Element.X - Offset]) +
+					FVec3(Positions[Element.Y - Offset]) +
+					FVec3(Positions[Element.Z - Offset])) / (FReal)3.;
+
+				const bool bIsKinematic0 = !InvMasses[Element.X - Offset];
+				const bool bIsKinematic1 = !InvMasses[Element.Y - Offset];
+				const bool bIsKinematic2 = !InvMasses[Element.Z - Offset];
+				const bool bIsKinematic = bIsKinematic0 || bIsKinematic1 || bIsKinematic2;
+
+				const FVec3 Force = FVec3(Forces[ElementIndex]) * ForceLength;
+				DrawLine(PDI, Position, Position + Force, bIsKinematic ? FColor::Cyan : FColor::Green);
+			}
+		}
+	}
+
+	void FVisualization::DrawLocalSpace(FPrimitiveDrawInterface* PDI) const
+	{
+		using namespace ::Chaos;
+
+		if (!Solver)
+		{
+			return;
+		}
+
+		// Draw local space
+		DrawCoordinateSystem(PDI, FQuat::Identity, Solver->GetLocalSpaceLocation());
+
+		// Draw reference spaces
+		for (const FClothingSimulationCloth* const Cloth : Solver->GetCloths())
+		{
+			if (Cloth->GetOffset(Solver) == INDEX_NONE)
+			{
+				continue;
+			}
+			const FRigidTransform3& ReferenceSpaceTransform = Cloth->GetReferenceSpaceTransform();
+			DrawCoordinateSystem(PDI, ReferenceSpaceTransform.GetRotation(), ReferenceSpaceTransform.GetLocation());
+		}
+	}
+
+	void FVisualization::DrawSelfCollision(FPrimitiveDrawInterface* PDI) const
+	{
+		using namespace ::Chaos;
+
+		if (!Solver)
+		{
+			return;
+		}
+
+		const FVec3& LocalSpaceLocation = Solver->GetLocalSpaceLocation();
+
+		for (const FClothingSimulationCloth* const Cloth : Solver->GetCloths())
+		{
+			const int32 Offset = Cloth->GetOffset(Solver);
+			if (Offset == INDEX_NONE)
+			{
+				continue;
+			}
+
+			// Draw constraints
+			const FClothConstraints& ClothConstraints = Solver->GetClothConstraints(Offset);
+
+			if (const Softs::FPBDCollisionSpringConstraints* const SelfCollisionConstraints = ClothConstraints.GetSelfCollisionConstraints().Get())
+			{
+				const TConstArrayView<Softs::FSolverVec3> Positions = Cloth->GetParticlePositions(Solver);
+				const TArray<TVec4<int32>>& Constraints = SelfCollisionConstraints->GetConstraints();
+				const TArray<Softs::FSolverVec3>& Barys = SelfCollisionConstraints->GetBarys();
+				const FReal Thickness = (FReal)SelfCollisionConstraints->GetThickness();
+				const FReal Height = Thickness + Thickness;
+				const TArray<bool>& FlipNormals = SelfCollisionConstraints->GetFlipNormals();
+
+				for (int32 Index = 0; Index < Constraints.Num(); ++Index)
+				{
+					const TVec4<int32>& Constraint = Constraints[Index];
+					const FVec3 Bary(Barys[Index]);
+
+					// Constraint index includes Offset, but so does Positions.
+					const FVector P = LocalSpaceLocation + FVector(Positions[Constraint[0] - Offset]);
+					const FVector P0 = LocalSpaceLocation + FVector(Positions[Constraint[1] - Offset]);
+					const FVector P1 = LocalSpaceLocation + FVector(Positions[Constraint[2] - Offset]);
+					const FVector P2 = LocalSpaceLocation + FVector(Positions[Constraint[3] - Offset]);
+
+					const FVector Pos0 = P0 * Bary[0] + P1 * Bary[1] + P2 * Bary[2];
+
+					static const FLinearColor Brown(0.1f, 0.05f, 0.f);
+					static const FLinearColor Red(0.3f, 0.f, 0.f);
+					const FTriangle Triangle(P0, P1, P2);
+					const FVector Normal = FlipNormals[Index] ? -Triangle.GetNormal() : Triangle.GetNormal();
+
+					// Draw point to surface line (=normal)
+					const FVector Pos1 = Pos0 + Height * Normal;
+					DrawPoint(PDI, Pos0, Brown, nullptr, 2.f);
+					DrawLine(PDI, Pos0, Pos1, FlipNormals[Index] ? Red : Brown);
+
+					// Draw pushup to point
+					static const FLinearColor Orange(0.3f, 0.15f, 0.f);
+					DrawPoint(PDI, P, Orange, nullptr, 2.f);
+					DrawLine(PDI, Pos1, P, Orange);
+				}
+			}
+		}
+	}
+
+	void FVisualization::DrawSelfIntersection(FPrimitiveDrawInterface* PDI) const
+	{
+		using namespace ::Chaos;
+
+		if (!Solver)
+		{
+			return;
+		}
+
+		const FVec3& LocalSpaceLocation = Solver->GetLocalSpaceLocation();
+
+		for (const FClothingSimulationCloth* const Cloth : Solver->GetCloths())
+		{
+			const int32 Offset = Cloth->GetOffset(Solver);
+			if (Offset == INDEX_NONE)
+			{
+				continue;
+			}
+
+			const FClothConstraints& ClothConstraints = Solver->GetClothConstraints(Offset);
+
+			static const FLinearColor Red(1.f, 0.f, 0.f);
+			static const FLinearColor White(1.f, 1.f, 1.f);
+			static const FLinearColor Black(0.f, 0.f, 0.f);
+			static const FLinearColor Teal(0.f, 0.5f, 0.5f);
+			static const FLinearColor Green(0.f, 1.f, 0.f);
+
+			if (const Softs::FPBDTriangleMeshCollisions* const SelfCollisionInit = ClothConstraints.GetSelfCollisionInit().Get())
+			{
+				const TConstArrayView<Softs::FSolverVec3> Positions = Cloth->GetParticlePositions(Solver);
+				const FTriangleMesh& TriangleMesh = Cloth->GetTriangleMesh(Solver);
+
+				// Draw contours
+				const TArray<TArray<Softs::FPBDTriangleMeshCollisions::FBarycentricPoint>>& ContourPoints = SelfCollisionInit->GetIntersectionContourPoints();
+				const TArray<Softs::FPBDTriangleMeshCollisions::FContourType>& ContourTypes = SelfCollisionInit->GetIntersectionContourTypes();
+				check(ContourPoints.Num() == ContourTypes.Num());
+
+				static const FLinearColor ColorsForType[(int8)Softs::FPBDTriangleMeshCollisions::FContourType::Count] =
+				{
+					Teal,
+					Red,
+					White,
+					Black
+				};
+				for( int32 ContourIndex = 0; ContourIndex < ContourPoints.Num(); ++ContourIndex)
+				{
+					const TArray<Softs::FPBDTriangleMeshCollisions::FBarycentricPoint>& Contour = ContourPoints[ContourIndex];
+					const FLinearColor& ContourColor = ColorsForType[(int8)ContourTypes[ContourIndex]];
+					for (int32 PointIdx = 0; PointIdx < Contour.Num() - 1; ++PointIdx)
+					{
+						const Softs::FPBDTriangleMeshCollisions::FBarycentricPoint& Point0 = Contour[PointIdx];
+						const FVector EndPoint0 = LocalSpaceLocation + (1.f - Point0.Bary[0] - Point0.Bary[1]) * Positions[Point0.Vertices[0] - Offset] + Point0.Bary[0] * Positions[Point0.Vertices[1] - Offset] + Point0.Bary[1] * Positions[Point0.Vertices[2] - Offset];
+						const Softs::FPBDTriangleMeshCollisions::FBarycentricPoint& Point1 = Contour[PointIdx+1];
+						const FVector EndPoint1 = LocalSpaceLocation + (1.f - Point1.Bary[0] - Point1.Bary[1]) * Positions[Point1.Vertices[0] - Offset] + Point1.Bary[0] * Positions[Point1.Vertices[1] - Offset] + Point1.Bary[1] * Positions[Point1.Vertices[2] - Offset];
+						DrawLine(PDI, EndPoint0, EndPoint1, ContourColor);
+						DrawPoint(PDI, EndPoint0, ContourColor, nullptr, 1.f);
+						DrawPoint(PDI, EndPoint1, ContourColor, nullptr, 1.f);
+					}
+				}
+
+				// Draw GIA colors
+				const TConstArrayView<Softs::FPBDTriangleMeshCollisions::FGIAColor>& VertexGIAColors = SelfCollisionInit->GetVertexGIAColors();
+				static const FLinearColor Gray(0.5f, 0.5f, 0.5f);
+				if (VertexGIAColors.Num())
+				{
+					const TArray<TVec3<int32>>& Elements = TriangleMesh.GetElements();
+					for (int32 ParticleIdx = Offset; ParticleIdx < VertexGIAColors.Num(); ++ParticleIdx)
+					{
+						if (VertexGIAColors[ParticleIdx].ContourIndexBits)
+						{
+							const bool bIsLoop = VertexGIAColors[ParticleIdx].IsLoop();
+							const bool bAnyWhite = (VertexGIAColors[ParticleIdx].ContourIndexBits & ~VertexGIAColors[ParticleIdx].ColorBits);
+							const bool bAnyBlack = (VertexGIAColors[ParticleIdx].ContourIndexBits & VertexGIAColors[ParticleIdx].ColorBits);
+							const FLinearColor& VertColor = bIsLoop ? Red : (bAnyWhite && bAnyBlack) ? Gray : bAnyWhite ? White : Black;
+						
+							DrawPoint(PDI, LocalSpaceLocation + Positions[ParticleIdx - Offset], VertColor, nullptr, 5.f);
+						}
+					}
+				}
+				const TArray<Softs::FPBDTriangleMeshCollisions::FGIAColor>& TriangleGIAColors = SelfCollisionInit->GetTriangleGIAColors();
+				if (TriangleGIAColors.Num() == TriangleMesh.GetNumElements())
+				{
+					const TArray<TVec3<int32>>& Elements = TriangleMesh.GetElements();
+					for (int32 TriangleIdx = 0; TriangleIdx < TriangleGIAColors.Num(); ++TriangleIdx)
+					{
+						if (TriangleGIAColors[TriangleIdx].ContourIndexBits)
+						{
+							const bool bIsLoop = TriangleGIAColors[TriangleIdx].IsLoop();
+							const bool bAnyWhite = (TriangleGIAColors[TriangleIdx].ContourIndexBits & ~TriangleGIAColors[TriangleIdx].ColorBits);
+							const bool bAnyBlack = (TriangleGIAColors[TriangleIdx].ContourIndexBits & TriangleGIAColors[TriangleIdx].ColorBits);
+							const FLinearColor& TriColor = bIsLoop ? Red : (bAnyWhite && bAnyBlack) ? Gray : bAnyWhite ? White : Black;
+							DrawLine(PDI, LocalSpaceLocation + Positions[Elements[TriangleIdx][0] - Offset], LocalSpaceLocation + Positions[Elements[TriangleIdx][1] - Offset], TriColor);
+							DrawLine(PDI, LocalSpaceLocation + Positions[Elements[TriangleIdx][1] - Offset], LocalSpaceLocation + Positions[Elements[TriangleIdx][2] - Offset], TriColor);
+							DrawLine(PDI, LocalSpaceLocation + Positions[Elements[TriangleIdx][0] - Offset], LocalSpaceLocation + Positions[Elements[TriangleIdx][2] - Offset], TriColor);
+						}
+					}
+				}
+
+				// Draw contour minimization gradients
+				const TArray<Softs::FPBDTriangleMeshCollisions::FContourMinimizationIntersection>& ContourMinimizationIntersections = SelfCollisionInit->GetContourMinimizationIntersections();
+				constexpr FReal MaxDrawImpulse = 1.;
+				constexpr FReal RegularizeEpsilonSq = 1.;
+				for (const Softs::FPBDTriangleMeshCollisions::FContourMinimizationIntersection& Intersection : ContourMinimizationIntersections)
+				{
+					Softs::FSolverReal GradientLength;
+					Softs::FSolverVec3 GradientDir;
+					Intersection.GlobalGradientVector.ToDirectionAndLength(GradientDir, GradientLength);
+					const FVector Delta = FVector(GradientDir) * MaxDrawImpulse * GradientLength * FMath::InvSqrt(GradientLength * GradientLength + RegularizeEpsilonSq);
+
+					const FVector EdgeCenter = LocalSpaceLocation + .5 * (Positions[Intersection.EdgeVertices[0] - Offset] + Positions[Intersection.EdgeVertices[1] - Offset]);
+					const FVector TriCenter = LocalSpaceLocation + (Positions[Intersection.FaceVertices[0] - Offset] + Positions[Intersection.FaceVertices[1] - Offset] + Positions[Intersection.FaceVertices[2] - Offset]) / 3.;
+
+					DrawPoint(PDI, EdgeCenter, Green, nullptr, 2.f);
+					DrawLine(PDI, EdgeCenter, EdgeCenter + Delta, Green);
+					DrawPoint(PDI, TriCenter, Green, nullptr, 2.f);
+					DrawLine(PDI, TriCenter, TriCenter - Delta, Green);
+				}
+			}
+		}
+	}
+}  // End namespace UE::Chaos::Cloth
+#else  // #if CHAOS_DEBUG_DRAW
+namespace UE::Chaos::Cloth
+{
+	FVisualization::FVisualization(const ::Chaos::FClothingSimulationSolver* /*InSolver*/) {}
+
+	FVisualization::~FVisualization() = default;
+}  // End namespace UE::Chaos::Cloth
+#endif  // #if CHAOS_DEBUG_DRAW
diff --git a/Engine/Plugins/Experimental/ChaosCloth/Source/Chaos/Private/ChaosCloth/ChaosClothingSimulation.cpp b/Engine/Plugins/Experimental/ChaosCloth/Source/Chaos/Private/ChaosCloth/ChaosClothingSimulation.cpp
--- a/Engine/Plugins/Experimental/ChaosCloth/Source/Chaos/Private/ChaosCloth/ChaosClothingSimulation.cpp
+++ b/Engine/Plugins/Experimental/ChaosCloth/Source/Chaos/Private/ChaosCloth/ChaosClothingSimulation.cpp
@@ -5,7 +5,6 @@
 #include "ClothingSimulation.h"
 #include "ClothingAsset.h"
 #include "ChaosCloth/ChaosClothConfig.h"
-#include "ChaosCloth/ChaosWeightMapTarget.h"
 
 #include "ChaosCloth/ChaosClothingSimulationSolver.h"
 #include "ChaosCloth/ChaosClothingSimulationSkeletalMesh.h"
@@ -14,40 +13,9 @@
 
 #include "PhysicsField/PhysicsFieldComponent.h"
 #include "PhysicsProxy/PerSolverFieldSystem.h"
-#include "DynamicMeshBuilder.h"
 
 #include "HAL/IConsoleManager.h"
 
-#if WITH_EDITOR || CHAOS_DEBUG_DRAW
-#include "Chaos/Sphere.h"
-#include "Chaos/TaperedCylinder.h"
-#include "Chaos/Capsule.h"
-#include "Chaos/TaperedCapsule.h"
-#include "Chaos/Convex.h"
-#include "Chaos/Triangle.h"
-#include "Chaos/PBDSphericalConstraint.h"
-#include "Chaos/PBDAnimDriveConstraint.h"
-#include "Chaos/PBDLongRangeConstraints.h"
-#include "Chaos/PBDSpringConstraints.h"
-#include "Chaos/PBDCollisionSpringConstraints.h"
-#include "Chaos/PBDTriangleMeshCollisions.h"
-#include "Chaos/PBDBendingConstraints.h"
-#include "Chaos/XPBDSpringConstraints.h"
-#include "Chaos/XPBDBendingConstraints.h"
-#include "Chaos/VelocityField.h"
-#endif  // #if WITH_EDITOR || CHAOS_DEBUG_DRAW
-
-#if WITH_EDITOR
-#include "Materials/Material.h"
-#include "Engine/Canvas.h"  // For debug draw text
-#include "CanvasItem.h"     //
-#include "Engine/Engine.h"  //
-#endif  // #if WITH_EDITOR
-
-#if CHAOS_DEBUG_DRAW
-#include "Chaos/DebugDrawQueue.h"
-#endif  // #if CHAOS_DEBUG_DRAW
-
 #if !UE_BUILD_SHIPPING
 #include "FramePro/FramePro.h"
 #else
@@ -299,10 +267,6 @@ FClothingSimulation::FClothingSimulation()
 #endif
 	, bHasInvalidReferenceBoneTransforms(false)
 {
-#if WITH_EDITOR
-	DebugClothMaterial = LoadObject<UMaterial>(nullptr, TEXT("/Engine/EditorMaterials/Cloth/CameraLitDoubleSided.CameraLitDoubleSided"), nullptr, LOAD_None, nullptr);  // LOAD_EditorOnly
-	DebugClothMaterialVertex = LoadObject<UMaterial>(nullptr, TEXT("/Engine/EditorMaterials/WidgetVertexColorMaterial"), nullptr, LOAD_None, nullptr);  // LOAD_EditorOnly
-#endif  // #if WITH_EDITOR
 #if !(UE_BUILD_SHIPPING || UE_BUILD_TEST)
 	if (!ClothingSimulationConsole::Command)
 	{
@@ -319,11 +283,15 @@ void FClothingSimulation::Initialize()
 	// Create solver
 	Solver = MakeUnique<FClothingSimulationSolver>();
 
+	// Assign the solver to the visualization
+	Visualization.SetSolver(Solver.Get());
+
 	ResetStats();
 }
 
 void FClothingSimulation::Shutdown()
 {
+	Visualization.SetSolver(nullptr);
 	Solver.Reset();
 	Meshes.Reset();
 	Cloths.Reset();
@@ -931,1399 +899,4 @@ FClothingSimulationCloth* Chaos::FClothingSimulation::GetCloth(int32 ClothId)
 	return Cloth ? Cloth->Get(): nullptr;
 }
 
-#if WITH_EDITOR
-void FClothingSimulation::AddReferencedObjects(FReferenceCollector& Collector)
-{
-	Collector.AddReferencedObject(DebugClothMaterial);
-}
-
-void FClothingSimulation::DebugDrawPhysMeshShaded(FPrimitiveDrawInterface* PDI) const
-{
-	if (!DebugClothMaterial)
-	{
-		return;
-	}
-
-	FDynamicMeshBuilder MeshBuilder(PDI->View->GetFeatureLevel());
-	int32 VertexIndex = 0;
-
-	for (const FClothingSimulationCloth* const Cloth : Solver->GetCloths())
-	{
-		const int32 Offset = Cloth->GetOffset(Solver.Get());
-		if (Offset == INDEX_NONE)
-		{
-			continue;
-		}
-
-		const TConstArrayView<TVec3<int32>> Elements = Cloth->GetTriangleMesh(Solver.Get()).GetElements();
-		const TConstArrayView<Softs::FSolverVec3> Positions = Cloth->GetParticlePositions(Solver.Get());
-		const TConstArrayView<Softs::FSolverReal> InvMasses = Cloth->GetParticleInvMasses(Solver.Get());
-		check(InvMasses.Num() == Positions.Num());
-
-		for (int32 ElementIndex = 0; ElementIndex < Elements.Num(); ++ElementIndex, VertexIndex += 3)
-		{
-			const TVec3<int32>& Element = Elements[ElementIndex];
-
-			const FVector3f Pos0(Positions[Element.X - Offset]); // TODO: Triangle Mesh shouldn't really be solver dependent (ie not use an offset)
-			const FVector3f Pos1(Positions[Element.Y - Offset]);
-			const FVector3f Pos2(Positions[Element.Z - Offset]);
-
-			const FVector3f Normal = FVector3f::CrossProduct(Pos2 - Pos0, Pos1 - Pos0).GetSafeNormal();
-			const FVector3f Tangent = ((Pos1 + Pos2) * 0.5f - Pos0).GetSafeNormal();
-
-			const bool bIsKinematic0 = (InvMasses[Element.X - Offset] == (Softs::FSolverReal)0.);
-			const bool bIsKinematic1 = (InvMasses[Element.Y - Offset] == (Softs::FSolverReal)0.);
-			const bool bIsKinematic2 = (InvMasses[Element.Z - Offset] == (Softs::FSolverReal)0.);
-
-			MeshBuilder.AddVertex(FDynamicMeshVertex(Pos0, Tangent, Normal, FVector2f(0.f, 0.f), bIsKinematic0 ? FColor::Purple : FColor::White));
-			MeshBuilder.AddVertex(FDynamicMeshVertex(Pos1, Tangent, Normal, FVector2f(0.f, 1.f), bIsKinematic1 ? FColor::Purple : FColor::White));
-			MeshBuilder.AddVertex(FDynamicMeshVertex(Pos2, Tangent, Normal, FVector2f(1.f, 1.f), bIsKinematic2 ? FColor::Purple : FColor::White));
-			MeshBuilder.AddTriangle(VertexIndex, VertexIndex + 1, VertexIndex + 2);
-		}
-	}
-
-	FMatrix LocalSimSpaceToWorld(FMatrix::Identity);
-	LocalSimSpaceToWorld.SetOrigin(Solver->GetLocalSpaceLocation());
-	MeshBuilder.Draw(PDI, LocalSimSpaceToWorld, DebugClothMaterial->GetRenderProxy(), SDPG_World, false, false);
-}
-
-static void DrawText(FCanvas* Canvas, const FSceneView* SceneView, const FVector& Pos, const FText& Text, const FLinearColor& Color)
-{
-	FVector2D PixelLocation;
-	if (SceneView->WorldToPixel(Pos, PixelLocation))
-	{
-		FCanvasTextItem TextItem(PixelLocation, Text, GEngine->GetSmallFont(), Color);
-		TextItem.Scale = FVector2D::UnitVector;
-		TextItem.EnableShadow(FLinearColor::Black);
-		TextItem.Draw(Canvas);
-	}
-}
-
-void FClothingSimulation::DebugDrawParticleIndices(FCanvas* Canvas, const FSceneView* SceneView) const
-{
-	static const FLinearColor DynamicColor = FColor::White;
-	static const FLinearColor KinematicColor = FColor::Purple;
-
-	const FVec3& LocalSpaceLocation = Solver->GetLocalSpaceLocation();
-
-	for (const FClothingSimulationCloth* const Cloth : Solver->GetCloths())
-	{
-		const int32 Offset = Cloth->GetOffset(Solver.Get());
-		if (Offset == INDEX_NONE)
-		{
-			continue;
-		}
-
-		const TConstArrayView<Softs::FSolverVec3> Positions = Cloth->GetParticlePositions(Solver.Get());
-		const TConstArrayView<Softs::FSolverReal> InvMasses = Cloth->GetParticleInvMasses(Solver.Get());
-		check(InvMasses.Num() == Positions.Num());
-
-		for (int32 Index = 0; Index < Positions.Num(); ++Index)
-		{
-			const FVector Position = LocalSpaceLocation + FVector(Positions[Index]);
-
-			const FText Text = FText::AsNumber(Offset + Index);
-			DrawText(Canvas, SceneView, Position, Text, InvMasses[Index] == (Softs::FSolverReal)0. ? KinematicColor : DynamicColor);
-		}
-	}
-}
-
-void FClothingSimulation::DebugDrawElementIndices(FCanvas* Canvas, const FSceneView* SceneView) const
-{
-	static const FLinearColor DynamicColor = FColor::White;
-	static const FLinearColor KinematicColor = FColor::Purple;
-
-	const FVec3& LocalSpaceLocation = Solver->GetLocalSpaceLocation();
-
-	for (const FClothingSimulationCloth* const Cloth : Solver->GetCloths())
-	{
-		const int32 Offset = Cloth->GetOffset(Solver.Get());
-		if (Offset == INDEX_NONE)
-		{
-			continue;
-		}
-
-		const TArray<TVec3<int32>>& Elements = Cloth->GetTriangleMesh(Solver.Get()).GetElements();
-		const TConstArrayView<Softs::FSolverVec3> Positions = Cloth->GetParticlePositions(Solver.Get());
-		const TConstArrayView<Softs::FSolverReal> InvMasses = Cloth->GetParticleInvMasses(Solver.Get());
-		check(InvMasses.Num() == Positions.Num());
-
-		for (int32 Index = 0; Index < Elements.Num(); ++Index)
-		{
-			const TVec3<int32>& Element = Elements[Index];
-			const FVector Position = LocalSpaceLocation + FVector(Positions[Element[0]] + Positions[Element[1]] + Positions[Element[2]]) / (FReal)3.;
-
-			const bool bIsKinematic0 = (InvMasses[Element.X - Offset] == (Softs::FSolverReal)0.);
-			const bool bIsKinematic1 = (InvMasses[Element.Y - Offset] == (Softs::FSolverReal)0.);
-			const bool bIsKinematic2 = (InvMasses[Element.Z - Offset] == (Softs::FSolverReal)0.);
-			const FLinearColor& Color = (bIsKinematic0 && bIsKinematic1 && bIsKinematic2) ? KinematicColor : DynamicColor;
-			const FText Text = FText::AsNumber(Index);
-			DrawText(Canvas, SceneView, Position, Text, Color);
-		}
-	}
-}
-
-void FClothingSimulation::DebugDrawMaxDistanceValues(FCanvas* Canvas, const FSceneView* SceneView) const
-{
-	static const FLinearColor DynamicColor = FColor::White;
-	static const FLinearColor KinematicColor = FColor::Purple;
-
-	FNumberFormattingOptions NumberFormattingOptions;
-	NumberFormattingOptions.AlwaysSign = false;
-	NumberFormattingOptions.UseGrouping = false;
-	NumberFormattingOptions.RoundingMode = ERoundingMode::HalfFromZero;
-	NumberFormattingOptions.MinimumIntegralDigits = 1;
-	NumberFormattingOptions.MaximumIntegralDigits = 6;
-	NumberFormattingOptions.MinimumFractionalDigits = 2;
-	NumberFormattingOptions.MaximumFractionalDigits = 2;
-
-	const FVec3& LocalSpaceLocation = Solver->GetLocalSpaceLocation();
-
-	for (const FClothingSimulationCloth* const Cloth : Solver->GetCloths())
-	{
-		const int32 Offset = Cloth->GetOffset(Solver.Get());
-		if (Offset == INDEX_NONE)
-		{
-			continue;
-		}
-
-		const TConstArrayView<FRealSingle>& MaxDistances = Cloth->GetWeightMaps(Solver.Get())[(int32)EChaosWeightMapTarget::MaxDistance];
-		if (!MaxDistances.Num())
-		{
-			continue;
-		}
-
-		const TConstArrayView<Softs::FSolverVec3> Positions = Cloth->GetAnimationPositions(Solver.Get());
-		const TConstArrayView<Softs::FSolverReal> InvMasses = Cloth->GetParticleInvMasses(Solver.Get());
-		check(MaxDistances.Num() == Positions.Num());
-		check(MaxDistances.Num() == InvMasses.Num());
-
-		for (int32 Index = 0; Index < MaxDistances.Num(); ++Index)
-		{
-			const FReal MaxDistance = MaxDistances[Index];
-			const FVector Position = LocalSpaceLocation + FVector(Positions[Index]);
-
-			const FText Text = FText::AsNumber(MaxDistance, &NumberFormattingOptions);
-			DrawText(Canvas, SceneView, Position, Text, InvMasses[Index] == (Softs::FSolverReal)0. ? KinematicColor : DynamicColor);
-		}
-	}
-}
-#endif  // #if WITH_EDITOR
-
-#if WITH_EDITOR || CHAOS_DEBUG_DRAW
-static void DrawPoint(FPrimitiveDrawInterface* PDI, const FVector& Pos, const FLinearColor& Color, UMaterial* DebugClothMaterialVertex, const float Thickness = 1.f)  // Use color or material
-{
-#if CHAOS_DEBUG_DRAW
-	if (!PDI)
-	{
-		FDebugDrawQueue::GetInstance().DrawDebugPoint(Pos, Color.ToFColor(true), false, KINDA_SMALL_NUMBER, SDPG_Foreground, Thickness);
-		return;
-	}
-#endif
-#if WITH_EDITOR
-	if (DebugClothMaterialVertex)
-	{
-		const FMatrix& ViewMatrix = PDI->View->ViewMatrices.GetViewMatrix();
-		const FVector XAxis = ViewMatrix.GetColumn(0); // Just using transpose here (orthogonal transform assumed)
-		const FVector YAxis = ViewMatrix.GetColumn(1);
-		DrawDisc(PDI, Pos, XAxis, YAxis, Color.ToFColor(true), 0.5f, 10, DebugClothMaterialVertex->GetRenderProxy(), SDPG_World);
-	}
-	else
-	{
-		PDI->DrawPoint(Pos, Color, Thickness, SDPG_World);
-	}
-#endif
-}
-
-static void DrawLine(FPrimitiveDrawInterface* PDI, const FVector& Pos0, const FVector& Pos1, const FLinearColor& Color)
-{
-#if CHAOS_DEBUG_DRAW
-	if (!PDI)
-	{
-		FDebugDrawQueue::GetInstance().DrawDebugLine(Pos0, Pos1, Color.ToFColor(true), false, KINDA_SMALL_NUMBER, SDPG_Foreground, 0.f);
-		return;
-	}
-#endif
-#if WITH_EDITOR
-	PDI->DrawLine(Pos0, Pos1, Color, SDPG_World, 0.0f, 0.001f);
-#endif
-}
-
-static void DrawArc(FPrimitiveDrawInterface* PDI, const FVector& Base, const FVector& X, const FVector& Y, FReal MinAngle, FReal MaxAngle, FReal Radius, const FLinearColor& Color)
-{
-	static const int32 Sections = 10;
-	const FReal AngleStep = FMath::DegreesToRadians((MaxAngle - MinAngle) / (FReal)Sections);
-	FReal CurrentAngle = FMath::DegreesToRadians(MinAngle);
-	FVector LastVertex = Base + Radius * (FMath::Cos(CurrentAngle) * X + FMath::Sin(CurrentAngle) * Y);
-
-	for(int32 i = 0; i < Sections; i++)
-	{
-		CurrentAngle += AngleStep;
-		const FVector ThisVertex = Base + Radius * (FMath::Cos(CurrentAngle) * X + FMath::Sin(CurrentAngle) * Y);
-		DrawLine(PDI, LastVertex, ThisVertex, Color);
-		LastVertex = ThisVertex;
-	}
-}
-
-static void DrawSphere(FPrimitiveDrawInterface* PDI, const TSphere<FReal, 3>& Sphere, const FQuat& Rotation, const FVector& Position, const FLinearColor& Color)
-{
-	const FReal Radius = Sphere.GetRadius();
-	const FVec3 Center = Position + Rotation.RotateVector(Sphere.GetCenter());
-#if CHAOS_DEBUG_DRAW
-	if (!PDI)
-	{
-		FDebugDrawQueue::GetInstance().DrawDebugSphere(Center, Radius, 12, Color.ToFColor(true), false, KINDA_SMALL_NUMBER, SDPG_Foreground, 0.f);
-		return;
-	}
-#endif
-#if WITH_EDITOR
-	const FTransform Transform(Rotation, Center);
-	DrawWireSphere(PDI, Transform, Color, Radius, 12, SDPG_World, 0.0f, 0.001f, false);
-#endif
-}
-
-static void DrawBox(FPrimitiveDrawInterface* PDI, const FAABB3& Box, const FQuat& Rotation, const FVector& Position, const FLinearColor& Color)
-{
-#if CHAOS_DEBUG_DRAW
-	if (!PDI)
-	{
-		const FVec3 Center = Position + Rotation.RotateVector(Box.GetCenter());
-		FDebugDrawQueue::GetInstance().DrawDebugBox(Center, Box.Extents() * 0.5f, Rotation, Color.ToFColor(true), false, KINDA_SMALL_NUMBER, SDPG_Foreground, 0.f);
-		return;
-	}
-#endif
-#if WITH_EDITOR
-	const FMatrix BoxToWorld = FTransform(Rotation, Position).ToMatrixNoScale();
-	DrawWireBox(PDI, BoxToWorld, FBox(Box.Min(), Box.Max()), Color, SDPG_World, 0.0f, 0.001f, false);
-#endif
-}
-
-static void DrawCapsule(FPrimitiveDrawInterface* PDI, const FCapsule& Capsule, const FQuat& Rotation, const FVector& Position, const FLinearColor& Color)
-{
-	const FReal Radius = Capsule.GetRadius();
-	const FReal HalfHeight = Capsule.GetHeight() * 0.5f + Radius;
-	const FVec3 Center = Position + Rotation.RotateVector(Capsule.GetCenter());
-#if CHAOS_DEBUG_DRAW
-	if (!PDI)
-	{
-		const FQuat Orientation = FQuat::FindBetweenNormals(FVec3::UpVector, Capsule.GetAxis());
-		FDebugDrawQueue::GetInstance().DrawDebugCapsule(Center, HalfHeight, Radius, Rotation * Orientation, Color.ToFColor(true), false, KINDA_SMALL_NUMBER, SDPG_Foreground, 0.f);
-		return;
-	}
-#endif
-#if WITH_EDITOR
-	const FVec3 Up = Capsule.GetAxis();
-	FVec3 Forward, Right;
-	Up.FindBestAxisVectors(Forward, Right);
-	const FVector X = Rotation.RotateVector(Forward);
-	const FVector Y = Rotation.RotateVector(Right);
-	const FVector Z = Rotation.RotateVector(Up);
-	DrawWireCapsule(PDI, Center, X, Y, Z, Color, Radius, HalfHeight, 12, SDPG_World, 0.0f, 0.001f, false);
-#endif
-}
-
-static void DrawTaperedCylinder(FPrimitiveDrawInterface* PDI, const FTaperedCylinder& TaperedCylinder, const FQuat& Rotation, const FVector& Position, const FLinearColor& Color)
-{
-	const FReal HalfHeight = TaperedCylinder.GetHeight() * 0.5f;
-	const FReal Radius1 = TaperedCylinder.GetRadius1();
-	const FReal Radius2 = TaperedCylinder.GetRadius2();
-	const FVector Position1 = Position + Rotation.RotateVector(TaperedCylinder.GetX1());
-	const FVector Position2 = Position + Rotation.RotateVector(TaperedCylinder.GetX2());
-	const FQuat Q = (Position2 - Position1).ToOrientationQuat();
-	const FVector I = Q.GetRightVector();
-	const FVector J = Q.GetUpVector();
-
-	static const int32 NumSides = 12;
-	static const FReal	AngleDelta = (FReal)2. * (FReal)PI / NumSides;
-	FVector	LastVertex1 = Position1 + I * Radius1;
-	FVector	LastVertex2 = Position2 + I * Radius2;
-
-	for (int32 SideIndex = 1; SideIndex <= NumSides; ++SideIndex)
-	{
-		const FReal Angle = AngleDelta * FReal(SideIndex);
-		const FVector ArcPos = I * FMath::Cos(Angle) + J * FMath::Sin(Angle);
-		const FVector Vertex1 = Position1 + ArcPos * Radius1;
-		const FVector Vertex2 = Position2 + ArcPos * Radius2;
-
-		DrawLine(PDI, LastVertex1, Vertex1, Color);
-		DrawLine(PDI, LastVertex2, Vertex2, Color);
-		DrawLine(PDI, LastVertex1, LastVertex2, Color);
-
-		LastVertex1 = Vertex1;
-		LastVertex2 = Vertex2;
-	}
-}
-
-static void DrawConvex(FPrimitiveDrawInterface* PDI, const FConvex& Convex, const FQuat& Rotation, const FVector& Position, const FLinearColor& Color)
-{
-	const TArray<FConvex::FPlaneType>& Planes = Convex.GetFaces();
-	for (int32 PlaneIndex1 = 0; PlaneIndex1 < Planes.Num(); ++PlaneIndex1)
-	{
-		const FConvex::FPlaneType& Plane1 = Planes[PlaneIndex1];
-
-		for (int32 PlaneIndex2 = PlaneIndex1 + 1; PlaneIndex2 < Planes.Num(); ++PlaneIndex2)
-		{
-			const FConvex::FPlaneType& Plane2 = Planes[PlaneIndex2];
-
-			// Find the two surface points that belong to both Plane1 and Plane2
-			uint32 ParticleIndex1 = INDEX_NONE;
-
-			const TArray<FConvex::FVec3Type>& Vertices = Convex.GetVertices();
-			for (int32 ParticleIndex = 0; ParticleIndex < Vertices.Num(); ++ParticleIndex)
-			{
-				const FConvex::FVec3Type& X = Vertices[ParticleIndex];
-
-				if (FMath::Square(Plane1.SignedDistance(X)) < KINDA_SMALL_NUMBER && 
-					FMath::Square(Plane2.SignedDistance(X)) < KINDA_SMALL_NUMBER)
-				{
-					if (ParticleIndex1 != INDEX_NONE)
-					{
-						const FVector X1(Vertices[ParticleIndex1]);
-						const FVector X2(X);
-						const FVector Position1 = Position + Rotation.RotateVector(X1);
-						const FVector Position2 = Position + Rotation.RotateVector(X2);
-						DrawLine(PDI, Position1, Position2, Color);
-						break;
-					}
-					ParticleIndex1 = ParticleIndex;
-				}
-			}
-		}
-	}
-}
-
-static void DrawCoordinateSystem(FPrimitiveDrawInterface* PDI, const FQuat& Rotation, const FVector& Position)
-{
-	const FVector X = Rotation.RotateVector(FVector::ForwardVector) * 10.f;
-	const FVector Y = Rotation.RotateVector(FVector::RightVector) * 10.f;
-	const FVector Z = Rotation.RotateVector(FVector::UpVector) * 10.f;
-
-	DrawLine(PDI, Position, Position + X, FLinearColor::Red);
-	DrawLine(PDI, Position, Position + Y, FLinearColor::Green);
-	DrawLine(PDI, Position, Position + Z, FLinearColor::Blue);
-}
-
-#if CHAOS_DEBUG_DRAW
-void FClothingSimulation::DebugDrawBounds() const
-{
-	check(Solver);
-
-	// Calculate World space bounds
-	const FBoxSphereBounds Bounds = Solver->CalculateBounds();
-
-	// Draw bounds
-	DrawBox(nullptr, FAABB3(-Bounds.BoxExtent, Bounds.BoxExtent), FQuat::Identity, Bounds.Origin, FLinearColor(FColor::Purple));
-	DrawSphere(nullptr, TSphere<FReal, 3>(FVector::ZeroVector, Bounds.SphereRadius), FQuat::Identity, Bounds.Origin, FLinearColor(FColor::Orange));
-
-	// Draw individual cloth bounds
-	static const FLinearColor Color = FLinearColor(FColor::Purple).Desaturate(0.5);
-	for (const TUniquePtr<FClothingSimulationCloth>& Cloth : Cloths)
-	{
-		if (Cloth->GetOffset(Solver.Get()) == INDEX_NONE)
-		{
-			continue;
-		}
-
-		const FAABB3 BoundingBox = Cloth->CalculateBoundingBox(Solver.Get());
-		DrawBox(nullptr, BoundingBox, FQuat::Identity, FVector::ZeroVector, Color);  // TODO: Express bounds in local coordinates for LWC
-	}
-}
-
-void FClothingSimulation::DebugDrawGravity() const
-{
-	check(Solver);
-
-	// Draw gravity
-	for (const FClothingSimulationCloth* const Cloth : Solver->GetCloths())
-	{
-		if (Cloth->GetOffset(Solver.Get()) == INDEX_NONE)
-		{
-			continue;
-		}
-
-		const FAABB3 Bounds = Cloth->CalculateBoundingBox(Solver.Get());
-
-		const FVector Pos0 = Bounds.Center();
-		const FVector Pos1 = Pos0 + FVector(Cloth->GetGravity(Solver.Get()));
-		DrawLine(nullptr, Pos0, Pos1, FLinearColor::Red);
-	}
-}
-#endif  // #if CHAOS_DEBUG_DRAW
-
-void FClothingSimulation::DebugDrawPhysMeshWired(FPrimitiveDrawInterface* PDI) const
-{
-	check(Solver);
-
-	static const FLinearColor DynamicColor = FColor::White;
-	static const FLinearColor KinematicColor = FColor::Purple;
-
-	const FVec3& LocalSpaceLocation = Solver->GetLocalSpaceLocation();
-
-	for (const FClothingSimulationCloth* const Cloth : Solver->GetCloths())
-	{
-		const int32 Offset = Cloth->GetOffset(Solver.Get());
-		if (Offset == INDEX_NONE)
-		{
-			continue;
-		}
-
-		const TConstArrayView<TVec3<int32>> Elements = Cloth->GetTriangleMesh(Solver.Get()).GetElements();
-		const TConstArrayView<Softs::FSolverVec3> Positions = Cloth->GetParticlePositions(Solver.Get());
-		const TConstArrayView<Softs::FSolverReal> InvMasses = Cloth->GetParticleInvMasses(Solver.Get());
-		check(InvMasses.Num() == Positions.Num());
-
-		for (int32 ElementIndex = 0; ElementIndex < Elements.Num(); ++ElementIndex)
-		{
-			const auto& Element = Elements[ElementIndex];
-
-			const FVector Pos0 = LocalSpaceLocation + FVector(Positions[Element.X - Offset]); // TODO: Triangle Mesh shouldn't really be solver dependent (ie not use an offset)
-			const FVector Pos1 = LocalSpaceLocation + FVector(Positions[Element.Y - Offset]);
-			const FVector Pos2 = LocalSpaceLocation + FVector(Positions[Element.Z - Offset]);
-
-			const bool bIsKinematic0 = (InvMasses[Element.X - Offset] == (Softs::FSolverReal)0.);
-			const bool bIsKinematic1 = (InvMasses[Element.Y - Offset] == (Softs::FSolverReal)0.);
-			const bool bIsKinematic2 = (InvMasses[Element.Z - Offset] == (Softs::FSolverReal)0.);
-
-			DrawLine(PDI, Pos0, Pos1, bIsKinematic0 && bIsKinematic1 ? KinematicColor : DynamicColor);
-			DrawLine(PDI, Pos1, Pos2, bIsKinematic1 && bIsKinematic2 ? KinematicColor : DynamicColor);
-			DrawLine(PDI, Pos2, Pos0, bIsKinematic2 && bIsKinematic0 ? KinematicColor : DynamicColor);
-		}
-	}
-}
-
-void FClothingSimulation::DebugDrawAnimMeshWired(FPrimitiveDrawInterface* PDI) const
-{
-	check(Solver);
-
-	static const FLinearColor KinematicColor = FColor::Purple;
-
-	const FVec3& LocalSpaceLocation = Solver->GetLocalSpaceLocation();
-
-	for (const FClothingSimulationCloth* const Cloth : Solver->GetCloths())
-	{
-		const int32 Offset = Cloth->GetOffset(Solver.Get());
-		if (Offset == INDEX_NONE)
-		{
-			continue;
-		}
-
-		const TConstArrayView<TVec3<int32>> Elements = Cloth->GetTriangleMesh(Solver.Get()).GetElements();
-		const TConstArrayView<Softs::FSolverVec3> Positions = Cloth->GetAnimationPositions(Solver.Get());
-
-		for (int32 ElementIndex = 0; ElementIndex < Elements.Num(); ++ElementIndex)
-		{
-			const auto& Element = Elements[ElementIndex];
-
-			const FVector Pos0 = LocalSpaceLocation + FVector(Positions[Element.X - Offset]); // TODO: Triangle Mesh shouldn't really be solver dependent (ie not use an offset)
-			const FVector Pos1 = LocalSpaceLocation + FVector(Positions[Element.Y - Offset]);
-			const FVector Pos2 = LocalSpaceLocation + FVector(Positions[Element.Z - Offset]);
-
-			DrawLine(PDI, Pos0, Pos1, KinematicColor);
-			DrawLine(PDI, Pos1, Pos2, KinematicColor);
-			DrawLine(PDI, Pos2, Pos0, KinematicColor);
-		}
-	}
-}
-
-void FClothingSimulation::DebugDrawAnimNormals(FPrimitiveDrawInterface* PDI) const
-{
-	check(Solver);
-
-	static const FLinearColor KinematicColor = FColor::Magenta;
-	constexpr FReal NormalLength = (FReal)20.;
-
-	const FVec3& LocalSpaceLocation = Solver->GetLocalSpaceLocation();
-
-	for (const FClothingSimulationCloth* const Cloth : Solver->GetCloths())
-	{
-		const int32 Offset = Cloth->GetOffset(Solver.Get());
-		if (Offset == INDEX_NONE)
-		{
-			continue;
-		}
-
-		const TConstArrayView<Softs::FSolverVec3> Positions = Cloth->GetAnimationPositions(Solver.Get());
-		const TConstArrayView<Softs::FSolverVec3> Normals = Cloth->GetAnimationNormals(Solver.Get());
-		check(Normals.Num() == Positions.Num());
-
-		for (int32 Index = 0; Index < Positions.Num(); ++Index)
-		{
-			const FVector Pos0 = LocalSpaceLocation + FVector(Positions[Index]);
-			const FVector Pos1 = Pos0 + FVector(Normals[Index]) * NormalLength;
-
-			DrawLine(PDI, Pos0, Pos1, KinematicColor);
-		}
-	}
-}
-
-void FClothingSimulation::DebugDrawPointNormals(FPrimitiveDrawInterface* PDI) const
-{
-	check(Solver);
-	static const FLinearColor DynamicColor = FColor::White;
-	static const FLinearColor KinematicColor = FColor::Purple;
-	constexpr FReal NormalLength = (FReal)20.;
-
-	const FVec3& LocalSpaceLocation = Solver->GetLocalSpaceLocation();
-
-	for (const FClothingSimulationCloth* const Cloth : Solver->GetCloths())
-	{
-		const int32 Offset = Cloth->GetOffset(Solver.Get());
-		if (Offset == INDEX_NONE)
-		{
-			continue;
-		}
-
-		const TConstArrayView<Softs::FSolverVec3> Positions = Cloth->GetParticlePositions(Solver.Get());
-		const TConstArrayView<Softs::FSolverVec3> Normals = Cloth->GetParticleNormals(Solver.Get());
-		const TConstArrayView<Softs::FSolverReal> InvMasses = Cloth->GetParticleInvMasses(Solver.Get());
-		check(Normals.Num() == Positions.Num());
-		check(InvMasses.Num() == Positions.Num());
-
-		for (int32 Index = 0; Index < Positions.Num(); ++Index)
-		{
-			const bool bIsKinematic = (InvMasses[Index] == (Softs::FSolverReal)0.);
-			const FVector Pos0 = LocalSpaceLocation + FVector(Positions[Index]);
-			const FVector Pos1 = Pos0 + FVector(Normals[Index]) * NormalLength;
-
-			DrawLine(PDI, Pos0, Pos1, bIsKinematic ? KinematicColor : DynamicColor);
-		}
-	}
-}
-
-void FClothingSimulation::DebugDrawPointVelocities(FPrimitiveDrawInterface* PDI) const
-{
-	check(Solver);
-	const FVec3& LocalSpaceLocation = Solver->GetLocalSpaceLocation();
-
-	for (const FClothingSimulationCloth* const Cloth : Solver->GetCloths())
-	{
-		const int32 Offset = Cloth->GetOffset(Solver.Get());
-		if (Offset == INDEX_NONE)
-		{
-			continue;
-		}
-
-		const TConstArrayView<Softs::FSolverVec3> Positions = Cloth->GetParticlePositions(Solver.Get());
-		const TConstArrayView<Softs::FSolverVec3> Velocities = Cloth->GetParticleVelocities(Solver.Get());
-		check(Velocities.Num() == Positions.Num());
-
-		const TConstArrayView<Softs::FSolverReal> InvMasses = Cloth->GetParticleInvMasses(Solver.Get());
-		check(InvMasses.Num() == Positions.Num());
-
-		for (int32 Index = 0; Index < Positions.Num(); ++Index)
-		{
-			constexpr FReal DefaultFPS = (FReal)60.;   // TODO: A CVAR would be nice for this
-			const bool bIsKinematic = (InvMasses[Index] == 0.f);
-
-			const FVec3 Pos0 = LocalSpaceLocation + Positions[Index];
-			const FVec3 Pos1 = Pos0 + FVec3(Velocities[Index]) / DefaultFPS;  // Velocity per frame if running at DefaultFPS
-
-			DrawLine(PDI, Pos0, Pos1, bIsKinematic ? FLinearColor::Black : FLinearColor::Yellow);
-		}
-	}
-}
-
-void FClothingSimulation::DebugDrawInversedPointNormals(FPrimitiveDrawInterface* PDI) const
-{
-	check(Solver);
-	static const FLinearColor DynamicColor = FColor::White;
-	static const FLinearColor KinematicColor = FColor::Purple;
-	constexpr FReal NormalLength = (FReal)20.;
-
-	const FVec3& LocalSpaceLocation = Solver->GetLocalSpaceLocation();
-
-	for (const FClothingSimulationCloth* const Cloth : Solver->GetCloths())
-	{
-		const int32 Offset = Cloth->GetOffset(Solver.Get());
-		if (Offset == INDEX_NONE)
-		{
-			continue;
-		}
-
-		const TConstArrayView<Softs::FSolverVec3> Positions = Cloth->GetParticlePositions(Solver.Get());
-		const TConstArrayView<Softs::FSolverVec3> Normals = Cloth->GetParticleNormals(Solver.Get());
-		const TConstArrayView<Softs::FSolverReal> InvMasses = Cloth->GetParticleInvMasses(Solver.Get());
-		check(Normals.Num() == Positions.Num());
-		check(InvMasses.Num() == Positions.Num());
-
-		for (int32 Index = 0; Index < Positions.Num(); ++Index)
-		{
-			const bool bIsKinematic = (InvMasses[Index] == (Softs::FSolverReal)0.);
-			const FVector Pos0 = LocalSpaceLocation + FVector(Positions[Index]);
-			const FVector Pos1 = Pos0 - FVector(Normals[Index]) * NormalLength;
-
-			DrawLine(PDI, Pos0, Pos1, bIsKinematic ? KinematicColor : DynamicColor);
-		}
-	}
-}
-
-void FClothingSimulation::DebugDrawCollision(FPrimitiveDrawInterface* PDI) const
-{
-	check(Solver);
-	auto DrawCollision =
-		[this, PDI](const FClothingSimulationCollider* Collider, const FClothingSimulationCloth* Cloth, FClothingSimulationCollider::ECollisionDataType CollisionDataType)
-		{
-			static const FLinearColor GlobalColor(FColor::Cyan);
-			static const FLinearColor DynamicColor(FColor::Orange);
-			static const FLinearColor LODsColor(FColor::Silver);
-			static const FLinearColor CollidedColor(FColor::Red);
-
-			const FLinearColor TypeColor =
-				(CollisionDataType == FClothingSimulationCollider::ECollisionDataType::LODless) ? GlobalColor :
-				(CollisionDataType == FClothingSimulationCollider::ECollisionDataType::External) ? DynamicColor : LODsColor;
-
-			const FVec3& LocalSpaceLocation = Solver->GetLocalSpaceLocation();
-
-			const TConstArrayView<TUniquePtr<FImplicitObject>> CollisionGeometries = Collider->GetCollisionGeometries(Solver.Get(), Cloth, CollisionDataType);
-			const TConstArrayView<Softs::FSolverVec3> Translations = Collider->GetCollisionTranslations(Solver.Get(), Cloth, CollisionDataType);
-			const TConstArrayView<Softs::FSolverRotation3> Rotations = Collider->GetCollisionRotations(Solver.Get(), Cloth, CollisionDataType);
-			const TConstArrayView<bool> CollisionStatus = Collider->GetCollisionStatus(Solver.Get(), Cloth, CollisionDataType);
-			check(CollisionGeometries.Num() == Translations.Num());
-			check(CollisionGeometries.Num() == Rotations.Num());
-
-			for (int32 Index = 0; Index < CollisionGeometries.Num(); ++Index)
-			{
-				if (const FImplicitObject* const Object = CollisionGeometries[Index].Get())
-				{
-					const FLinearColor Color = CollisionStatus[Index] ? CollidedColor : TypeColor;
-					const FVec3 Position = LocalSpaceLocation + FVec3(Translations[Index]);
-					const FRotation3 Rotation(Rotations[Index]);
-
-					switch (Object->GetType())
-					{
-					case ImplicitObjectType::Sphere:
-						DrawSphere(PDI, Object->GetObjectChecked<TSphere<FReal, 3>>(), Rotation, Position, Color);
-						break;
-
-					case ImplicitObjectType::Box:
-						DrawBox(PDI, Object->GetObjectChecked<TBox<FReal, 3>>().BoundingBox(), Rotation, Position, Color);
-						break;
-
-					case ImplicitObjectType::Capsule:
-						DrawCapsule(PDI, Object->GetObjectChecked<FCapsule>(), Rotation, Position, Color);
-						break;
-
-					case ImplicitObjectType::Union:  // Union only used as old style tapered capsules
-						for (const TUniquePtr<FImplicitObject>& SubObjectPtr : Object->GetObjectChecked<FImplicitObjectUnion>().GetObjects())
-						{
-							if (const FImplicitObject* const SubObject = SubObjectPtr.Get())
-							{
-								switch (SubObject->GetType())
-								{
-								case ImplicitObjectType::Sphere:
-									DrawSphere(PDI, SubObject->GetObjectChecked<TSphere<FReal, 3>>(), Rotation, Position, Color);
-									break;
-
-								case ImplicitObjectType::TaperedCylinder:
-									DrawTaperedCylinder(PDI, SubObject->GetObjectChecked<FTaperedCylinder>(), Rotation, Position, Color);
-									break;
-
-								default:
-									break;
-								}
-							}
-						}
-						break;
-
-					case ImplicitObjectType::TaperedCapsule:  // New collision tapered capsules implicit type that replaces the union
-						{
-							const FTaperedCapsule& TaperedCapsule = Object->GetObjectChecked<FTaperedCapsule>();
-							const FVec3 X1 = TaperedCapsule.GetX1();
-							const FVec3 X2 = TaperedCapsule.GetX2();
-							const FReal Radius1 = TaperedCapsule.GetRadius1();
-							const FReal Radius2 = TaperedCapsule.GetRadius2();
-							DrawSphere(PDI, TSphere<FReal, 3>(X1, Radius1), Rotation, Position, Color);
-							DrawSphere(PDI, TSphere<FReal, 3>(X2, Radius2), Rotation, Position, Color);
-							DrawTaperedCylinder(PDI, FTaperedCylinder(X1, X2, Radius1, Radius2), Rotation, Position, Color);
-						}
-						break;
-
-					case ImplicitObjectType::Convex:
-						DrawConvex(PDI, Object->GetObjectChecked<FConvex>(), Rotation, Position, Color);
-						break;
-
-					default:
-						DrawCoordinateSystem(PDI, Rotation, Position);  // Draw everything else as a coordinate for now
-						break;
-					}
-				}
-			}
-		};
-
-	// Draw collisions
-	for (const FClothingSimulationCloth* const Cloth : Solver->GetCloths())
-	{
-		for (const FClothingSimulationCollider* const Collider : Cloth->GetColliders())
-		{
-			DrawCollision(Collider, Cloth, FClothingSimulationCollider::ECollisionDataType::LODless);
-			DrawCollision(Collider, Cloth, FClothingSimulationCollider::ECollisionDataType::External);
-			DrawCollision(Collider, Cloth, FClothingSimulationCollider::ECollisionDataType::LODs);
-		}
-	}
-
-	// Draw contacts
-	check(Solver->GetCollisionContacts().Num() == Solver->GetCollisionNormals().Num());
-	constexpr FReal NormalLength = (FReal)10.;
-
-	const FVec3& LocalSpaceLocation = Solver->GetLocalSpaceLocation();
-	for (int32 i = 0; i < Solver->GetCollisionContacts().Num(); ++i)
-	{
-		const FVec3 Pos0 = LocalSpaceLocation + FVec3(Solver->GetCollisionContacts()[i]);
-		const FVec3 Normal = FVec3(Solver->GetCollisionNormals()[i]);
-
-		// Draw contact
-		FVec3 TangentU, TangentV;
-		Normal.FindBestAxisVectors(TangentU, TangentV);
-
-		DrawLine(PDI, Pos0 + TangentU, Pos0 + TangentV, FLinearColor::Black);
-		DrawLine(PDI, Pos0 + TangentU, Pos0 - TangentV, FLinearColor::Black);
-		DrawLine(PDI, Pos0 - TangentU, Pos0 - TangentV, FLinearColor::Black);
-		DrawLine(PDI, Pos0 - TangentU, Pos0 + TangentV, FLinearColor::Black);
-
-		// Draw normal
-		static const FLinearColor Brown(0.1f, 0.05f, 0.f);
-		const FVec3 Pos1 = Pos0 + NormalLength * Normal;
-		DrawLine(PDI, Pos0, Pos1, Brown);
-	}
-}
-
-void FClothingSimulation::DebugDrawBackstops(FPrimitiveDrawInterface* PDI) const
-{
-	auto DrawBackstop = [PDI](const FVector& Position, const FVector& Normal, FReal Radius, const FVector& Axis, const FLinearColor& Color)
-	{
-		static const FReal MaxCosAngle = (FReal)0.99;
-		if (FMath::Abs(FVector::DotProduct(Normal, Axis)) < MaxCosAngle)
-		{
-			static const FReal ArcLength = (FReal)5.; // Arch length in cm
-			const FReal ArcAngle = (FReal)360. * ArcLength / FMath::Max((Radius * (FReal)2. * (FReal)PI), ArcLength);
-			DrawArc(PDI, Position, Normal, FVector::CrossProduct(Axis, Normal).GetSafeNormal(), -ArcAngle / (FReal)2., ArcAngle / (FReal)2., Radius, Color);
-		}
-	};
-
-	const FVec3& LocalSpaceLocation = Solver->GetLocalSpaceLocation();
-
-	uint8 ColorSeed = 0;
-
-	for (const FClothingSimulationCloth* const Cloth : Solver->GetCloths())
-	{
-		const int32 Offset = Cloth->GetOffset(Solver.Get());
-		if (Offset == INDEX_NONE)
-		{
-			continue;
-		}
-
-		const FClothConstraints& ClothConstraints = Solver->GetClothConstraints(Offset);
-		if (const Softs::FPBDSphericalBackstopConstraint* const BackstopConstraint = ClothConstraints.GetBackstopConstraints().Get())
-		{
-			const bool bUseLegacyBackstop = BackstopConstraint->UseLegacyBackstop();
-			const TConstArrayView<FRealSingle>& BackstopDistances = Cloth->GetWeightMaps(Solver.Get())[(int32)EChaosWeightMapTarget::BackstopDistance];
-			const TConstArrayView<FRealSingle>& BackstopRadiuses = Cloth->GetWeightMaps(Solver.Get())[(int32)EChaosWeightMapTarget::BackstopRadius];
-			const TConstArrayView<Softs::FSolverVec3> AnimationPositions = Cloth->GetAnimationPositions(Solver.Get());
-			const TConstArrayView<Softs::FSolverVec3> AnimationNormals = Cloth->GetAnimationNormals(Solver.Get());
-			const TConstArrayView<Softs::FSolverVec3> ParticlePositions = Cloth->GetParticlePositions(Solver.Get());
-
-			for (int32 Index = 0; Index < AnimationPositions.Num(); ++Index)
-			{
-				ColorSeed += 157;  // Prime number that gives a good spread of colors without getting too similar as a rand might do.
-				const FLinearColor ColorLight = FLinearColor::MakeFromHSV8(ColorSeed, 160, 128);
-				const FLinearColor ColorDark = FLinearColor::MakeFromHSV8(ColorSeed, 160, 64);
-
-				const FReal BackstopRadius = BackstopRadiuses[Index] * BackstopConstraint->GetSphereRadiiMultiplier();
-				const FReal BackstopDistance = BackstopDistances[Index];
-
-				const FVector AnimationNormal(AnimationNormals[Index]);
-
-				// Draw a line to show the current distance to the sphere
-				const FVector Pos0 = LocalSpaceLocation + FVector(AnimationPositions[Index]);
-				const FVector Pos1 = Pos0 - (bUseLegacyBackstop ? BackstopDistance - BackstopRadius : BackstopDistance) * AnimationNormal;
-				const FVector Pos2 = LocalSpaceLocation + FVector(ParticlePositions[Index]);
-				DrawLine(PDI, Pos1, Pos2, ColorLight);
-
-				// Draw the sphere
-				if (BackstopRadius > 0.f)
-				{
-					const FVector Center = Pos0 - (bUseLegacyBackstop ? BackstopDistance : BackstopRadius + BackstopDistance) * AnimationNormal;
-					DrawBackstop(Center, AnimationNormal, BackstopRadius, FVector::ForwardVector, ColorDark);
-					DrawBackstop(Center, AnimationNormal, BackstopRadius, FVector::UpVector, ColorDark);
-					DrawBackstop(Center, AnimationNormal, BackstopRadius, FVector::RightVector, ColorDark);
-				}
-			}
-		}
-	}
-}
-
-void FClothingSimulation::DebugDrawBackstopDistances(FPrimitiveDrawInterface* PDI) const
-{
-	const FVec3& LocalSpaceLocation = Solver->GetLocalSpaceLocation();
-
-	uint8 ColorSeed = 0;
-
-	for (const FClothingSimulationCloth* const Cloth : Solver->GetCloths())
-	{
-		const int32 Offset = Cloth->GetOffset(Solver.Get());
-		if (Offset == INDEX_NONE)
-		{
-			continue;
-		}
-
-		const FClothConstraints& ClothConstraints = Solver->GetClothConstraints(Offset);
-		if (const Softs::FPBDSphericalBackstopConstraint* const BackstopConstraint = ClothConstraints.GetBackstopConstraints().Get())
-		{
-			const bool bUseLegacyBackstop = BackstopConstraint->UseLegacyBackstop();
-			const TConstArrayView<FRealSingle>& BackstopDistances = Cloth->GetWeightMaps(Solver.Get())[(int32)EChaosWeightMapTarget::BackstopDistance];
-			const TConstArrayView<FRealSingle>& BackstopRadiuses = Cloth->GetWeightMaps(Solver.Get())[(int32)EChaosWeightMapTarget::BackstopRadius];
-			const TConstArrayView<Softs::FSolverVec3> AnimationPositions = Cloth->GetAnimationPositions(Solver.Get());
-			const TConstArrayView<Softs::FSolverVec3> AnimationNormals = Cloth->GetAnimationNormals(Solver.Get());
-
-			for (int32 Index = 0; Index < AnimationPositions.Num(); ++Index)
-			{
-				ColorSeed += 157;  // Prime number that gives a good spread of colors without getting too similar as a rand might do.
-				const FLinearColor ColorLight = FLinearColor::MakeFromHSV8(ColorSeed, 160, 128);
-				const FLinearColor ColorDark = FLinearColor::MakeFromHSV8(ColorSeed, 160, 64);
-
-				const FReal BackstopRadius = BackstopRadiuses[Index] * BackstopConstraint->GetSphereRadiiMultiplier();
-				const FReal BackstopDistance = BackstopDistances[Index];
-
-				const FVector AnimationNormal(AnimationNormals[Index]);
-
-				// Draw a line to the sphere boundary
-				const FVector Pos0 = LocalSpaceLocation + FVector(AnimationPositions[Index]);
-				const FVector Pos1 = Pos0 - (bUseLegacyBackstop ? BackstopDistance - BackstopRadius : BackstopDistance) * AnimationNormal;
-				DrawLine(PDI, Pos0, Pos1, ColorDark);
-			}
-		}
-	}
-}
-
-void FClothingSimulation::DebugDrawMaxDistances(FPrimitiveDrawInterface* PDI) const
-{
-	check(Solver);
-	
-	// Draw max distances
-	const FVec3& LocalSpaceLocation = Solver->GetLocalSpaceLocation();
-
-	for (const FClothingSimulationCloth* const Cloth : Solver->GetCloths())
-	{
-		const int32 Offset = Cloth->GetOffset(Solver.Get());
-		if (Offset == INDEX_NONE)
-		{
-			continue;
-		}
-
-		const TConstArrayView<FRealSingle>& MaxDistances = Cloth->GetWeightMaps(Solver.Get())[(int32)EChaosWeightMapTarget::MaxDistance];
-		if (!MaxDistances.Num())
-		{
-			continue;
-		}
-
-		const TConstArrayView<Softs::FSolverReal> InvMasses = Cloth->GetParticleInvMasses(Solver.Get());
-		const TConstArrayView<Softs::FSolverVec3> Positions = Cloth->GetAnimationPositions(Solver.Get());
-		const TConstArrayView<Softs::FSolverVec3> Normals = Cloth->GetAnimationNormals(Solver.Get());
-		check(Normals.Num() == Positions.Num());
-		check(MaxDistances.Num() == Positions.Num());
-		check(InvMasses.Num() == Positions.Num());
-
-		for (int32 Index = 0; Index < MaxDistances.Num(); ++Index)
-		{
-			const FReal MaxDistance = (FReal)MaxDistances[Index];
-			const FVector Position = LocalSpaceLocation + FVector(Positions[Index]);
-			if (InvMasses[Index] == (Softs::FSolverReal)0.)
-			{
-#if WITH_EDITOR
-				DrawPoint(PDI, Position, FLinearColor::Red, DebugClothMaterialVertex);
-#else
-				DrawPoint(nullptr, Position, FLinearColor::Red, nullptr);
-#endif
-			}
-			else
-			{
-				DrawLine(PDI, Position, Position + FVector(Normals[Index]) * MaxDistance, FLinearColor::White);
-			}
-		}
-	}
-}
-
-void FClothingSimulation::DebugDrawAnimDrive(FPrimitiveDrawInterface* PDI) const
-{
-	check(Solver);
-
-	const FVec3& LocalSpaceLocation = Solver->GetLocalSpaceLocation();
-
-	for (const FClothingSimulationCloth* const Cloth : Solver->GetCloths())
-	{
-		const int32 Offset = Cloth->GetOffset(Solver.Get());
-		if (Offset == INDEX_NONE)
-		{
-			continue;
-		}
-
-		const FClothConstraints& ClothConstraints = Solver->GetClothConstraints(Offset);
-		if (const Softs::FPBDAnimDriveConstraint* const AnimDriveConstraint = ClothConstraints.GetAnimDriveConstraints().Get())
-		{
-			const TConstArrayView<FRealSingle>& AnimDriveStiffnessMultipliers = Cloth->GetWeightMaps(Solver.Get())[(int32)EChaosWeightMapTarget::AnimDriveStiffness];
-			const TConstArrayView<Softs::FSolverVec3> AnimationPositions = Cloth->GetAnimationPositions(Solver.Get());
-			const TConstArrayView<Softs::FSolverVec3> ParticlePositions = Cloth->GetParticlePositions(Solver.Get());
-			check(ParticlePositions.Num() == AnimationPositions.Num());
-
-			const FVec2 AnimDriveStiffness = AnimDriveConstraint->GetStiffness();
-			const FRealSingle StiffnessOffset = AnimDriveStiffness[0];
-			const FRealSingle StiffnessRange = AnimDriveStiffness[1] - AnimDriveStiffness[0];
-
-			for (int32 Index = 0; Index < ParticlePositions.Num(); ++Index)
-			{
-				const FRealSingle Stiffness = AnimDriveStiffnessMultipliers.IsValidIndex(Index) ?
-					StiffnessOffset + AnimDriveStiffnessMultipliers[Index] * StiffnessRange :
-					StiffnessOffset;
-
-				const FVector AnimationPosition = LocalSpaceLocation + FVector(AnimationPositions[Index]);
-				const FVector ParticlePosition = LocalSpaceLocation + FVector(ParticlePositions[Index]);
-				DrawLine(PDI, AnimationPosition, ParticlePosition, FLinearColor(FColor::Cyan) * Stiffness);
-			}
-		}
-	}
-}
-
-static void DebugDrawSpringConstraintColors(FPrimitiveDrawInterface* PDI, const TConstArrayView<Softs::FSolverVec3>& Positions, const FVec3& LocalSpaceLocation, const Softs::FPBDSpringConstraints* const SpringConstraints)
-{
-	check(SpringConstraints);
-
-	const TArray<TVec2<int32>>& Constraints = SpringConstraints->GetConstraints();
-	const TArray<int32>& ConstraintsPerColorStartIndex = SpringConstraints->GetConstraintsPerColorStartIndex();
-	if (ConstraintsPerColorStartIndex.Num() > 1)
-	{
-		const int32 ConstraintColorNum = ConstraintsPerColorStartIndex.Num() - 1;
-		const uint8 HueOffset = 196 / ConstraintColorNum;
-		auto ConstraintColor =
-			[HueOffset](int32 ColorIndex)->FLinearColor
-		{
-			return FLinearColor::MakeFromHSV8(ColorIndex * HueOffset, 255, 255);
-		};
-
-		for (int32 ConstraintColorIndex = 0; ConstraintColorIndex < ConstraintColorNum; ++ConstraintColorIndex)
-		{
-			const int32 ColorStart = ConstraintsPerColorStartIndex[ConstraintColorIndex];
-			const int32 ColorEnd = ConstraintsPerColorStartIndex[ConstraintColorIndex + 1];
-			const FLinearColor DrawColor = ConstraintColor(ConstraintColorIndex);
-			for (int32 ConstraintIndex = ColorStart; ConstraintIndex < ColorEnd; ++ConstraintIndex)
-			{
-				// Draw line
-				const TVec2<int32>& Constraint = Constraints[ConstraintIndex];
-				const FVec3 Pos0 = FVec3(Positions[Constraint[0]]) + LocalSpaceLocation;
-				const FVec3 Pos1 = FVec3(Positions[Constraint[1]]) + LocalSpaceLocation;
-				DrawLine(PDI, Pos0, Pos1, DrawColor);
-			}
-		}
-	}
-	else
-	{
-		for (const TVec2<int32>& Constraint : Constraints)
-		{
-			// Draw line
-			const FVec3 Pos0 = FVec3(Positions[Constraint[0]]) + LocalSpaceLocation;
-			const FVec3 Pos1 = FVec3(Positions[Constraint[1]]) + LocalSpaceLocation;
-
-			DrawLine(PDI, Pos0, Pos1, FLinearColor::Black);
-		}
-	}
-}
-
-void FClothingSimulation::DebugDrawEdgeConstraint(FPrimitiveDrawInterface* PDI) const
-{
-	const FVec3& LocalSpaceLocation = Solver->GetLocalSpaceLocation();
-
-	for (const FClothingSimulationCloth* const Cloth : Solver->GetCloths())
-	{
-		const int32 Offset = Cloth->GetOffset(Solver.Get());
-		if (Offset == INDEX_NONE)
-		{
-			continue;
-		}
-
-		// Draw constraints
-		const FClothConstraints& ClothConstraints = Solver->GetClothConstraints(Offset);
-
-		const TConstArrayView<Softs::FSolverVec3> Positions = Cloth->GetParticlePositions(Solver.Get());
-
-		if (const Softs::FPBDSpringConstraints* const EdgeConstraints = ClothConstraints.GetEdgeConstraints().Get())
-		{
-			DebugDrawSpringConstraintColors(PDI, Positions, LocalSpaceLocation, EdgeConstraints);
-		}
-	}
-}
-
-void FClothingSimulation::DebugDrawBendingConstraint(FPrimitiveDrawInterface* PDI) const
-{
-	const FVec3& LocalSpaceLocation = Solver->GetLocalSpaceLocation();
-
-	for (const FClothingSimulationCloth* const Cloth : Solver->GetCloths())
-	{
-		const int32 Offset = Cloth->GetOffset(Solver.Get());
-		if (Offset == INDEX_NONE)
-		{
-			continue;
-		}
-
-		// Draw constraints
-		const FClothConstraints& ClothConstraints = Solver->GetClothConstraints(Offset);
-
-		const TConstArrayView<Softs::FSolverVec3> Positions = Cloth->GetParticlePositions(Solver.Get());
-
-		if (const Softs::FPBDSpringConstraints* const BendingConstraints = ClothConstraints.GetBendingConstraints().Get())
-		{
-			DebugDrawSpringConstraintColors(PDI, Positions, LocalSpaceLocation, BendingConstraints);
-		}
-
-		if (const Softs::FPBDBendingConstraints* const BendingConstraints = ClothConstraints.GetBendingElementConstraints().Get())
-		{
-			const TArray<TVec4<int32>>& Constraints = BendingConstraints->GetConstraints();
-			const TArray<bool>& IsBuckled = BendingConstraints->GetIsBuckled();
-
-			// Color constraint edge with red or blue: Red = Buckled, Blue = Not Buckled. 
-			for (int32 ConstraintIndex = 0; ConstraintIndex < Constraints.Num(); ++ConstraintIndex)
-			{
-				const Softs::FSolverVec3& P1 = Positions[Constraints[ConstraintIndex][0]];
-				const Softs::FSolverVec3& P2 = Positions[Constraints[ConstraintIndex][1]];
-				
-				const bool bIsBuckled = IsBuckled[ConstraintIndex];
-
-				const FVec3 Pos0 = FVec3(P1) + LocalSpaceLocation;
-				const FVec3 Pos1 = FVec3(P2) + LocalSpaceLocation;
-				DrawLine(PDI, Pos0, Pos1, bIsBuckled ? FLinearColor::Red : FLinearColor::Blue);
-			}
-		}
-
-		if (const Softs::FXPBDBendingConstraints* const BendingConstraints = ClothConstraints.GetXBendingElementConstraints().Get())
-		{
-			const TArray<TVec4<int32>>& Constraints = BendingConstraints->GetConstraints();
-			const TArray<bool>& IsBuckled = BendingConstraints->GetIsBuckled();
-
-			// Color constraint edge with red or blue: Red = Buckled, Blue = Not Buckled.
-			for (int32 ConstraintIndex = 0; ConstraintIndex < Constraints.Num(); ++ConstraintIndex)
-			{
-				const Softs::FSolverVec3& P1 = Positions[Constraints[ConstraintIndex][0]];
-				const Softs::FSolverVec3& P2 = Positions[Constraints[ConstraintIndex][1]];
-
-				const bool bIsBuckled = IsBuckled[ConstraintIndex];
-
-				const FVec3 Pos0 = FVec3(P1) + LocalSpaceLocation;
-				const FVec3 Pos1 = FVec3(P2) + LocalSpaceLocation;
-				DrawLine(PDI, Pos0, Pos1, bIsBuckled ? FLinearColor::Red : FLinearColor::Blue);
-			}
-		}
-	}
-}
-
-void FClothingSimulation::DebugDrawLongRangeConstraint(FPrimitiveDrawInterface* PDI) const
-{
-	const FVec3& LocalSpaceLocation = Solver->GetLocalSpaceLocation();
-
-	auto PseudoRandomColor =
-		[](int32 NumColorRotations) -> FLinearColor
-		{
-			static const uint8 Spread = 157;  // Prime number that gives a good spread of colors without getting too similar as a rand might do.
-			uint8 Seed = Spread;
-			for (int32 i = 0; i < NumColorRotations; ++i)
-			{
-				Seed += Spread;
-			}
-			return FLinearColor::MakeFromHSV8(Seed, 160, 128);
-		};
-
-	auto Darken =
-		[](const FLinearColor& Color) -> FLinearColor
-		{
-			FLinearColor ColorHSV = Color.LinearRGBToHSV();
-			ColorHSV.B *= .5f;
-			return ColorHSV.HSVToLinearRGB();
-		};
-
-	int32 ColorOffset = 0;
-
-	for (const FClothingSimulationCloth* const Cloth : Solver->GetCloths())
-	{
-		const int32 Offset = Cloth->GetOffset(Solver.Get());
-		if (Offset == INDEX_NONE)
-		{
-			continue;
-		}
-
-		// Draw constraints
-		const FClothConstraints& ClothConstraints = Solver->GetClothConstraints(Offset);
-		const TConstArrayView<Softs::FSolverVec3> Positions = Cloth->GetParticlePositions(Solver.Get());
-
-		if (const Softs::FPBDLongRangeConstraints* const LongRangeConstraints = ClothConstraints.GetLongRangeConstraints().Get())
-		{
-			const TArray<TConstArrayView<Softs::FPBDLongRangeConstraints::FTether>>& Tethers = LongRangeConstraints->GetTethers();
-
-			for (int32 BatchIndex = 0; BatchIndex < Tethers.Num(); ++BatchIndex)
-			{
-				const FLinearColor Color = PseudoRandomColor(ColorOffset + BatchIndex);
-				const FLinearColor DarkenedColor = Darken(Color);
-
-				const TConstArrayView<Softs::FPBDLongRangeConstraints::FTether>& TetherBatch = Tethers[BatchIndex];
-
-				// Draw tethers
-				for (const Softs::FPBDLongRangeConstraints::FTether& Tether : TetherBatch)
-				{
-					const int32 KinematicIndex = LongRangeConstraints->GetStartIndex(Tether);
-					const int32 DynamicIndex = LongRangeConstraints->GetEndIndex(Tether);
-					const FReal TargetLength = LongRangeConstraints->GetTargetLength(Tether);
-
-					const FVec3 Pos0 = FVec3(Positions[KinematicIndex]) + LocalSpaceLocation;
-					const FVec3 Pos1 = FVec3(Positions[DynamicIndex]) + LocalSpaceLocation;
-
-					DrawLine(PDI, Pos0, Pos1, Color);
-#if WITH_EDITOR
-					DrawPoint(PDI, Pos1, Color, DebugClothMaterialVertex);
-#else
-					DrawPoint(nullptr, Pos1, Color, nullptr);
-#endif
-					FVec3 Direction = Pos1 - Pos0;
-					const float Length = Direction.SafeNormalize();
-					if (Length > SMALL_NUMBER)
-					{
-						const FVec3 Pos2 = Pos1 + Direction * (TargetLength - Length);
-						DrawLine(PDI, Pos1, Pos2, DarkenedColor);
-					}
-				}
-			}
-
-			// Rotate the colors for each cloth
-			ColorOffset += Tethers.Num();
-		}
-	}
-}
-
-void FClothingSimulation::DebugDrawWindAndPressureForces(FPrimitiveDrawInterface* PDI) const
-{
-	check(Solver);
-
-	constexpr FReal ForceLength = (FReal)10.;
-	const FVec3& LocalSpaceLocation = Solver->GetLocalSpaceLocation();
-
-	for (const FClothingSimulationCloth* const Cloth : Solver->GetCloths())
-	{
-		const int32 Offset = Cloth->GetOffset(Solver.Get());
-		if (Offset == INDEX_NONE)
-		{
-			continue;
-		}
-
-		const Softs::FVelocityAndPressureField& VelocityField = Solver->GetWindVelocityAndPressureField(Cloth->GetGroupId());
-
-		const TConstArrayView<TVec3<int32>>& Elements = VelocityField.GetElements();
-		const TConstArrayView<Softs::FSolverVec3> Forces = VelocityField.GetForces();
-		const TConstArrayView<Softs::FSolverVec3> Positions = Cloth->GetParticlePositions(Solver.Get());
-		const TConstArrayView<Softs::FSolverReal> InvMasses = Cloth->GetParticleInvMasses(Solver.Get());
-		check(InvMasses.Num() == Positions.Num());
-
-		for (int32 ElementIndex = 0; ElementIndex < Elements.Num(); ++ElementIndex)
-		{
-			const TVec3<int32>& Element = Elements[ElementIndex];
-			const FVec3 Position = LocalSpaceLocation + (
-				FVec3(Positions[Element.X - Offset]) +
-				FVec3(Positions[Element.Y - Offset]) +
-				FVec3(Positions[Element.Z - Offset])) / (FReal)3.;
-
-			const bool bIsKinematic0 = !InvMasses[Element.X - Offset];
-			const bool bIsKinematic1 = !InvMasses[Element.Y - Offset];
-			const bool bIsKinematic2 = !InvMasses[Element.Z - Offset];
-			const bool bIsKinematic = bIsKinematic0 || bIsKinematic1 || bIsKinematic2;
-
-			const FVec3 Force = FVec3(Forces[ElementIndex]) * ForceLength;
-			DrawLine(PDI, Position, Position + Force, bIsKinematic ? FColor::Cyan : FColor::Green);
-		}
-	}
-}
-
-void FClothingSimulation::DebugDrawLocalSpace(FPrimitiveDrawInterface* PDI) const
-{
-	check(Solver);
-
-	// Draw local space
-	DrawCoordinateSystem(PDI, FQuat::Identity, Solver->GetLocalSpaceLocation());
-
-	// Draw reference spaces
-	for (const FClothingSimulationCloth* const Cloth : Solver->GetCloths())
-	{
-		if (Cloth->GetOffset(Solver.Get()) == INDEX_NONE)
-		{
-			continue;
-		}
-		const FRigidTransform3& ReferenceSpaceTransform = Cloth->GetReferenceSpaceTransform();
-		DrawCoordinateSystem(PDI, ReferenceSpaceTransform.GetRotation(), ReferenceSpaceTransform.GetLocation());
-	}
-}
-
-void FClothingSimulation::DebugDrawSelfCollision(FPrimitiveDrawInterface* PDI) const
-{
-	const FVec3& LocalSpaceLocation = Solver->GetLocalSpaceLocation();
-
-	for (const FClothingSimulationCloth* const Cloth : Solver->GetCloths())
-	{
-		const int32 Offset = Cloth->GetOffset(Solver.Get());
-		if (Offset == INDEX_NONE)
-		{
-			continue;
-		}
-
-		// Draw constraints
-		const FClothConstraints& ClothConstraints = Solver->GetClothConstraints(Offset);
-
-		if (const Softs::FPBDCollisionSpringConstraints* const SelfCollisionConstraints = ClothConstraints.GetSelfCollisionConstraints().Get())
-		{
-			const TConstArrayView<Softs::FSolverVec3> Positions = Cloth->GetParticlePositions(Solver.Get());
-			const TArray<TVec4<int32>>& Constraints = SelfCollisionConstraints->GetConstraints();
-			const TArray<Softs::FSolverVec3>& Barys = SelfCollisionConstraints->GetBarys();
-			const FReal Thickness = (FReal)SelfCollisionConstraints->GetThickness();
-			const FReal Height = Thickness + Thickness;
-			const TArray<bool>& FlipNormals = SelfCollisionConstraints->GetFlipNormals();
-
-			for (int32 Index = 0; Index < Constraints.Num(); ++Index)
-			{
-				const TVec4<int32>& Constraint = Constraints[Index];
-				const FVec3 Bary(Barys[Index]);
-
-				// Constraint index includes Offset, but so does Positions.
-				const FVector P = LocalSpaceLocation + FVector(Positions[Constraint[0] - Offset]);
-				const FVector P0 = LocalSpaceLocation + FVector(Positions[Constraint[1] - Offset]);
-				const FVector P1 = LocalSpaceLocation + FVector(Positions[Constraint[2] - Offset]);
-				const FVector P2 = LocalSpaceLocation + FVector(Positions[Constraint[3] - Offset]);
-
-				const FVector Pos0 = P0 * Bary[0] + P1 * Bary[1] + P2 * Bary[2];
-
-				static const FLinearColor Brown(0.1f, 0.05f, 0.f);
-				static const FLinearColor Red(0.3f, 0.f, 0.f);
-				const FTriangle Triangle(P0, P1, P2);
-				const FVector Normal = FlipNormals[Index] ? -Triangle.GetNormal() : Triangle.GetNormal();
-
-				// Draw point to surface line (=normal)
-				const FVector Pos1 = Pos0 + Height * Normal;
-				DrawPoint(PDI, Pos0, Brown, nullptr, 2.f);
-				DrawLine(PDI, Pos0, Pos1, FlipNormals[Index] ? Red : Brown);
-
-				// Draw pushup to point
-				static const FLinearColor Orange(0.3f, 0.15f, 0.f);
-				DrawPoint(PDI, P, Orange, nullptr, 2.f);
-				DrawLine(PDI, Pos1, P, Orange);
-			}
-		}
-	}
-}
-
-void FClothingSimulation::DebugDrawSelfIntersection(FPrimitiveDrawInterface* PDI) const
-{
-	const FVec3& LocalSpaceLocation = Solver->GetLocalSpaceLocation();
-
-	for (const FClothingSimulationCloth* const Cloth : Solver->GetCloths())
-	{
-		const int32 Offset = Cloth->GetOffset(Solver.Get());
-		if (Offset == INDEX_NONE)
-		{
-			continue;
-		}
-
-		const FClothConstraints& ClothConstraints = Solver->GetClothConstraints(Offset);
-
-		static const FLinearColor Red(1.f, 0.f, 0.f);
-		static const FLinearColor White(1.f, 1.f, 1.f);
-		static const FLinearColor Black(0.f, 0.f, 0.f);
-		static const FLinearColor Teal(0.f, 0.5f, 0.5f);
-		static const FLinearColor Green(0.f, 1.f, 0.f);
-
-		if (const Softs::FPBDTriangleMeshCollisions* const SelfCollisionInit = ClothConstraints.GetSelfCollisionInit().Get())
-		{
-			const TConstArrayView<Softs::FSolverVec3> Positions = Cloth->GetParticlePositions(Solver.Get());
-			const FTriangleMesh& TriangleMesh = Cloth->GetTriangleMesh(Solver.Get());
-
-			// Draw contours
-			const TArray<TArray<Softs::FPBDTriangleMeshCollisions::FBarycentricPoint>>& ContourPoints = SelfCollisionInit->GetIntersectionContourPoints();
-			const TArray<Softs::FPBDTriangleMeshCollisions::FContourType>& ContourTypes = SelfCollisionInit->GetIntersectionContourTypes();
-			check(ContourPoints.Num() == ContourTypes.Num());
-
-			static const FLinearColor ColorsForType[(int8)Softs::FPBDTriangleMeshCollisions::FContourType::Count] =
-			{
-				Teal,
-				Red,
-				White,
-				Black
-			};
-			for( int32 ContourIndex = 0; ContourIndex < ContourPoints.Num(); ++ContourIndex)
-			{
-				const TArray<Softs::FPBDTriangleMeshCollisions::FBarycentricPoint>& Contour = ContourPoints[ContourIndex];
-				const FLinearColor& ContourColor = ColorsForType[(int8)ContourTypes[ContourIndex]];
-				for (int32 PointIdx = 0; PointIdx < Contour.Num() - 1; ++PointIdx)
-				{
-					const Softs::FPBDTriangleMeshCollisions::FBarycentricPoint& Point0 = Contour[PointIdx];
-					const FVector EndPoint0 = LocalSpaceLocation + (1.f - Point0.Bary[0] - Point0.Bary[1]) * Positions[Point0.Vertices[0] - Offset] + Point0.Bary[0] * Positions[Point0.Vertices[1] - Offset] + Point0.Bary[1] * Positions[Point0.Vertices[2] - Offset];
-					const Softs::FPBDTriangleMeshCollisions::FBarycentricPoint& Point1 = Contour[PointIdx+1];
-					const FVector EndPoint1 = LocalSpaceLocation + (1.f - Point1.Bary[0] - Point1.Bary[1]) * Positions[Point1.Vertices[0] - Offset] + Point1.Bary[0] * Positions[Point1.Vertices[1] - Offset] + Point1.Bary[1] * Positions[Point1.Vertices[2] - Offset];
-					DrawLine(PDI, EndPoint0, EndPoint1, ContourColor);
-					DrawPoint(PDI, EndPoint0, ContourColor, nullptr, 1.f);
-					DrawPoint(PDI, EndPoint1, ContourColor, nullptr, 1.f);
-				}
-			}
-
-			// Draw GIA colors
-			const TConstArrayView<Softs::FPBDTriangleMeshCollisions::FGIAColor>& VertexGIAColors = SelfCollisionInit->GetVertexGIAColors();
-			static const FLinearColor Gray(0.5f, 0.5f, 0.5f);
-			if (VertexGIAColors.Num())
-			{
-				const TArray<TVec3<int32>>& Elements = TriangleMesh.GetElements();
-				for (int32 ParticleIdx = Offset; ParticleIdx < VertexGIAColors.Num(); ++ParticleIdx)
-				{
-					if (VertexGIAColors[ParticleIdx].ContourIndexBits)
-					{
-						const bool bIsLoop = VertexGIAColors[ParticleIdx].IsLoop();
-						const bool bAnyWhite = (VertexGIAColors[ParticleIdx].ContourIndexBits & ~VertexGIAColors[ParticleIdx].ColorBits);
-						const bool bAnyBlack = (VertexGIAColors[ParticleIdx].ContourIndexBits & VertexGIAColors[ParticleIdx].ColorBits);
-						const FLinearColor& VertColor = bIsLoop ? Red : (bAnyWhite && bAnyBlack) ? Gray : bAnyWhite ? White : Black;
-						
-						DrawPoint(PDI, LocalSpaceLocation + Positions[ParticleIdx - Offset], VertColor, nullptr, 5.f);
-					}
-				}
-			}
-			const TArray<Softs::FPBDTriangleMeshCollisions::FGIAColor>& TriangleGIAColors = SelfCollisionInit->GetTriangleGIAColors();
-			if (TriangleGIAColors.Num() == TriangleMesh.GetNumElements())
-			{
-				const TArray<TVec3<int32>>& Elements = TriangleMesh.GetElements();
-				for (int32 TriangleIdx = 0; TriangleIdx < TriangleGIAColors.Num(); ++TriangleIdx)
-				{
-					if (TriangleGIAColors[TriangleIdx].ContourIndexBits)
-					{
-						const bool bIsLoop = TriangleGIAColors[TriangleIdx].IsLoop();
-						const bool bAnyWhite = (TriangleGIAColors[TriangleIdx].ContourIndexBits & ~TriangleGIAColors[TriangleIdx].ColorBits);
-						const bool bAnyBlack = (TriangleGIAColors[TriangleIdx].ContourIndexBits & TriangleGIAColors[TriangleIdx].ColorBits);
-						const FLinearColor& TriColor = bIsLoop ? Red : (bAnyWhite && bAnyBlack) ? Gray : bAnyWhite ? White : Black;
-						DrawLine(PDI, LocalSpaceLocation + Positions[Elements[TriangleIdx][0] - Offset], LocalSpaceLocation + Positions[Elements[TriangleIdx][1] - Offset], TriColor);
-						DrawLine(PDI, LocalSpaceLocation + Positions[Elements[TriangleIdx][1] - Offset], LocalSpaceLocation + Positions[Elements[TriangleIdx][2] - Offset], TriColor);
-						DrawLine(PDI, LocalSpaceLocation + Positions[Elements[TriangleIdx][0] - Offset], LocalSpaceLocation + Positions[Elements[TriangleIdx][2] - Offset], TriColor);
-					}
-				}
-			}
-
-			// Draw contour minimization gradients
-			const TArray<Softs::FPBDTriangleMeshCollisions::FContourMinimizationIntersection>& ContourMinimizationIntersections = SelfCollisionInit->GetContourMinimizationIntersections();
-			constexpr FReal MaxDrawImpulse = 1.;
-			constexpr FReal RegularizeEpsilonSq = 1.;
-			for (const Softs::FPBDTriangleMeshCollisions::FContourMinimizationIntersection& Intersection : ContourMinimizationIntersections)
-			{
-				Softs::FSolverReal GradientLength;
-				Softs::FSolverVec3 GradientDir;
-				Intersection.GlobalGradientVector.ToDirectionAndLength(GradientDir, GradientLength);
-				const FVector Delta = FVector(GradientDir) * MaxDrawImpulse * GradientLength * FMath::InvSqrt(GradientLength * GradientLength + RegularizeEpsilonSq);
-
-				const FVector EdgeCenter = LocalSpaceLocation + .5 * (Positions[Intersection.EdgeVertices[0] - Offset] + Positions[Intersection.EdgeVertices[1] - Offset]);
-				const FVector TriCenter = LocalSpaceLocation + (Positions[Intersection.FaceVertices[0] - Offset] + Positions[Intersection.FaceVertices[1] - Offset] + Positions[Intersection.FaceVertices[2] - Offset]) / 3.;
-
-				DrawPoint(PDI, EdgeCenter, Green, nullptr, 2.f);
-				DrawLine(PDI, EdgeCenter, EdgeCenter + Delta, Green);
-				DrawPoint(PDI, TriCenter, Green, nullptr, 2.f);
-				DrawLine(PDI, TriCenter, TriCenter - Delta, Green);
-			}
-		}
-	}
-}
-
-#endif  // #if WITH_EDITOR || CHAOS_DEBUG_DRAW?
 }  // End namespace Chaos
diff --git a/Engine/Plugins/Experimental/ChaosCloth/Source/Chaos/Public/ChaosCloth/ChaosClothVisualization.h b/Engine/Plugins/Experimental/ChaosCloth/Source/Chaos/Public/ChaosCloth/ChaosClothVisualization.h
new file mode 100644
--- /dev/null
+++ b/Engine/Plugins/Experimental/ChaosCloth/Source/Chaos/Public/ChaosCloth/ChaosClothVisualization.h
@@ -0,0 +1,112 @@
+// Copyright Epic Games, Inc. All Rights Reserved.
+#pragma once
+
+#include "Misc/Build.h"      // For CHAOS_DEBUG_DRAW
+#include "Chaos/Declares.h"  //
+
+#if WITH_EDITOR
+#include "UObject/GCObject.h"
+#endif  // #if WITH_EDITOR
+
+class FCanvas;
+class FSceneView;
+class FPrimitiveDrawInterface;
+class UMaterial;
+
+namespace Chaos
+{
+	class FClothingSimulationSolver;
+}
+
+namespace UE::Chaos::Cloth
+{
+	class CHAOSCLOTH_API FVisualization
+#if WITH_EDITOR
+		: public FGCObject  // Add garbage collection for cloth material
+#endif  // #if WITH_EDITOR
+	{
+	public:
+		explicit FVisualization(const ::Chaos::FClothingSimulationSolver* InSolver = nullptr);
+		virtual ~FVisualization();
+
+#if CHAOS_DEBUG_DRAW
+		// Editor & runtime functions
+		void SetSolver(const ::Chaos::FClothingSimulationSolver* InSolver);
+
+		void DrawPhysMeshWired(FPrimitiveDrawInterface* PDI = nullptr) const;
+		void DrawAnimMeshWired(FPrimitiveDrawInterface* PDI = nullptr) const;
+		void DrawAnimNormals(FPrimitiveDrawInterface* PDI = nullptr) const;
+		void DrawPointNormals(FPrimitiveDrawInterface* PDI = nullptr) const;
+		void DrawPointVelocities(FPrimitiveDrawInterface* PDI = nullptr) const;
+		void DrawCollision(FPrimitiveDrawInterface* PDI = nullptr) const;
+		void DrawBackstops(FPrimitiveDrawInterface* PDI = nullptr) const;
+		void DrawBackstopDistances(FPrimitiveDrawInterface* PDI = nullptr) const;
+		void DrawMaxDistances(FPrimitiveDrawInterface* PDI = nullptr) const;
+		void DrawAnimDrive(FPrimitiveDrawInterface* PDI = nullptr) const;
+		void DrawEdgeConstraint(FPrimitiveDrawInterface* PDI = nullptr) const;
+		void DrawBendingConstraint(FPrimitiveDrawInterface* PDI = nullptr) const;
+		void DrawLongRangeConstraint(FPrimitiveDrawInterface* PDI = nullptr) const;
+		void DrawWindAndPressureForces(FPrimitiveDrawInterface* PDI = nullptr) const;
+		void DrawLocalSpace(FPrimitiveDrawInterface* PDI = nullptr) const;
+		void DrawSelfCollision(FPrimitiveDrawInterface* PDI = nullptr) const;
+		void DrawSelfIntersection(FPrimitiveDrawInterface* PDI = nullptr) const;
+		void DrawBounds(FPrimitiveDrawInterface* PDI = nullptr) const;
+		void DrawGravity(FPrimitiveDrawInterface* PDI = nullptr) const;
+
+#else  // #if CHAOS_DEBUG_DRAW
+		void SetSolver(const ::Chaos::FClothingSimulationSolver* /*InSolver*/) {}
+
+		void DrawPhysMeshWired(FPrimitiveDrawInterface* /*PDI*/ = nullptr) const {}
+		void DrawAnimMeshWired(FPrimitiveDrawInterface* /*PDI*/ = nullptr) const {}
+		void DrawAnimNormals(FPrimitiveDrawInterface* /*PDI*/ = nullptr) const {}
+		void DrawPointNormals(FPrimitiveDrawInterface* /*PDI*/ = nullptr) const {}
+		void DrawPointVelocities(FPrimitiveDrawInterface* /*PDI*/ = nullptr) const {}
+		void DrawCollision(FPrimitiveDrawInterface* /*PDI*/ = nullptr) const {}
+		void DrawBackstops(FPrimitiveDrawInterface* /*PDI*/ = nullptr) const {}
+		void DrawBackstopDistances(FPrimitiveDrawInterface* /*PDI*/ = nullptr) const {}
+		void DrawMaxDistances(FPrimitiveDrawInterface* /*PDI*/ = nullptr) const {}
+		void DrawAnimDrive(FPrimitiveDrawInterface* /*PDI*/ = nullptr) const {}
+		void DrawEdgeConstraint(FPrimitiveDrawInterface* /*PDI*/ = nullptr) const {}
+		void DrawBendingConstraint(FPrimitiveDrawInterface* /*PDI*/ = nullptr) const {}
+		void DrawLongRangeConstraint(FPrimitiveDrawInterface* /*PDI*/ = nullptr) const {}
+		void DrawWindAndPressureForces(FPrimitiveDrawInterface* /*PDI*/ = nullptr) const {}
+		void DrawLocalSpace(FPrimitiveDrawInterface* /*PDI*/ = nullptr) const {}
+		void DrawSelfCollision(FPrimitiveDrawInterface* /*PDI*/ = nullptr) const {}
+		void DrawSelfIntersection(FPrimitiveDrawInterface* /*PDI*/ = nullptr) const {}
+		void DrawBounds(FPrimitiveDrawInterface* /*PDI*/ = nullptr) const {}
+		void DrawGravity(FPrimitiveDrawInterface* /*PDI*/ = nullptr) const {}
+
+#endif  // #if CHAOS_DEBUG_DRAW
+
+#if WITH_EDITOR && CHAOS_DEBUG_DRAW
+		// Editor only functions
+		void DrawPhysMeshShaded(FPrimitiveDrawInterface* PDI) const;
+		void DrawParticleIndices(FCanvas* Canvas, const FSceneView* SceneView) const;
+		void DrawElementIndices(FCanvas* Canvas, const FSceneView* SceneView) const;
+		void DrawMaxDistanceValues(FCanvas* Canvas, const FSceneView* SceneView) const;
+
+	protected:
+		// FGCObject interface
+		virtual void AddReferencedObjects(FReferenceCollector& Collector) override;
+		virtual FString GetReferencerName() const override { return TEXT("UE::Chaos::Cloth::FVisualization"); }
+		// End of FGCObject interface
+
+#else  // #if WITH_EDITOR && CHAOS_DEBUG_DRAW
+		void DrawPhysMeshShaded(FPrimitiveDrawInterface* /*PDI*/) const {}
+		void DrawParticleIndices(FCanvas* /*Canvas*/, const FSceneView* /*SceneView*/) const {}
+		void DrawElementIndices(FCanvas* /*Canvas*/, const FSceneView* /*SceneView*/) const {}
+		void DrawMaxDistanceValues(FCanvas* /*Canvas*/, const FSceneView* /*SceneView*/) const {}
+#endif  // #if WITH_EDITOR && CHAOS_DEBUG_DRAW
+
+#if CHAOS_DEBUG_DRAW
+	private:
+		// Simulation objects
+		const ::Chaos::FClothingSimulationSolver* Solver;
+#if WITH_EDITOR
+		// Visualization material
+		const UMaterial* ClothMaterial;
+		const UMaterial* ClothMaterialVertex;
+#endif  // #if WITH_EDITOR
+#endif  // #if CHAOS_DEBUG_DRAW
+	};
+} // End namespace UE::Chaos::Cloth
diff --git a/Engine/Plugins/Experimental/ChaosCloth/Source/Chaos/Public/ChaosCloth/ChaosClothingSimulation.h b/Engine/Plugins/Experimental/ChaosCloth/Source/Chaos/Public/ChaosCloth/ChaosClothingSimulation.h
--- a/Engine/Plugins/Experimental/ChaosCloth/Source/Chaos/Public/ChaosCloth/ChaosClothingSimulation.h
+++ b/Engine/Plugins/Experimental/ChaosCloth/Source/Chaos/Public/ChaosCloth/ChaosClothingSimulation.h
@@ -4,16 +4,13 @@
 #include "ClothingSimulation.h"
 #include "ClothingAsset.h"
 #include "ClothCollisionData.h"
-#include "Chaos/ChaosDebugDrawDeclares.h"
 #include "ChaosCloth/ChaosClothConfig.h"
+#include "ChaosCloth/ChaosClothVisualization.h"
 #include "Templates/Atomic.h"
 #include "Templates/UniquePtr.h"
 
 class USkeletalMeshComponent;
 class FClothingSimulationContextCommon;
-#if WITH_EDITOR
-class UMaterial;
-#endif
 
 namespace Chaos
 {
@@ -27,9 +24,6 @@ namespace Chaos
 	typedef FClothingSimulationContextCommon FClothingSimulationContext;
 
 	class FClothingSimulation : public FClothingSimulationCommon
-#if WITH_EDITOR
-		, public FGCObject  // Add garbage collection for debug cloth material
-#endif  // #if WITH_EDITOR
 	{
 	public:
 		FClothingSimulation();
@@ -88,45 +82,37 @@ namespace Chaos
 		FClothingSimulationCloth* GetCloth(int32 ClothId);
 
 #if WITH_EDITOR
-		// FGCObject interface
-		virtual void AddReferencedObjects(FReferenceCollector& Collector) override;
-		virtual FString GetReferencerName() const override
-		{
-			return TEXT("Chaos::FClothingSimulation");
-		}
-		// End of FGCObject interface
-
 		// Editor only debug draw function
-		CHAOSCLOTH_API void DebugDrawPhysMeshShaded(FPrimitiveDrawInterface* PDI) const;
-		CHAOSCLOTH_API void DebugDrawParticleIndices(FCanvas* Canvas, const FSceneView* SceneView) const;
-		CHAOSCLOTH_API void DebugDrawElementIndices(FCanvas* Canvas, const FSceneView* SceneView) const;
-		CHAOSCLOTH_API void DebugDrawMaxDistanceValues(FCanvas* Canvas, const FSceneView* SceneView) const;
+		CHAOSCLOTH_API void DebugDrawPhysMeshShaded(FPrimitiveDrawInterface* PDI) const { Visualization.DrawPhysMeshShaded(PDI); }
+		CHAOSCLOTH_API void DebugDrawParticleIndices(FCanvas* Canvas, const FSceneView* SceneView) const { Visualization.DrawParticleIndices(Canvas, SceneView); }
+		CHAOSCLOTH_API void DebugDrawElementIndices(FCanvas* Canvas, const FSceneView* SceneView) const { Visualization.DrawElementIndices(Canvas, SceneView); }
+		CHAOSCLOTH_API void DebugDrawMaxDistanceValues(FCanvas* Canvas, const FSceneView* SceneView) const { Visualization.DrawMaxDistanceValues(Canvas, SceneView); }
 #endif  // #if WITH_EDITOR
 
-#if WITH_EDITOR || CHAOS_DEBUG_DRAW
+#if CHAOS_DEBUG_DRAW
 		// Editor & runtime debug draw functions
-		CHAOSCLOTH_API void DebugDrawPhysMeshWired(FPrimitiveDrawInterface* PDI = nullptr) const;
-		CHAOSCLOTH_API void DebugDrawAnimMeshWired(FPrimitiveDrawInterface* PDI = nullptr) const;
-		CHAOSCLOTH_API void DebugDrawAnimNormals(FPrimitiveDrawInterface* PDI = nullptr) const;
-		CHAOSCLOTH_API void DebugDrawPointNormals(FPrimitiveDrawInterface* PDI = nullptr) const;
-		CHAOSCLOTH_API void DebugDrawPointVelocities(FPrimitiveDrawInterface* PDI = nullptr) const;
-		UE_DEPRECATED(5.0, "DebugDrawInversedPointNormals is mostly redundant and will be removed, use DebugDrawPointNormals instead")
-		CHAOSCLOTH_API void DebugDrawInversedPointNormals(FPrimitiveDrawInterface* PDI = nullptr) const;
-		CHAOSCLOTH_API void DebugDrawCollision(FPrimitiveDrawInterface* PDI = nullptr) const;
-		CHAOSCLOTH_API void DebugDrawBackstops(FPrimitiveDrawInterface* PDI = nullptr) const;
-		CHAOSCLOTH_API void DebugDrawBackstopDistances(FPrimitiveDrawInterface* PDI = nullptr) const;
-		CHAOSCLOTH_API void DebugDrawMaxDistances(FPrimitiveDrawInterface* PDI = nullptr) const;
-		CHAOSCLOTH_API void DebugDrawAnimDrive(FPrimitiveDrawInterface* PDI = nullptr) const;
-		CHAOSCLOTH_API void DebugDrawEdgeConstraint(FPrimitiveDrawInterface* PDI = nullptr) const;
-		CHAOSCLOTH_API void DebugDrawBendingConstraint(FPrimitiveDrawInterface* PDI = nullptr) const;
-		CHAOSCLOTH_API void DebugDrawLongRangeConstraint(FPrimitiveDrawInterface* PDI = nullptr) const;
-		CHAOSCLOTH_API void DebugDrawWindAndPressureForces(FPrimitiveDrawInterface* PDI = nullptr) const;
-		UE_DEPRECATED(5.1, "Chaos::Softs::FVelocityField has been renamed FVelocityAndPressureField to match its new behavior.")
-		CHAOSCLOTH_API void DebugDrawWindForces(FPrimitiveDrawInterface* PDI = nullptr) const { return DebugDrawWindAndPressureForces(PDI); }
-		CHAOSCLOTH_API void DebugDrawLocalSpace(FPrimitiveDrawInterface* PDI = nullptr) const;
-		CHAOSCLOTH_API void DebugDrawSelfCollision(FPrimitiveDrawInterface* PDI = nullptr) const;
-		CHAOSCLOTH_API void DebugDrawSelfIntersection(FPrimitiveDrawInterface* PDI = nullptr) const;
-#endif  // #if WITH_EDITOR || CHAOS_DEBUG_DRAW
+		CHAOSCLOTH_API void DebugDrawPhysMeshWired(FPrimitiveDrawInterface* PDI = nullptr) const { Visualization.DrawPhysMeshWired(PDI); }
+		CHAOSCLOTH_API void DebugDrawAnimMeshWired(FPrimitiveDrawInterface* PDI = nullptr) const { Visualization.DrawAnimMeshWired(PDI); }
+		CHAOSCLOTH_API void DebugDrawAnimNormals(FPrimitiveDrawInterface* PDI = nullptr) const { Visualization.DrawAnimNormals(PDI); }
+		CHAOSCLOTH_API void DebugDrawPointNormals(FPrimitiveDrawInterface* PDI = nullptr) const { Visualization.DrawPointNormals(PDI); }
+		CHAOSCLOTH_API void DebugDrawPointVelocities(FPrimitiveDrawInterface* PDI = nullptr) const { Visualization.DrawPointVelocities(PDI); }
+		CHAOSCLOTH_API void DebugDrawCollision(FPrimitiveDrawInterface* PDI = nullptr) const { Visualization.DrawCollision(PDI); }
+		CHAOSCLOTH_API void DebugDrawBackstops(FPrimitiveDrawInterface* PDI = nullptr) const { Visualization.DrawBackstops(PDI); }
+		CHAOSCLOTH_API void DebugDrawBackstopDistances(FPrimitiveDrawInterface* PDI = nullptr) const { Visualization.DrawBackstopDistances(PDI); }
+		CHAOSCLOTH_API void DebugDrawMaxDistances(FPrimitiveDrawInterface* PDI = nullptr) const { Visualization.DrawMaxDistances(PDI); }
+		CHAOSCLOTH_API void DebugDrawAnimDrive(FPrimitiveDrawInterface* PDI = nullptr) const { Visualization.DrawAnimDrive(PDI); }
+		CHAOSCLOTH_API void DebugDrawEdgeConstraint(FPrimitiveDrawInterface* PDI = nullptr) const { Visualization.DrawEdgeConstraint(PDI); }
+		CHAOSCLOTH_API void DebugDrawBendingConstraint(FPrimitiveDrawInterface* PDI = nullptr) const { Visualization.DrawBendingConstraint(PDI); }
+		CHAOSCLOTH_API void DebugDrawLongRangeConstraint(FPrimitiveDrawInterface* PDI = nullptr) const { Visualization.DrawLongRangeConstraint(PDI); }
+		CHAOSCLOTH_API void DebugDrawWindAndPressureForces(FPrimitiveDrawInterface* PDI = nullptr) const { Visualization.DrawWindAndPressureForces(PDI); }
+		UE_DEPRECATED(5.1, "DebugDrawWindForces has been renamed DebugDrawWindAndPressureForces.")
+		CHAOSCLOTH_API void DebugDrawWindForces(FPrimitiveDrawInterface* PDI = nullptr) const { Visualization.DrawWindAndPressureForces(PDI); }
+		CHAOSCLOTH_API void DebugDrawLocalSpace(FPrimitiveDrawInterface* PDI = nullptr) const { Visualization.DrawLocalSpace(PDI); }
+		CHAOSCLOTH_API void DebugDrawSelfCollision(FPrimitiveDrawInterface* PDI = nullptr) const { Visualization.DrawSelfCollision(PDI); }
+		CHAOSCLOTH_API void DebugDrawSelfIntersection(FPrimitiveDrawInterface* PDI = nullptr) const { Visualization.DrawSelfIntersection(PDI); }
+		CHAOSCLOTH_API void DebugDrawBounds(FPrimitiveDrawInterface* PDI = nullptr) const { Visualization.DrawBounds(PDI); }
+		CHAOSCLOTH_API void DebugDrawGravity(FPrimitiveDrawInterface* PDI = nullptr) const { Visualization.DrawGravity(PDI); }
+#endif  // #if CHAOS_DEBUG_DRAW
 
 	private:
 		void ResetStats();
@@ -134,13 +120,10 @@ namespace Chaos
 
 		void UpdateSimulationFromSharedSimConfig();
 
-#if CHAOS_DEBUG_DRAW
-		// Runtime only debug draw functions
-		void DebugDrawBounds() const;
-		void DebugDrawGravity() const;
-#endif  // #if CHAOS_DEBUG_DRAW
-
 	private:
+		// Visualization object
+		UE::Chaos::Cloth::FVisualization Visualization;
+
 		// Simulation objects
 		TUniquePtr<FClothingSimulationSolver> Solver;  // Default solver
 PRAGMA_DISABLE_DEPRECATION_WARNINGS  // TODO: CHAOS_IS_CLOTHINGSIMULATIONMESH_ABSTRACT
@@ -170,12 +153,6 @@ PRAGMA_ENABLE_DEPRECATION_WARNINGS
 		FVector GravityOverride;
 		FReal MaxDistancesMultipliers;
 
-#if WITH_EDITOR
-		// Visualization material
-		UMaterial* DebugClothMaterial;
-		UMaterial* DebugClothMaterialVertex;
-#endif
-
 #if !(UE_BUILD_SHIPPING || UE_BUILD_TEST)
 		int32 StepCount;
 		int32 ResetCount;
