Chaos Cloth - Fix time smoothing effect on velocity scale when the time step is decreasing.

diff --git a/Engine/Plugins/Experimental/ChaosCloth/Source/Chaos/Private/ChaosCloth/ChaosClothingSimulationSolver.cpp b/Engine/Plugins/Experimental/ChaosCloth/Source/Chaos/Private/ChaosCloth/ChaosClothingSimulationSolver.cpp
--- a/Engine/Plugins/Experimental/ChaosCloth/Source/Chaos/Private/ChaosCloth/ChaosClothingSimulationSolver.cpp
+++ b/Engine/Plugins/Experimental/ChaosCloth/Source/Chaos/Private/ChaosCloth/ChaosClothingSimulationSolver.cpp
@@ -585,7 +585,7 @@ void FClothingSimulationSolver::SetReferenceVelocityScale(
 	const FRigidTransform3 DeltaTransform = ReferenceSpaceTransform.GetRelativeTransform(OldReferenceSpaceTransform);
 
 	// Apply linear velocity scale
-	const FVec3 LinearRatio = FVec3(1.) - FVec3(LinearVelocityScale).BoundToBox(FVec3(0.), FVec3(1.)) * SolverVelocityScale;
+	const FVec3 LinearRatio = FVec3(1.) - FVec3(LinearVelocityScale * SolverVelocityScale).BoundToBox(FVec3(0.), FVec3(1.));
 	const FVec3 DeltaPosition = LinearRatio * DeltaTransform.GetTranslation();
 
 	// Apply angular velocity scale
