Chaos Cloth:
- Add a Physics Asset property to the new Cloth Asset.
- Rename the Visualization class and move it back to the original Chaos namespace for convenience of use.
- Deprecate the USkeletalMeshComponent and UClothingAssetCommon dependent FClothingSimulationCollider constructor.
- Remove support for the legacy Apex collisions.

#okforgitub public

diff --git a/Engine/Plugins/Experimental/ChaosCloth/Source/Chaos/Private/ChaosCloth/ChaosClothVisualization.cpp b/Engine/Plugins/Experimental/ChaosCloth/Source/Chaos/Private/ChaosCloth/ChaosClothVisualization.cpp
--- a/Engine/Plugins/Experimental/ChaosCloth/Source/Chaos/Private/ChaosCloth/ChaosClothVisualization.cpp
+++ b/Engine/Plugins/Experimental/ChaosCloth/Source/Chaos/Private/ChaosCloth/ChaosClothVisualization.cpp
@@ -24,6 +24,7 @@
 #include "Chaos/XPBDBendingConstraints.h"
 #include "Chaos/XPBDSpringConstraints.h"
 #include "DynamicMeshBuilder.h"
+#include "Engine/EngineTypes.h"
 #if WITH_EDITOR
 #include "Materials/Material.h"
 #include "Engine/Canvas.h"  // For draw text
@@ -31,9 +32,9 @@
 #include "Engine/Engine.h"  //
 #endif  // #if WITH_EDITOR
 
-namespace UE::Chaos::Cloth
+namespace Chaos
 {
-	FVisualization::FVisualization(const ::Chaos::FClothingSimulationSolver* InSolver)
+	FClothVisualization::FClothVisualization(const ::Chaos::FClothingSimulationSolver* InSolver)
 		: Solver(InSolver)
 	{
 #if WITH_EDITOR
@@ -42,24 +43,22 @@ namespace UE::Chaos::Cloth
 #endif  // #if WITH_EDITOR
 	}
 
-	FVisualization::~FVisualization() = default;
+	FClothVisualization::~FClothVisualization() = default;
 
-	void FVisualization::SetSolver(const ::Chaos::FClothingSimulationSolver* InSolver)
+	void FClothVisualization::SetSolver(const ::Chaos::FClothingSimulationSolver* InSolver)
 	{
 		Solver = InSolver;
 	}
 
 #if WITH_EDITOR
-	void FVisualization::AddReferencedObjects(FReferenceCollector& Collector)
+	void FClothVisualization::AddReferencedObjects(FReferenceCollector& Collector)
 	{
 		Collector.AddReferencedObject(ClothMaterial);
 		Collector.AddReferencedObject(ClothMaterialVertex);
 	}
 
-	void FVisualization::DrawPhysMeshShaded(FPrimitiveDrawInterface* PDI) const
+	void FClothVisualization::DrawPhysMeshShaded(FPrimitiveDrawInterface* PDI) const
 	{
-		using namespace ::Chaos;
-
 		if (!Solver || !ClothMaterial)
 		{
 			return;
@@ -120,10 +119,8 @@ namespace UE::Chaos::Cloth
 		}
 	}
 
-	void FVisualization::DrawParticleIndices(FCanvas* Canvas, const FSceneView* SceneView) const
+	void FClothVisualization::DrawParticleIndices(FCanvas* Canvas, const FSceneView* SceneView) const
 	{
-		using namespace ::Chaos;
-
 		if (!Solver)
 		{
 			return;
@@ -156,10 +153,8 @@ namespace UE::Chaos::Cloth
 		}
 	}
 
-	void FVisualization::DrawElementIndices(FCanvas* Canvas, const FSceneView* SceneView) const
+	void FClothVisualization::DrawElementIndices(FCanvas* Canvas, const FSceneView* SceneView) const
 	{
-		using namespace ::Chaos;
-
 		if (!Solver)
 		{
 			return;
@@ -198,10 +193,8 @@ namespace UE::Chaos::Cloth
 		}
 	}
 
-	void FVisualization::DrawMaxDistanceValues(FCanvas* Canvas, const FSceneView* SceneView) const
+	void FClothVisualization::DrawMaxDistanceValues(FCanvas* Canvas, const FSceneView* SceneView) const
 	{
-		using namespace ::Chaos;
-
 		if (!Solver)
 		{
 			return;
@@ -254,8 +247,6 @@ namespace UE::Chaos::Cloth
 
 	static void DrawPoint(FPrimitiveDrawInterface* PDI, const FVector& Pos, const FLinearColor& Color, const UMaterial* ClothMaterialVertex, const float Thickness = 1.f)  // Use color or material
 	{
-		using namespace ::Chaos;
-
 		if (!PDI)
 		{
 			FDebugDrawQueue::GetInstance().DrawDebugPoint(Pos, Color.ToFColor(true), false, KINDA_SMALL_NUMBER, SDPG_Foreground, Thickness);
@@ -278,8 +269,6 @@ namespace UE::Chaos::Cloth
 
 	static void DrawLine(FPrimitiveDrawInterface* PDI, const FVector& Pos0, const FVector& Pos1, const FLinearColor& Color)
 	{
-		using namespace ::Chaos;
-
 		if (!PDI)
 		{
 			FDebugDrawQueue::GetInstance().DrawDebugLine(Pos0, Pos1, Color.ToFColor(true), false, KINDA_SMALL_NUMBER, SDPG_Foreground, 0.f);
@@ -292,8 +281,6 @@ namespace UE::Chaos::Cloth
 
 	static void DrawArc(FPrimitiveDrawInterface* PDI, const FVector& Base, const FVector& X, const FVector& Y, ::Chaos::FReal MinAngle, float MaxAngle, ::Chaos::FReal Radius, const FLinearColor& Color)
 	{
-		using namespace ::Chaos;
-
 		static const int32 Sections = 10;
 		const FReal AngleStep = FMath::DegreesToRadians((MaxAngle - MinAngle) / (FReal)Sections);
 		FReal CurrentAngle = FMath::DegreesToRadians(MinAngle);
@@ -310,8 +297,6 @@ namespace UE::Chaos::Cloth
 
 	static void DrawSphere(FPrimitiveDrawInterface* PDI, const ::Chaos::TSphere<::Chaos::FReal, 3>& Sphere, const FQuat& Rotation, const FVector& Position, const FLinearColor& Color)
 	{
-		using namespace ::Chaos;
-
 		const FReal Radius = Sphere.GetRadius();
 		const FVec3 Center = Position + Rotation.RotateVector(Sphere.GetCenter());
 		if (!PDI)
@@ -327,8 +312,6 @@ namespace UE::Chaos::Cloth
 
 	static void DrawBox(FPrimitiveDrawInterface* PDI, const ::Chaos::FAABB3& Box, const FQuat& Rotation, const FVector& Position, const FLinearColor& Color)
 	{
-		using namespace ::Chaos;
-
 		if (!PDI)
 		{
 			const FVec3 Center = Position + Rotation.RotateVector(Box.GetCenter());
@@ -343,8 +326,6 @@ namespace UE::Chaos::Cloth
 
 	static void DrawCapsule(FPrimitiveDrawInterface* PDI, const ::Chaos::FCapsule& Capsule, const FQuat& Rotation, const FVector& Position, const FLinearColor& Color)
 	{
-		using namespace ::Chaos;
-
 		const FReal Radius = Capsule.GetRadius();
 		const FReal HalfHeight = Capsule.GetHeight() * 0.5f + Radius;
 		const FVec3 Center = Position + Rotation.RotateVector(Capsule.GetCenter());
@@ -367,8 +348,6 @@ namespace UE::Chaos::Cloth
 
 	static void DrawTaperedCylinder(FPrimitiveDrawInterface* PDI, const ::Chaos::FTaperedCylinder& TaperedCylinder, const FQuat& Rotation, const FVector& Position, const FLinearColor& Color)
 	{
-		using namespace ::Chaos;
-
 		const FReal HalfHeight = TaperedCylinder.GetHeight() * 0.5f;
 		const FReal Radius1 = TaperedCylinder.GetRadius1();
 		const FReal Radius2 = TaperedCylinder.GetRadius2();
@@ -401,8 +380,6 @@ namespace UE::Chaos::Cloth
 
 	static void DrawConvex(FPrimitiveDrawInterface* PDI, const ::Chaos::FConvex& Convex, const FQuat& Rotation, const FVector& Position, const FLinearColor& Color)
 	{
-		using namespace ::Chaos;
-
 		const TArray<FConvex::FPlaneType>& Planes = Convex.GetFaces();
 		for (int32 PlaneIndex1 = 0; PlaneIndex1 < Planes.Num(); ++PlaneIndex1)
 		{
@@ -450,10 +427,8 @@ namespace UE::Chaos::Cloth
 		DrawLine(PDI, Position, Position + Z, FLinearColor::Blue);
 	}
 
-	void FVisualization::DrawBounds(FPrimitiveDrawInterface* PDI) const
+	void FClothVisualization::DrawBounds(FPrimitiveDrawInterface* PDI) const
 	{
-		using namespace ::Chaos;
-
 		if (!Solver)
 		{
 			return;
@@ -480,10 +455,8 @@ namespace UE::Chaos::Cloth
 		}
 	}
 
-	void FVisualization::DrawGravity(FPrimitiveDrawInterface* PDI) const
+	void FClothVisualization::DrawGravity(FPrimitiveDrawInterface* PDI) const
 	{
-		using namespace ::Chaos;
-
 		if (!Solver)
 		{
 			return;
@@ -505,10 +478,8 @@ namespace UE::Chaos::Cloth
 		}
 	}
 
-	void FVisualization::DrawPhysMeshWired(FPrimitiveDrawInterface* PDI) const
+	void FClothVisualization::DrawPhysMeshWired(FPrimitiveDrawInterface* PDI) const
 	{
-		using namespace ::Chaos;
-
 		if (!Solver)
 		{
 			return;
@@ -551,10 +522,8 @@ namespace UE::Chaos::Cloth
 		}
 	}
 
-	void FVisualization::DrawAnimMeshWired(FPrimitiveDrawInterface* PDI) const
+	void FClothVisualization::DrawAnimMeshWired(FPrimitiveDrawInterface* PDI) const
 	{
-		using namespace ::Chaos;
-
 		if (!Solver)
 		{
 			return;
@@ -590,10 +559,8 @@ namespace UE::Chaos::Cloth
 		}
 	}
 
-	void FVisualization::DrawAnimNormals(FPrimitiveDrawInterface* PDI) const
+	void FClothVisualization::DrawAnimNormals(FPrimitiveDrawInterface* PDI) const
 	{
-		using namespace ::Chaos;
-
 		if (!Solver)
 		{
 			return;
@@ -626,10 +593,8 @@ namespace UE::Chaos::Cloth
 		}
 	}
 
-	void FVisualization::DrawPointNormals(FPrimitiveDrawInterface* PDI) const
+	void FClothVisualization::DrawPointNormals(FPrimitiveDrawInterface* PDI) const
 	{
-		using namespace ::Chaos;
-
 		if (!Solver)
 		{
 			return;
@@ -666,10 +631,8 @@ namespace UE::Chaos::Cloth
 		}
 	}
 
-	void FVisualization::DrawPointVelocities(FPrimitiveDrawInterface* PDI) const
+	void FClothVisualization::DrawPointVelocities(FPrimitiveDrawInterface* PDI) const
 	{
-		using namespace ::Chaos;
-
 		if (!Solver)
 		{
 			return;
@@ -705,10 +668,8 @@ namespace UE::Chaos::Cloth
 		}
 	}
 
-	void FVisualization::DrawCollision(FPrimitiveDrawInterface* PDI) const
+	void FClothVisualization::DrawCollision(FPrimitiveDrawInterface* PDI) const
 	{
-		using namespace ::Chaos;
-
 		if (!Solver)
 		{
 			return;
@@ -841,10 +802,8 @@ namespace UE::Chaos::Cloth
 		}
 	}
 
-	void FVisualization::DrawBackstops(FPrimitiveDrawInterface* PDI) const
+	void FClothVisualization::DrawBackstops(FPrimitiveDrawInterface* PDI) const
 	{
-		using namespace ::Chaos;
-
 		if (!Solver)
 		{
 			return;
@@ -913,10 +872,8 @@ namespace UE::Chaos::Cloth
 		}
 	}
 
-	void FVisualization::DrawBackstopDistances(FPrimitiveDrawInterface* PDI) const
+	void FClothVisualization::DrawBackstopDistances(FPrimitiveDrawInterface* PDI) const
 	{
-		using namespace ::Chaos;
-
 		if (!Solver)
 		{
 			return;
@@ -963,10 +920,8 @@ namespace UE::Chaos::Cloth
 		}
 	}
 
-	void FVisualization::DrawMaxDistances(FPrimitiveDrawInterface* PDI) const
+	void FClothVisualization::DrawMaxDistances(FPrimitiveDrawInterface* PDI) const
 	{
-		using namespace ::Chaos;
-
 		if (!Solver)
 		{
 			return;
@@ -1016,10 +971,8 @@ namespace UE::Chaos::Cloth
 		}
 	}
 
-	void FVisualization::DrawAnimDrive(FPrimitiveDrawInterface* PDI) const
+	void FClothVisualization::DrawAnimDrive(FPrimitiveDrawInterface* PDI) const
 	{
-		using namespace ::Chaos;
-
 		if (!Solver)
 		{
 			return;
@@ -1063,8 +1016,6 @@ namespace UE::Chaos::Cloth
 
 	static void DrawSpringConstraintColors(FPrimitiveDrawInterface* PDI, const TConstArrayView<::Chaos::Softs::FSolverVec3>& Positions, const ::Chaos::FVec3& LocalSpaceLocation, const ::Chaos::Softs::FPBDSpringConstraints* const SpringConstraints)
 	{
-		using namespace ::Chaos;
-
 		check(SpringConstraints);
 
 		const TArray<TVec2<int32>>& Constraints = SpringConstraints->GetConstraints();
@@ -1107,10 +1058,8 @@ namespace UE::Chaos::Cloth
 		}
 	}
 
-	void FVisualization::DrawEdgeConstraint(FPrimitiveDrawInterface* PDI) const
+	void FClothVisualization::DrawEdgeConstraint(FPrimitiveDrawInterface* PDI) const
 	{
-		using namespace ::Chaos;
-
 		if (!Solver)
 		{
 			return;
@@ -1138,10 +1087,8 @@ namespace UE::Chaos::Cloth
 		}
 	}
 
-	void FVisualization::DrawBendingConstraint(FPrimitiveDrawInterface* PDI) const
+	void FClothVisualization::DrawBendingConstraint(FPrimitiveDrawInterface* PDI) const
 	{
-		using namespace ::Chaos;
-
 		if (!Solver)
 		{
 			return;
@@ -1207,10 +1154,8 @@ namespace UE::Chaos::Cloth
 		}
 	}
 
-	void FVisualization::DrawLongRangeConstraint(FPrimitiveDrawInterface* PDI) const
+	void FClothVisualization::DrawLongRangeConstraint(FPrimitiveDrawInterface* PDI) const
 	{
-		using namespace ::Chaos;
-
 		if (!Solver)
 		{
 			return;
@@ -1295,10 +1240,8 @@ namespace UE::Chaos::Cloth
 		}
 	}
 
-	void FVisualization::DrawWindAndPressureForces(FPrimitiveDrawInterface* PDI) const
+	void FClothVisualization::DrawWindAndPressureForces(FPrimitiveDrawInterface* PDI) const
 	{
-		using namespace ::Chaos;
-
 		if (!Solver)
 		{
 			return;
@@ -1342,10 +1285,8 @@ namespace UE::Chaos::Cloth
 		}
 	}
 
-	void FVisualization::DrawLocalSpace(FPrimitiveDrawInterface* PDI) const
+	void FClothVisualization::DrawLocalSpace(FPrimitiveDrawInterface* PDI) const
 	{
-		using namespace ::Chaos;
-
 		if (!Solver)
 		{
 			return;
@@ -1366,10 +1307,8 @@ namespace UE::Chaos::Cloth
 		}
 	}
 
-	void FVisualization::DrawSelfCollision(FPrimitiveDrawInterface* PDI) const
+	void FClothVisualization::DrawSelfCollision(FPrimitiveDrawInterface* PDI) const
 	{
-		using namespace ::Chaos;
-
 		if (!Solver)
 		{
 			return;
@@ -1429,10 +1368,8 @@ namespace UE::Chaos::Cloth
 		}
 	}
 
-	void FVisualization::DrawSelfIntersection(FPrimitiveDrawInterface* PDI) const
+	void FClothVisualization::DrawSelfIntersection(FPrimitiveDrawInterface* PDI) const
 	{
-		using namespace ::Chaos;
-
 		if (!Solver)
 		{
 			return;
@@ -1549,12 +1486,12 @@ namespace UE::Chaos::Cloth
 			}
 		}
 	}
-}  // End namespace UE::Chaos::Cloth
+}  // End namespace Chaos
 #else  // #if CHAOS_DEBUG_DRAW
-namespace UE::Chaos::Cloth
+namespace Chaos
 {
-	FVisualization::FVisualization(const ::Chaos::FClothingSimulationSolver* /*InSolver*/) {}
+	FClothVisualization::FClothVisualization(const ::Chaos::FClothingSimulationSolver* /*InSolver*/) {}
 
-	FVisualization::~FVisualization() = default;
+	FClothVisualization::~FClothVisualization() = default;
 }  // End namespace UE::Chaos::Cloth
 #endif  // #if CHAOS_DEBUG_DRAW
diff --git a/Engine/Plugins/Experimental/ChaosCloth/Source/Chaos/Private/ChaosCloth/ChaosClothingSimulation.cpp b/Engine/Plugins/Experimental/ChaosCloth/Source/Chaos/Private/ChaosCloth/ChaosClothingSimulation.cpp
--- a/Engine/Plugins/Experimental/ChaosCloth/Source/Chaos/Private/ChaosCloth/ChaosClothingSimulation.cpp
+++ b/Engine/Plugins/Experimental/ChaosCloth/Source/Chaos/Private/ChaosCloth/ChaosClothingSimulation.cpp
@@ -358,10 +358,8 @@ void FClothingSimulation::CreateActor(USkeletalMeshComponent* InOwnerComponent,
 
 	// Create collider node
 	const int32 ColliderIndex = Colliders.Emplace(MakeUnique<FClothingSimulationCollider>(
-		Asset,
-		InOwnerComponent,
-		/*bInUseLODIndexOverride =*/ false,
-		/*InLODIndexOverride =*/ INDEX_NONE));
+		Asset->PhysicsAsset,
+		&CastChecked<USkeletalMesh>(Asset->GetOuter())->GetRefSkeleton()));
 
 	// Set the external collision data to get updated at every frame
 	Colliders[ColliderIndex]->SetCollisionData(&ExternalCollisionData);
diff --git a/Engine/Plugins/Experimental/ChaosCloth/Source/Chaos/Private/ChaosCloth/ChaosClothingSimulationCollider.cpp b/Engine/Plugins/Experimental/ChaosCloth/Source/Chaos/Private/ChaosCloth/ChaosClothingSimulationCollider.cpp
--- a/Engine/Plugins/Experimental/ChaosCloth/Source/Chaos/Private/ChaosCloth/ChaosClothingSimulationCollider.cpp
+++ b/Engine/Plugins/Experimental/ChaosCloth/Source/Chaos/Private/ChaosCloth/ChaosClothingSimulationCollider.cpp
@@ -2,6 +2,7 @@
 #include "ChaosCloth/ChaosClothingSimulationCollider.h"
 #include "ChaosCloth/ChaosClothingSimulationSolver.h"
 #include "ChaosCloth/ChaosClothingSimulationCloth.h"
+#include "ChaosCloth/ChaosClothingSimulationMesh.h"
 #include "ChaosCloth/ChaosClothPrivate.h"
 #include "Chaos/Capsule.h"
 #include "Chaos/Plane.h"
@@ -11,10 +12,8 @@
 #include "Chaos/Convex.h"
 #include "Chaos/Levelset.h"
 #include "Chaos/ImplicitObjectTransformed.h"
-#include "ClothingSimulation.h" // For context
 #include "ClothingAsset.h"
 #include "Engine/SkeletalMesh.h"
-#include "Components/SkeletalMeshComponent.h"
 #include "PhysicsEngine/PhysicsAsset.h"
 #include "PhysicsEngine/AggregateGeom.h"
 #include "HAL/PlatformMath.h"
@@ -48,7 +47,7 @@ struct FClothingSimulationCollider::FLODData
 		const TArray<int32>& UsedBoneIndices = TArray<int32>());
 	void Remove(FClothingSimulationSolver* Solver, FClothingSimulationCloth* Cloth);
 
-	void Update(FClothingSimulationSolver* Solver, FClothingSimulationCloth* Cloth, const FClothingSimulationContextCommon* Context);
+	void Update(FClothingSimulationSolver* Solver, FClothingSimulationCloth* Cloth, const FTransform& ComponentTransform, const TArray<FTransform>& BoneTransforms);
 
 	void Enable(FClothingSimulationSolver* Solver, FClothingSimulationCloth* Cloth, bool bEnable);
 
@@ -345,7 +344,11 @@ void FClothingSimulationCollider::FLODData::Remove(FClothingSimulationSolver* So
 	Offsets.Remove(FSolverClothPair(Solver, Cloth));
 }
 
-void FClothingSimulationCollider::FLODData::Update(FClothingSimulationSolver* Solver, FClothingSimulationCloth* Cloth, const FClothingSimulationContextCommon* Context)
+void FClothingSimulationCollider::FLODData::Update(
+	FClothingSimulationSolver* Solver,
+	FClothingSimulationCloth* Cloth,
+	const FTransform& ComponentTransform,
+	const TArray<FTransform>& BoneTransforms)
 {
 	check(Solver);
 	check(Cloth);
@@ -356,9 +359,7 @@ void FClothingSimulationCollider::FLODData::Update(FClothingSimulationSolver* So
 		const Softs::FSolverRigidTransform3* BaseTransforms = Solver->GetCollisionBaseTransforms(Offset);
 		Softs::FSolverRigidTransform3* const CollisionTransforms = Solver->GetCollisionTransforms(Offset);
 
-		const TArrayView<const FTransform> BoneTransforms = Context ? Context->BoneTransforms : TArrayView<const FTransform>();
-
-		FTransform ComponentToLocalSpaceReal = Context ? Context->ComponentToWorld : FTransform::Identity;
+		FTransform ComponentToLocalSpaceReal = ComponentTransform;
 		ComponentToLocalSpaceReal.AddToTranslation(-Solver->GetLocalSpaceLocation());
 		const Softs::FSolverTransform3 ComponentToLocalSpace(ComponentToLocalSpaceReal);  // LWC, now in local space, therefore it is safe to use the solver transform type
 	
@@ -395,20 +396,31 @@ void FClothingSimulationCollider::FLODData::ResetStartPose(FClothingSimulationSo
 	}
 }
 
+FClothingSimulationCollider::FClothingSimulationCollider(const UPhysicsAsset* InPhysicsAsset, const FReferenceSkeleton* InReferenceSkeleton)
+	: PhysicsAsset(InPhysicsAsset)
+	, ReferenceSkeleton(InReferenceSkeleton)
+{
+	// Prepare LOD array
+	// Note: There aren't any collision LODs stored with this constructor version anymore,
+	//       it's only this physics asset's and external collisions.
+	const int32 NumLODData = (int32)ECollisionDataType::LODs;
+	LODData.Reserve(NumLODData);
+	for (int32 Index = 0; Index < NumLODData; ++Index)
+	{
+		LODData.Add(MakeUnique<FLODData>());
+	}
+}
+
 FClothingSimulationCollider::FClothingSimulationCollider(
 	const UClothingAssetCommon* InAsset,
-	const USkeletalMeshComponent* InSkeletalMeshComponent,
-	bool bInUseLODIndexOverride,
-	int32 InLODIndexOverride)
-	: Asset(InAsset)
-	, SkeletalMeshComponent(InSkeletalMeshComponent)
-	, CollisionData(nullptr)
-	, bUseLODIndexOverride(bInUseLODIndexOverride)
-	, LODIndexOverride(InLODIndexOverride)
-	, Scale(1.f)
+	const USkeletalMeshComponent* /*InSkeletalMeshComponent*/,
+	bool /*bInUseLODIndexOverride*/,
+	int32 /*InLODIndexOverride*/)
+	: PhysicsAsset(InAsset->PhysicsAsset)
+	, ReferenceSkeleton(&CastChecked<USkeletalMesh>(InAsset->GetOuter())->GetRefSkeleton())
 {
 	// Prepare LOD array
-	const int32 NumLODs = Asset ? Asset->LodData.Num() : 0;
+	const int32 NumLODs = InAsset ? InAsset->LodData.Num() : 0;
 	const int32 NumLODData = (int32)ECollisionDataType::LODs + NumLODs;
 	LODData.Reserve(NumLODData);
 	for (int32 Index = 0; Index < NumLODData; ++Index)
@@ -450,25 +462,18 @@ void FClothingSimulationCollider::ExtractPhysicsAssetCollision(FClothCollisionDa
 	ClothCollisionData.Reset();
 	UsedBoneIndices.Reset();
 
-	if (!Asset)
-	{
-		return;
-	}
-
-	if (const UPhysicsAsset* const PhysAsset = Asset->PhysicsAsset)
+	if (PhysicsAsset)
 	{
-		const USkeletalMesh* const TargetMesh = CastChecked<USkeletalMesh>(Asset->GetOuter());
+		UsedBoneIndices.Reserve(PhysicsAsset->SkeletalBodySetups.Num());
 
-		UsedBoneIndices.Reserve(PhysAsset->SkeletalBodySetups.Num());
-
-		for (const USkeletalBodySetup* BodySetup : PhysAsset->SkeletalBodySetups)
+		for (const USkeletalBodySetup* BodySetup : PhysicsAsset->SkeletalBodySetups)
 		{
 			if (!BodySetup)
 			{
 				continue;
 			}
 
-			const int32 MeshBoneIndex = TargetMesh->GetRefSkeleton().FindBoneIndex(BodySetup->BoneName);
+			const int32 MeshBoneIndex = ReferenceSkeleton ? ReferenceSkeleton->FindBoneIndex(BodySetup->BoneName) : INDEX_NONE;
 			const int32 MappedBoneIndex = UsedBoneIndices.Add(MeshBoneIndex);
 			
 			// Add capsules
@@ -614,9 +619,8 @@ void FClothingSimulationCollider::ExtractPhysicsAssetCollision(FClothCollisionDa
 					LevelSetCollisions.Add({ LevelSetElem.GetLevelSet(), LevelSetElem.GetTransform(), MappedBoneIndex });
 				}
 			}
-
 		}  // End for PhysAsset->SkeletalBodySetups
-	}  // End if Asset->PhysicsAsset
+	}  // End if PhysicsAsset
 }
 
 
@@ -654,6 +658,7 @@ void FClothingSimulationCollider::Add(FClothingSimulationSolver* Solver, FClothi
 {
 	check(Solver);
 	check(Cloth);
+	check(Cloth->GetMesh());
 
 	// Can't add a collider twice to the same solver/cloth pair
 	check(!LODIndices.Find(FSolverClothPair(Solver, Cloth)));
@@ -663,16 +668,16 @@ void FClothingSimulationCollider::Add(FClothingSimulationSolver* Solver, FClothi
 	LODIndex = INDEX_NONE;
 
 	// Initialize scale
-	const FClothingSimulationContextCommon* const Context = SkeletalMeshComponent ? static_cast<const FClothingSimulationContextCommon*>(SkeletalMeshComponent->GetClothingSimulationContext()) : nullptr;
-	const FVec3 Scale3D = Context ? (FVec3)Context->ComponentToWorld.GetScale3D() : FVec3(1.f);
+PRAGMA_DISABLE_DEPRECATION_WARNINGS
+	const FVec3 Scale3D = (FVec3)Cloth->GetMesh()->GetComponentToWorldTransform().GetScale3D();
 	UE_CLOG(FMath::Abs(Scale3D.X - Scale3D.Y) > KINDA_SMALL_NUMBER || FMath::Abs(Scale3D.X - Scale3D.Z) > KINDA_SMALL_NUMBER,
 		LogChaosCloth, Warning, TEXT(
-			"Actor '%s' component '%s' has a non uniform scale, and has a cloth simulation attached. "
+			"Object '%s' has a non uniform scale, and has a cloth simulation attached. "
 			"The collision volumes might no longer correctly match the shape of the mesh. "
 			"Please update this component transform scale with the same value for all scale axis."),
-		SkeletalMeshComponent->GetOwner() ? *SkeletalMeshComponent->GetOwner()->GetName() : TEXT("None"),
-		*SkeletalMeshComponent->GetName());
+		*Cloth->GetMesh()->GetDebugName());
 	Scale = Scale3D.X;
+PRAGMA_ENABLE_DEPRECATION_WARNINGS
 
 	// Create physics asset collisions, this will affect all LODs, so store it at index 0
 	FClothCollisionData PhysicsAssetCollisionData;
@@ -681,29 +686,6 @@ void FClothingSimulationCollider::Add(FClothingSimulationSolver* Solver, FClothi
 	ExtractPhysicsAssetCollision(PhysicsAssetCollisionData, LevelSetCollisions, UsedBoneIndices);
 
 	LODData[(int32)ECollisionDataType::LODless]->Add(Solver, Cloth, PhysicsAssetCollisionData, LevelSetCollisions, Scale, UsedBoneIndices);
-
-	// Create legacy asset LOD collisions
-	const int32 NumLODs = Asset ? Asset->LodData.Num() : 0;
-	for (int32 Index = 0; Index < NumLODs; ++Index)
-	{
-		// Warn about legacy apex collisions
-		const FClothCollisionData& AssetCollisionData = Asset->LodData[Index].CollisionData;
-		UE_CLOG(AssetCollisionData.Spheres.Num() > 0 || AssetCollisionData.SphereConnections.Num() > 0 || AssetCollisionData.Convexes.Num() > 0,
-			LogChaosCloth, Warning, TEXT(
-				"Actor '%s' component '%s' has %d sphere, %d capsule, and %d convex collision objects for "
-				"physics authored as part of a LOD construct, probably by the Apex cloth authoring system. "
-				"This is deprecated. Please update your asset!"),
-			SkeletalMeshComponent->GetOwner() ? *SkeletalMeshComponent->GetOwner()->GetName() : TEXT("None"),
-			*SkeletalMeshComponent->GetName(),
-			AssetCollisionData.Spheres.Num(),
-			AssetCollisionData.SphereConnections.Num(),
-			AssetCollisionData.Convexes.Num());
-
-		const TArray<FLevelSetCollisionData> AssetLevelSetCollisions;		// Shouldn't be any level set collider on legacy asset
-
-		// Add legacy collision
-		LODData[(int32)ECollisionDataType::LODs + Index]->Add(Solver, Cloth, AssetCollisionData, AssetLevelSetCollisions, Scale, Asset->UsedBoneIndices);
-	}
 }
 
 void FClothingSimulationCollider::Remove(FClothingSimulationSolver* Solver, FClothingSimulationCloth* Cloth)
@@ -743,21 +725,25 @@ void FClothingSimulationCollider::Update(FClothingSimulationSolver* Solver, FClo
 {
 	check(Solver);
 	check(Cloth);
+	check(Cloth->GetMesh());
 	TRACE_CPUPROFILER_EVENT_SCOPE(FClothingSimulationCollider_Update);
 	SCOPE_CYCLE_COUNTER(STAT_ChaosClothingSimulationColliderUpdate);
 
 	// Update the collision transforms
 	// Note: All the LODs are updated at once, so that the previous transforms are always correct during LOD switching
-	const FClothingSimulationContextCommon* const Context = SkeletalMeshComponent ? static_cast<const FClothingSimulationContextCommon*>(SkeletalMeshComponent->GetClothingSimulationContext()) : nullptr;
+PRAGMA_DISABLE_DEPRECATION_WARNINGS
+	const FTransform& ComponentTransform = Cloth->GetMesh()->GetComponentToWorldTransform();
+	const TArray<FTransform>& BoneTransforms = Cloth->GetMesh()->GetBoneTransforms();
+PRAGMA_ENABLE_DEPRECATION_WARNINGS
 	for (TUniquePtr<FLODData>& LODDatum : LODData)
 	{
-		LODDatum->Update(Solver, Cloth, Context);
+		LODDatum->Update(Solver, Cloth, ComponentTransform, BoneTransforms);
 	}
 
 	// Update current LOD index
 	int32& LODIndex = LODIndices.FindChecked(FSolverClothPair(Solver, Cloth));
 	const int32 PrevLODIndex = LODIndex;
-	LODIndex = (int32)ECollisionDataType::LODs + (bUseLODIndexOverride ? LODIndexOverride : Cloth ? Cloth->GetLODIndex(Solver) : INDEX_NONE);
+	LODIndex = (int32)ECollisionDataType::LODs + (Cloth ? Cloth->GetLODIndex(Solver) : INDEX_NONE);
 	if (!LODData.IsValidIndex(LODIndex))
 	{
 		LODIndex = (int32)ECollisionDataType::LODs + INDEX_NONE;
diff --git a/Engine/Plugins/Experimental/ChaosCloth/Source/Chaos/Private/ChaosCloth/ChaosClothingSimulationMesh.cpp b/Engine/Plugins/Experimental/ChaosCloth/Source/Chaos/Private/ChaosCloth/ChaosClothingSimulationMesh.cpp
--- a/Engine/Plugins/Experimental/ChaosCloth/Source/Chaos/Private/ChaosCloth/ChaosClothingSimulationMesh.cpp
+++ b/Engine/Plugins/Experimental/ChaosCloth/Source/Chaos/Private/ChaosCloth/ChaosClothingSimulationMesh.cpp
@@ -38,15 +38,23 @@ namespace Chaos
 
 PRAGMA_DISABLE_DEPRECATION_WARNINGS  // TODO: CHAOS_IS_CLOTHINGSIMULATIONMESH_ABSTRACT
 
-FClothingSimulationMesh::FClothingSimulationMesh()
+FClothingSimulationMesh::FClothingSimulationMesh(const FString& InDebugName)
 	: Asset(nullptr)
 	, SkeletalMeshComponent(nullptr)
+#if !UE_BUILD_SHIPPING
+	, DebugName(InDebugName)
+#endif
 {
 }
 
 FClothingSimulationMesh::FClothingSimulationMesh(const UClothingAssetCommon* InAsset, const USkeletalMeshComponent* InSkeletalMeshComponent)
 	: Asset(InAsset)
 	, SkeletalMeshComponent(InSkeletalMeshComponent)
+#if !UE_BUILD_SHIPPING
+	, DebugName(InSkeletalMeshComponent->GetOwner() ?
+		FString::Format(TEXT("{0}|{1}"), { InSkeletalMeshComponent->GetOwner()->GetName(), InSkeletalMeshComponent->GetName() } ) :
+		InSkeletalMeshComponent->GetName())
+#endif
 {
 }
 
@@ -170,7 +178,7 @@ int32 FClothingSimulationMesh::GetReferenceBoneIndex() const
 	return Asset ? Asset->ReferenceBoneIndex : INDEX_NONE;
 }
 
-FRigidTransform3 FClothingSimulationMesh::GetReferenceBoneTransform() const
+FTransform FClothingSimulationMesh::GetReferenceBoneTransform() const
 {
 	if (const FClothingSimulationContextCommon* const Context = GetContext())
 	{
@@ -181,7 +189,12 @@ FRigidTransform3 FClothingSimulationMesh::GetReferenceBoneTransform() const
 			BoneTransforms[ReferenceBoneIndex] * Context->ComponentToWorld :
 			Context->ComponentToWorld;
 	}
-	return FRigidTransform3::Identity;
+	return FTransform::Identity;
+}
+
+const TArray<FTransform>& FClothingSimulationMesh::GetBoneTransforms() const
+{
+	return GetContext()->BoneTransforms;
 }
 
 const FTransform& FClothingSimulationMesh::GetComponentToWorldTransform() const
diff --git a/Engine/Plugins/Experimental/ChaosCloth/Source/Chaos/Private/ChaosCloth/ChaosClothingSimulationSkeletalMesh.cpp b/Engine/Plugins/Experimental/ChaosCloth/Source/Chaos/Private/ChaosCloth/ChaosClothingSimulationSkeletalMesh.cpp
--- a/Engine/Plugins/Experimental/ChaosCloth/Source/Chaos/Private/ChaosCloth/ChaosClothingSimulationSkeletalMesh.cpp
+++ b/Engine/Plugins/Experimental/ChaosCloth/Source/Chaos/Private/ChaosCloth/ChaosClothingSimulationSkeletalMesh.cpp
@@ -10,7 +10,9 @@ namespace Chaos
 {
 	FClothingSimulationSkeletalMesh::FClothingSimulationSkeletalMesh(const UClothingAssetCommon* InAsset, const USkeletalMeshComponent* InSkeletalMeshComponent)
 PRAGMA_DISABLE_DEPRECATION_WARNINGS  // TODO: CHAOS_IS_CLOTHINGSIMULATIONMESH_ABSTRACT
-		: FClothingSimulationMesh()
+		: FClothingSimulationMesh(InSkeletalMeshComponent->GetOwner() ?
+				FString::Format(TEXT("{0}|{1}"), { InSkeletalMeshComponent->GetOwner()->GetName(), InSkeletalMeshComponent->GetName() }) :
+				InSkeletalMeshComponent->GetName())
 PRAGMA_ENABLE_DEPRECATION_WARNINGS
 		, Asset(InAsset)
 		, SkeletalMeshComponent(InSkeletalMeshComponent)
@@ -135,7 +137,7 @@ PRAGMA_ENABLE_DEPRECATION_WARNINGS
 		return Asset ? Asset->ReferenceBoneIndex : INDEX_NONE;
 	}
 
-	FRigidTransform3 FClothingSimulationSkeletalMesh::GetReferenceBoneTransform() const
+	FTransform FClothingSimulationSkeletalMesh::GetReferenceBoneTransform() const
 	{
 		if (const FClothingSimulationContextCommon* const Context = GetContext())
 		{
@@ -146,7 +148,12 @@ PRAGMA_ENABLE_DEPRECATION_WARNINGS
 				BoneTransforms[ReferenceBoneIndex] * Context->ComponentToWorld :
 				Context->ComponentToWorld;
 		}
-		return FRigidTransform3::Identity;
+		return FTransform::Identity;
+	}
+
+	const TArray<FTransform>& FClothingSimulationSkeletalMesh::GetBoneTransforms() const
+	{
+		return GetContext()->BoneTransforms;
 	}
 
 	const FTransform& FClothingSimulationSkeletalMesh::GetComponentToWorldTransform() const
diff --git a/Engine/Plugins/Experimental/ChaosCloth/Source/Chaos/Public/ChaosCloth/ChaosClothVisualization.h b/Engine/Plugins/Experimental/ChaosCloth/Source/Chaos/Public/ChaosCloth/ChaosClothVisualization.h
--- a/Engine/Plugins/Experimental/ChaosCloth/Source/Chaos/Public/ChaosCloth/ChaosClothVisualization.h
+++ b/Engine/Plugins/Experimental/ChaosCloth/Source/Chaos/Public/ChaosCloth/ChaosClothVisualization.h
@@ -16,18 +16,15 @@ class UMaterial;
 namespace Chaos
 {
 	class FClothingSimulationSolver;
-}
 
-namespace UE::Chaos::Cloth
-{
-	class CHAOSCLOTH_API FVisualization
+	class CHAOSCLOTH_API FClothVisualization
 #if WITH_EDITOR
 		: public FGCObject  // Add garbage collection for cloth material
 #endif  // #if WITH_EDITOR
 	{
 	public:
-		explicit FVisualization(const ::Chaos::FClothingSimulationSolver* InSolver = nullptr);
-		virtual ~FVisualization();
+		explicit FClothVisualization(const ::Chaos::FClothingSimulationSolver* InSolver = nullptr);
+		virtual ~FClothVisualization();
 
 #if CHAOS_DEBUG_DRAW
 		// Editor & runtime functions
diff --git a/Engine/Plugins/Experimental/ChaosCloth/Source/Chaos/Public/ChaosCloth/ChaosClothingSimulation.h b/Engine/Plugins/Experimental/ChaosCloth/Source/Chaos/Public/ChaosCloth/ChaosClothingSimulation.h
--- a/Engine/Plugins/Experimental/ChaosCloth/Source/Chaos/Public/ChaosCloth/ChaosClothingSimulation.h
+++ b/Engine/Plugins/Experimental/ChaosCloth/Source/Chaos/Public/ChaosCloth/ChaosClothingSimulation.h
@@ -122,7 +122,7 @@ namespace Chaos
 
 	private:
 		// Visualization object
-		UE::Chaos::Cloth::FVisualization Visualization;
+		FClothVisualization Visualization;
 
 		// Simulation objects
 		TUniquePtr<FClothingSimulationSolver> Solver;  // Default solver
diff --git a/Engine/Plugins/Experimental/ChaosCloth/Source/Chaos/Public/ChaosCloth/ChaosClothingSimulationCollider.h b/Engine/Plugins/Experimental/ChaosCloth/Source/Chaos/Public/ChaosCloth/ChaosClothingSimulationCollider.h
--- a/Engine/Plugins/Experimental/ChaosCloth/Source/Chaos/Public/ChaosCloth/ChaosClothingSimulationCollider.h
+++ b/Engine/Plugins/Experimental/ChaosCloth/Source/Chaos/Public/ChaosCloth/ChaosClothingSimulationCollider.h
@@ -7,7 +7,9 @@
 
 class USkeletalMeshComponent;
 class UClothingAssetCommon;
+class UPhysicsAsset;
 class FClothingSimulationContextCommon;
+struct FReferenceSkeleton;
 
 namespace Chaos
 {
@@ -28,11 +30,15 @@ namespace Chaos
 			LODs,  // LODIndex based start slot for LODs collisions
 		};
 
+		FClothingSimulationCollider(const UPhysicsAsset* InPhysicsAsset, const FReferenceSkeleton* InReferenceSkeleton);
+
+		UE_DEPRECATED(5.2, "Use FClothingSimulationCollider(const UPhysicsAsset* InPhysicsAsset, const FReferenceSkeleton* InReferenceSkeleton) instead.")
 		FClothingSimulationCollider(
 			const UClothingAssetCommon* InAsset,  // Cloth asset for collision data, can be nullptr
 			const USkeletalMeshComponent* InSkeletalMeshComponent,  // For asset LODs management, can be nullptr
 			bool bInUseLODIndexOverride,
 			int32 InLODIndexOverride);
+
 		~FClothingSimulationCollider();
 
 		FClothingSimulationCollider(const FClothingSimulationCollider&) = delete;
@@ -99,18 +105,17 @@ namespace Chaos
 
 		struct FLODData;
 
-		const UClothingAssetCommon* Asset;
-		const USkeletalMeshComponent* SkeletalMeshComponent;
-		const FClothCollisionData* CollisionData;
-		bool bUseLODIndexOverride;
-		int32 LODIndexOverride;
-		bool bHasExternalCollisionChanged;
+		const UPhysicsAsset* PhysicsAsset = nullptr;
+		const FReferenceSkeleton* ReferenceSkeleton = nullptr;
+		const FClothCollisionData* CollisionData = nullptr;  // External collision data
+
+		bool bHasExternalCollisionChanged = false;
 
 		// Collision primitives
 		TArray<TUniquePtr<FLODData>> LODData;  // Actual LODs start at ECollisionDataType::LODs
 		TMap<FSolverClothPair, int32> LODIndices;
 
 		// Initial scale
-		FReal Scale;
+		FReal Scale = 1.f;
 	};
 } // namespace Chaos
\ No newline at end of file
diff --git a/Engine/Plugins/Experimental/ChaosCloth/Source/Chaos/Public/ChaosCloth/ChaosClothingSimulationMesh.h b/Engine/Plugins/Experimental/ChaosCloth/Source/Chaos/Public/ChaosCloth/ChaosClothingSimulationMesh.h
--- a/Engine/Plugins/Experimental/ChaosCloth/Source/Chaos/Public/ChaosCloth/ChaosClothingSimulationMesh.h
+++ b/Engine/Plugins/Experimental/ChaosCloth/Source/Chaos/Public/ChaosCloth/ChaosClothingSimulationMesh.h
@@ -30,7 +30,8 @@ namespace Chaos
 #endif
 	{
 	public:
-		FClothingSimulationMesh();
+		explicit FClothingSimulationMesh(const FString& InDebugName);
+
 		virtual ~FClothingSimulationMesh();
 
 		FClothingSimulationMesh(const FClothingSimulationMesh&) = delete;
@@ -38,6 +39,12 @@ namespace Chaos
 		FClothingSimulationMesh& operator=(const FClothingSimulationMesh&) = delete;
 		FClothingSimulationMesh& operator=(FClothingSimulationMesh&&) = delete;
 
+#if !UE_BUILD_SHIPPING
+		const FString& GetDebugName() const { return DebugName; }
+#else
+		const FString& GetDebugName() const { return FText::GetEmpty().ToString(); }
+#endif
+
 #if !defined(CHAOS_IS_CLOTHINGSIMULATIONMESH_ABSTRACT) || !CHAOS_IS_CLOTHINGSIMULATIONMESH_ABSTRACT
 		UE_DEPRECATED(5.2, "Use the FClothingSimulationSkeletalMesh constructor instead.")
 		FClothingSimulationMesh(const UClothingAssetCommon* InAsset, const USkeletalMeshComponent* InSkeletalMeshComponent);
@@ -57,7 +64,8 @@ namespace Chaos
 		virtual TArray<TConstArrayView<FRealSingle>> GetWeightMaps(int32 LODIndex) const;
 		virtual TArray<TConstArrayView<TTuple<int32, int32, float>>> GetTethers(int32 LODIndex, bool bUseGeodesicTethers) const;
 		virtual int32 GetReferenceBoneIndex() const;
-		virtual FRigidTransform3 GetReferenceBoneTransform() const;
+		virtual FTransform GetReferenceBoneTransform() const;
+		virtual const TArray<FTransform>& GetBoneTransforms() const;
 		virtual const FTransform& GetComponentToWorldTransform() const;
 		virtual const TArray<FMatrix44f>& GetRefToLocalMatrices() const;
 		virtual TConstArrayView<int32> GetBoneMap() const;
@@ -65,7 +73,6 @@ namespace Chaos
 		virtual TConstArrayView<FMeshToMeshVertData> GetTransitionUpSkinData(int32 LODIndex) const;
 		virtual TConstArrayView<FMeshToMeshVertData> GetTransitionDownSkinData(int32 LODIndex) const;
 #else
-		// ---- Cloth interface ----
 		/* Return the number of LODs on this mesh. */
 		virtual int32 GetNumLODs() const = 0;
 
@@ -100,7 +107,10 @@ namespace Chaos
 		virtual int32 GetReferenceBoneIndex() const = 0;
 
 		/* Return the transform of the bone treated as the root of the simulation space. */
-		virtual FRigidTransform3 GetReferenceBoneTransform() const = 0;
+		virtual FTransform GetReferenceBoneTransform() const = 0;
+
+		/* Return the bone transforms as required when updating the collider pose. */
+		const TArray<FTransform>& GetBoneTransforms() const = 0;
 
 		/* Return the transform of the bone treated as the root of the simulation space. */
 		virtual const FTransform& GetComponentToWorldTransform() const = 0;
@@ -164,6 +174,11 @@ namespace Chaos
 		const UClothingAssetCommon* Asset;
 		const USkeletalMeshComponent* SkeletalMeshComponent;
 #endif
+
+#if !UE_BUILD_SHIPPING
+		/** Debug name of the source component. */
+		FString DebugName;
+#endif
 	};
 } // namespace Chaos
 
diff --git a/Engine/Plugins/Experimental/ChaosCloth/Source/Chaos/Public/ChaosCloth/ChaosClothingSimulationSkeletalMesh.h b/Engine/Plugins/Experimental/ChaosCloth/Source/Chaos/Public/ChaosCloth/ChaosClothingSimulationSkeletalMesh.h
--- a/Engine/Plugins/Experimental/ChaosCloth/Source/Chaos/Public/ChaosCloth/ChaosClothingSimulationSkeletalMesh.h
+++ b/Engine/Plugins/Experimental/ChaosCloth/Source/Chaos/Public/ChaosCloth/ChaosClothingSimulationSkeletalMesh.h
@@ -44,7 +44,8 @@ PRAGMA_ENABLE_DEPRECATION_WARNINGS
 		virtual TArray<TConstArrayView<FRealSingle>> GetWeightMaps(int32 LODIndex) const override;
 		virtual TArray<TConstArrayView<TTuple<int32, int32, float>>> GetTethers(int32 LODIndex, bool bUseGeodesicTethers) const override;
 		virtual int32 GetReferenceBoneIndex() const override;
-		virtual FRigidTransform3 GetReferenceBoneTransform() const override;
+		virtual FTransform GetReferenceBoneTransform() const override;
+		virtual const TArray<FTransform>& GetBoneTransforms() const override;
 		virtual const FTransform& GetComponentToWorldTransform() const override;
 		virtual const TArray<FMatrix44f>& GetRefToLocalMatrices() const override;
 		virtual TConstArrayView<int32> GetBoneMap() const override;
diff --git a/Engine/Plugins/Experimental/ChaosClothAsset/Source/ChaosClothAssetEngine/Private/ChaosClothAsset/ClothAsset.cpp b/Engine/Plugins/Experimental/ChaosClothAsset/Source/ChaosClothAssetEngine/Private/ChaosClothAsset/ClothAsset.cpp
--- a/Engine/Plugins/Experimental/ChaosClothAsset/Source/ChaosClothAssetEngine/Private/ChaosClothAsset/ClothAsset.cpp
+++ b/Engine/Plugins/Experimental/ChaosClothAsset/Source/ChaosClothAssetEngine/Private/ChaosClothAsset/ClothAsset.cpp
@@ -3,6 +3,7 @@
 #include "ChaosClothAsset/ClothAsset.h"
 #include "ChaosClothAsset/ClothAssetBuilder.h"
 #include "ChaosClothAsset/ClothAdapter.h"
+#include "ChaosClothAsset/ClothComponent.h"
 #include "ChaosClothAsset/ClothGeometryTools.h"
 #include "ChaosClothAsset/ClothAssetPrivate.h"
 #include "ChaosClothAsset/ClothSimulationModel.h"
@@ -68,7 +69,7 @@ void UChaosClothAsset::Serialize(FArchive& Ar)
 	Ar << GetRefSkeleton();
 	if (Ar.IsLoading())
 	{
-		const bool bRebuildNameMap = false;
+		constexpr bool bRebuildNameMap = false;
 		GetRefSkeleton().RebuildRefSkeleton(GetSkeleton(), bRebuildNameMap);
 	}
 
@@ -87,6 +88,27 @@ void UChaosClothAsset::PostLoad()
 	Super::PostLoad();
 }
 
+#if WITH_EDITOR
+void UChaosClothAsset::PostEditChangeProperty(FPropertyChangedEvent& PropertyChangedEvent)
+{
+	if (PropertyChangedEvent.GetPropertyName() == GET_MEMBER_NAME_CHECKED(UChaosClothAsset, PhysicsAsset))
+	{
+		// Recreate the simulation proxies with the updated physics asset
+		for (TObjectIterator<UChaosClothComponent> ObjectIterator; ObjectIterator; ++ObjectIterator)
+		{
+			if (UChaosClothComponent* const Component = *ObjectIterator)
+			{
+				if (Component->GetClothAsset() == this)
+				{
+					const FComponentReregisterContext Context(Component);  // Context goes out of scope, causing the Component to be re-registered
+				}
+			}
+		}
+	}
+	Super::PostEditChangeProperty(PropertyChangedEvent);
+}
+#endif // #if WITH_EDITOR
+
 void UChaosClothAsset::BeginPostLoadInternal(FSkinnedAssetPostLoadContext& Context)
 {
 #if WITH_EDITOR
diff --git a/Engine/Plugins/Experimental/ChaosClothAsset/Source/ChaosClothAssetEngine/Private/ChaosClothAsset/ClothComponent.cpp b/Engine/Plugins/Experimental/ChaosClothAsset/Source/ChaosClothAssetEngine/Private/ChaosClothAsset/ClothComponent.cpp
--- a/Engine/Plugins/Experimental/ChaosClothAsset/Source/ChaosClothAssetEngine/Private/ChaosClothAsset/ClothComponent.cpp
+++ b/Engine/Plugins/Experimental/ChaosClothAsset/Source/ChaosClothAssetEngine/Private/ChaosClothAsset/ClothComponent.cpp
@@ -8,6 +8,7 @@
 #include "HAL/IConsoleManager.h"
 #include "Rendering/SkeletalMeshRenderData.h"
 #include "SkeletalRenderPublic.h"
+#include "PhysicsEngine/PhysicsAsset.h"
 
 CSV_DECLARE_CATEGORY_MODULE_EXTERN(ENGINE_API, Animation);
 
@@ -207,6 +208,9 @@ void UChaosClothComponent::PostLoad()
 
 void UChaosClothComponent::SetClothAsset(UChaosClothAsset* InClothAsset)
 {
+	// De-register component, and re-register it when it reaches end of scope
+	const FComponentReregisterContext ComponentReregisterContext(this);
+
 	PRAGMA_DISABLE_DEPRECATION_WARNINGS
 #if WITH_EDITORONLY_DATA
 	ClothAsset = InClothAsset;
@@ -226,12 +230,12 @@ void UChaosClothComponent::PostEditChangeProperty(FPropertyChangedEvent& Propert
 	// Set the skinned asset pointer with the alias pointer (must happen before the call to Super::PostEditChangeProperty)
 	if (const FProperty* const Property = PropertyChangedEvent.Property)
 	{
-		PRAGMA_DISABLE_DEPRECATION_WARNINGS
+PRAGMA_DISABLE_DEPRECATION_WARNINGS
 		if (Property->GetFName() == GET_MEMBER_NAME_CHECKED(UChaosClothComponent, ClothAsset))
 		{
-			SetSkinnedAsset(ClothAsset);
+			SetClothAsset(ClothAsset);
 		}
-		PRAGMA_ENABLE_DEPRECATION_WARNINGS
+PRAGMA_ENABLE_DEPRECATION_WARNINGS
 	}
 
 	Super::PostEditChangeProperty(PropertyChangedEvent);
diff --git a/Engine/Plugins/Experimental/ChaosClothAsset/Source/ChaosClothAssetEngine/Private/ChaosClothAsset/ClothSimulationMesh.cpp b/Engine/Plugins/Experimental/ChaosClothAsset/Source/ChaosClothAssetEngine/Private/ChaosClothAsset/ClothSimulationMesh.cpp
--- a/Engine/Plugins/Experimental/ChaosClothAsset/Source/ChaosClothAssetEngine/Private/ChaosClothAsset/ClothSimulationMesh.cpp
+++ b/Engine/Plugins/Experimental/ChaosClothAsset/Source/ChaosClothAssetEngine/Private/ChaosClothAsset/ClothSimulationMesh.cpp
@@ -8,9 +8,9 @@
 
 namespace UE::Chaos::ClothAsset
 {
-	FClothSimulationMesh::FClothSimulationMesh(const FChaosClothSimulationModel& InClothSimulationModel, const FClothSimulationContext& InClothSimulationContext)
+	FClothSimulationMesh::FClothSimulationMesh(const FChaosClothSimulationModel& InClothSimulationModel, const FClothSimulationContext& InClothSimulationContext, const FString& DebugName)
 PRAGMA_DISABLE_DEPRECATION_WARNINGS  // TODO: CHAOS_IS_CLOTHINGSIMULATIONMESH_ABSTRACT
-		: ::Chaos::FClothingSimulationMesh()
+		: ::Chaos::FClothingSimulationMesh(DebugName)
 PRAGMA_ENABLE_DEPRECATION_WARNINGS
 		, ClothSimulationModel(InClothSimulationModel)
 		, ClothSimulationContext(InClothSimulationContext)
@@ -74,7 +74,7 @@ PRAGMA_ENABLE_DEPRECATION_WARNINGS
 		return ClothSimulationModel.ReferenceBoneIndex;
 	}
 
-	::Chaos::FRigidTransform3 FClothSimulationMesh::GetReferenceBoneTransform() const
+	FTransform FClothSimulationMesh::GetReferenceBoneTransform() const
 	{
 		// TODO: Leader pose component (see FClothingSimulationContextCommon::FillBoneTransforms)
 		const TArray<FTransform>& BoneTransforms = ClothSimulationContext.BoneTransforms;
@@ -86,6 +86,12 @@ PRAGMA_ENABLE_DEPRECATION_WARNINGS
 		return GetComponentToWorldTransform();
 	}
 
+	const TArray<FTransform>& FClothSimulationMesh::GetBoneTransforms() const
+	{
+		// TODO: Leader pose component (see FClothingSimulationContextCommon::FillBoneTransforms)
+		return ClothSimulationContext.BoneTransforms;
+	}
+
 	const FTransform& FClothSimulationMesh::GetComponentToWorldTransform() const
 	{
 		return ClothSimulationContext.ComponentTransform;
diff --git a/Engine/Plugins/Experimental/ChaosClothAsset/Source/ChaosClothAssetEngine/Private/ChaosClothAsset/ClothSimulationMesh.h b/Engine/Plugins/Experimental/ChaosClothAsset/Source/ChaosClothAssetEngine/Private/ChaosClothAsset/ClothSimulationMesh.h
--- a/Engine/Plugins/Experimental/ChaosClothAsset/Source/ChaosClothAssetEngine/Private/ChaosClothAsset/ClothSimulationMesh.h
+++ b/Engine/Plugins/Experimental/ChaosClothAsset/Source/ChaosClothAssetEngine/Private/ChaosClothAsset/ClothSimulationMesh.h
@@ -15,7 +15,7 @@ PRAGMA_DISABLE_DEPRECATION_WARNINGS  // TODO: CHAOS_IS_CLOTHINGSIMULATIONMESH_AB
 PRAGMA_ENABLE_DEPRECATION_WARNINGS
 	{
 	public:
-		FClothSimulationMesh(const FChaosClothSimulationModel& InClothSimulationModel, const FClothSimulationContext& InClothSimulationContext);
+		FClothSimulationMesh(const FChaosClothSimulationModel& InClothSimulationModel, const FClothSimulationContext& InClothSimulationContext, const FString& DebugName);
 PRAGMA_DISABLE_DEPRECATION_WARNINGS  // TODO: CHAOS_IS_CLOTHINGSIMULATIONMESH_ABSTRACT
 		virtual ~FClothSimulationMesh() override = default;
 PRAGMA_ENABLE_DEPRECATION_WARNINGS
@@ -37,7 +37,8 @@ PRAGMA_ENABLE_DEPRECATION_WARNINGS
 		virtual TArray<TConstArrayView<::Chaos::FRealSingle>> GetWeightMaps(int32 LODIndex) const override;
 		virtual TArray<TConstArrayView<TTuple<int32, int32, float>>> GetTethers(int32 LODIndex, bool bUseGeodesicTethers) const override;
 		virtual int32 GetReferenceBoneIndex() const override;
-		virtual ::Chaos::FRigidTransform3 GetReferenceBoneTransform() const override;
+		virtual FTransform GetReferenceBoneTransform() const override;
+		virtual const TArray<FTransform>& GetBoneTransforms() const override;
 		virtual const FTransform& GetComponentToWorldTransform() const override;
 		virtual const TArray<FMatrix44f>& GetRefToLocalMatrices() const override;
 		virtual TConstArrayView<int32> GetBoneMap() const override;
diff --git a/Engine/Plugins/Experimental/ChaosClothAsset/Source/ChaosClothAssetEngine/Private/ChaosClothAsset/ClothSimulationProxy.cpp b/Engine/Plugins/Experimental/ChaosClothAsset/Source/ChaosClothAssetEngine/Private/ChaosClothAsset/ClothSimulationProxy.cpp
--- a/Engine/Plugins/Experimental/ChaosClothAsset/Source/ChaosClothAssetEngine/Private/ChaosClothAsset/ClothSimulationProxy.cpp
+++ b/Engine/Plugins/Experimental/ChaosClothAsset/Source/ChaosClothAssetEngine/Private/ChaosClothAsset/ClothSimulationProxy.cpp
@@ -9,6 +9,7 @@
 #include "ChaosCloth/ChaosClothingSimulationCloth.h"
 #include "ChaosCloth/ChaosClothingSimulationCollider.h"
 #include "ChaosCloth/ChaosClothingSimulationSolver.h"
+#include "ChaosCloth/ChaosClothVisualization.h"
 #include "PhysicsEngine/PhysicsSettings.h"
 
 #if INTEL_ISPC
@@ -84,6 +85,7 @@ namespace UE::Chaos::ClothAsset
 	FClothSimulationProxy::FClothSimulationProxy(const UChaosClothComponent& InClothComponent)
 		: ClothComponent(InClothComponent)
 		, Solver(MakeUnique<::Chaos::FClothingSimulationSolver>())
+		, Visualization(MakeUnique<::Chaos::FClothVisualization>(Solver.Get()))
 		, MaxDeltaTime(UPhysicsSettings::Get()->MaxPhysicsDeltaTime)
 	{
 		using namespace ::Chaos;
@@ -96,16 +98,26 @@ namespace UE::Chaos::ClothAsset
 
 		// Create mesh node
 		const UChaosClothAsset* const ClothAsset = ClothComponent.GetClothAsset();
-		const FChaosClothSimulationModel* const ClothSimulationModel = ClothAsset ? ClothAsset->GetClothSimulationModel() : nullptr;
+		if (!ClothAsset)
+		{
+			return;
+		}
+		const FChaosClothSimulationModel* const ClothSimulationModel = ClothAsset->GetClothSimulationModel();
 		check(ClothSimulationModel);
-		const int32 MeshIndex = Meshes.Emplace(MakeUnique<FClothSimulationMesh>(*ClothSimulationModel, ClothSimulationContext));
+
+		FString DebugName;
+#if !UE_BUILD_SHIPPING
+		DebugName = ClothComponent.GetOwner() ?
+			FString::Format(TEXT("{0}|{1}"), { ClothComponent.GetOwner()->GetName(), ClothComponent.GetName() }) :
+			ClothComponent.GetName();
+#endif
+		const int32 MeshIndex = Meshes.Emplace(MakeUnique<FClothSimulationMesh>(*ClothSimulationModel, ClothSimulationContext, DebugName));
 
 		// Create collider node
-		const int32 ColliderIndex = Colliders.Emplace(MakeUnique<FClothingSimulationCollider>(
-			nullptr,
-			nullptr,  // TODO: Cloth Asset Mesh 
-			/*bInUseLODIndexOverride =*/ false,
-			/*InLODIndexOverride =*/ INDEX_NONE));
+PRAGMA_DISABLE_DEPRECATION_WARNINGS  // TODO: CHAOS_IS_CLOTHINGSIMULATIONMESH_ABSTRACT
+		const FReferenceSkeleton* const ReferenceSkeleton = &ClothAsset->GetRefSkeleton();
+		const int32 ColliderIndex = ClothComponent.GetPhysicsAsset() ? Colliders.Emplace(MakeUnique<FClothingSimulationCollider>(ClothComponent.GetPhysicsAsset(), ReferenceSkeleton)) : INDEX_NONE;
+PRAGMA_ENABLE_DEPRECATION_WARNINGS
 
 		//// Set the external collision data to get updated at every frame
 		//Colliders[ColliderIndex]->SetCollisionData(&ExternalCollisionData);
@@ -115,7 +127,7 @@ namespace UE::Chaos::ClothAsset
 PRAGMA_DISABLE_DEPRECATION_WARNINGS  // TODO: CHAOS_IS_CLOTHINGSIMULATIONMESH_ABSTRACT
 			Meshes[MeshIndex].Get(),
 PRAGMA_ENABLE_DEPRECATION_WARNINGS
-			TArray<FClothingSimulationCollider*>({ Colliders[ColliderIndex].Get() }),
+			ColliderIndex != INDEX_NONE ? TArray<FClothingSimulationCollider*>({ Colliders[ColliderIndex].Get() }) : TArray<FClothingSimulationCollider*>(),
 			0,
 			FClothingSimulationCloth::EMassMode::Density,  //	(FClothingSimulationCloth::EMassMode)ClothConfig->MassMode,
 			0.35f,  //	ClothConfig->GetMassValue(),
@@ -239,6 +251,10 @@ PRAGMA_ENABLE_DEPRECATION_WARNINGS
 		const float CurrSimulationTime = (float)((FPlatformTime::Seconds() - StartTime) * 1000.);
 		static const float SimulationTimeDecay = 0.03f; // 0.03 seems to provide a good rate of update for the instant average
 		SimulationTime = PrevSimulationTime ? PrevSimulationTime + (CurrSimulationTime - PrevSimulationTime) * SimulationTimeDecay : CurrSimulationTime;
+
+		// Visualization
+		Visualization->DrawPhysMeshWired();
+		Visualization->DrawCollision();
 	}
 
 	void FClothSimulationProxy::CompleteParallelSimulation_GameThread()
diff --git a/Engine/Plugins/Experimental/ChaosClothAsset/Source/ChaosClothAssetEngine/Private/ChaosClothAsset/ClothSimulationProxy.h b/Engine/Plugins/Experimental/ChaosClothAsset/Source/ChaosClothAssetEngine/Private/ChaosClothAsset/ClothSimulationProxy.h
--- a/Engine/Plugins/Experimental/ChaosClothAsset/Source/ChaosClothAssetEngine/Private/ChaosClothAsset/ClothSimulationProxy.h
+++ b/Engine/Plugins/Experimental/ChaosClothAsset/Source/ChaosClothAssetEngine/Private/ChaosClothAsset/ClothSimulationProxy.h
@@ -15,6 +15,7 @@ namespace Chaos
 	class FClothingSimulationMesh;
 	class FClothingSimulationCloth;
 	class FClothingSimulationCollider;
+	class FClothVisualization;
 }
 
 class UChaosClothComponent;
@@ -79,6 +80,7 @@ PRAGMA_DISABLE_DEPRECATION_WARNINGS  // TODO: CHAOS_IS_CLOTHINGSIMULATIONMESH_AB
 PRAGMA_ENABLE_DEPRECATION_WARNINGS
 		TArray<TUniquePtr<::Chaos::FClothingSimulationCloth>> Cloths;
 		TArray<TUniquePtr<::Chaos::FClothingSimulationCollider>> Colliders;
+		TUniquePtr<::Chaos::FClothVisualization> Visualization;
 
 		// Properties that must be readable from all threads
 		std::atomic<int32> NumCloths;
diff --git a/Engine/Plugins/Experimental/ChaosClothAsset/Source/ChaosClothAssetEngine/Public/ChaosClothAsset/ClothAsset.h b/Engine/Plugins/Experimental/ChaosClothAsset/Source/ChaosClothAssetEngine/Public/ChaosClothAsset/ClothAsset.h
--- a/Engine/Plugins/Experimental/ChaosClothAsset/Source/ChaosClothAssetEngine/Public/ChaosClothAsset/ClothAsset.h
+++ b/Engine/Plugins/Experimental/ChaosClothAsset/Source/ChaosClothAssetEngine/Public/ChaosClothAsset/ClothAsset.h
@@ -13,6 +13,7 @@ namespace UE::Chaos::ClothAsset { class FClothCollection; }
 class FSkeletalMeshRenderData;
 class FSkeletalMeshModel;
 struct FChaosClothSimulationModel;
+struct FSkeletalMeshLODInfo;
 
 UENUM()
 enum class EClothAssetAsyncProperties : uint64
@@ -40,6 +41,9 @@ public:
 	virtual void BeginDestroy() override;
 	virtual bool IsReadyForFinishDestroy() override;
 	virtual void PostLoad() override;
+#if WITH_EDITOR
+	virtual void PostEditChangeProperty(FPropertyChangedEvent& PropertyChangedEvent) override;
+#endif
 	//~ End UObject interface
 
 	//~ Begin USkinnedAsset interface
@@ -60,7 +64,7 @@ public:
 	virtual const TArray<FSkeletalMeshLODInfo>& GetLODInfoArray() const override { return LODInfo; }
 	virtual FSkeletalMeshRenderData* GetResourceForRendering() const override;
 	virtual int32 GetDefaultMinLod() const										{ return 0; }
-	virtual UPhysicsAsset* GetPhysicsAsset() const								{ return nullptr; }
+	virtual UPhysicsAsset* GetPhysicsAsset() const								{ return PhysicsAsset; }
 	virtual TArray<FSkeletalMaterial>& GetMaterials() override					{ return Materials; }
 	virtual const TArray<FSkeletalMaterial>& GetMaterials() const override		{ return Materials; }
 	virtual int32 GetLODNum() const override									{ return LODInfo.Num(); }
@@ -172,7 +176,9 @@ private:
 	UPROPERTY(EditAnywhere, Category = Materials)
 	TArray<FSkeletalMaterial> Materials;
 
-	FBoxSphereBounds Bounds;
+	/** Physics asset used for collision. */
+	UPROPERTY(EditAnywhere, Category = Collision)
+	TObjectPtr<UPhysicsAsset> PhysicsAsset;
 
 	/** Skeleton. */
 	UPROPERTY(EditAnywhere, AssetRegistrySearchable, Category = Skeleton)
@@ -219,7 +225,10 @@ private:
 	UPROPERTY()
 	FGuid AssetGuid;
 
-	/** mesh-space ref pose, where parent matrices are applied to ref pose matrices */
+	/** Bounds for this asset. */
+	FBoxSphereBounds Bounds;
+
+	/** Mesh-space ref pose, where parent matrices are applied to ref pose matrices. */
 	TArray<FMatrix> CachedComposedRefPoseMatrices;
 
 	/** Cloth Collection containing this asset data. */
diff --git a/Engine/Plugins/Experimental/ChaosClothAsset/Source/ChaosClothAssetEngine/Public/ChaosClothAsset/ClothComponent.h b/Engine/Plugins/Experimental/ChaosClothAsset/Source/ChaosClothAssetEngine/Public/ChaosClothAsset/ClothComponent.h
--- a/Engine/Plugins/Experimental/ChaosClothAsset/Source/ChaosClothAssetEngine/Public/ChaosClothAsset/ClothComponent.h
+++ b/Engine/Plugins/Experimental/ChaosClothAsset/Source/ChaosClothAssetEngine/Public/ChaosClothAsset/ClothComponent.h
@@ -17,7 +17,7 @@ namespace UE::Chaos::ClothAsset
 /**
  * Cloth simulation component.
  */
-UCLASS(ClassGroup = Physics, Meta = (BlueprintSpawnableComponent, ToolTip = "Chaos cloth simulation component."), DisplayName = "Cloth Simulation", HideCategories = (Object, "Mesh|SkeletalAsset", Physics, Constraints, Advanced, Cooking, Collision))
+UCLASS(ClassGroup = Physics, Meta = (BlueprintSpawnableComponent, ToolTip = "Chaos cloth simulation component."), DisplayName = "Cloth Simulation", HideCategories = (Object, "Mesh|SkeletalAsset", Constraints, Advanced, Cooking, Collision))
 class CHAOSCLOTHASSETENGINE_API UChaosClothComponent : public USkinnedMeshComponent
 {
 	GENERATED_BODY()
@@ -89,7 +89,7 @@ private:
 #if WITH_EDITORONLY_DATA
 	/** Cloth asset used by this component. */
 	UE_DEPRECATED(5.1, "This property isn't deprecated, but getter and setter must be used at all times to preserve correct operations.")
-	UPROPERTY(EditAnywhere, Transient, BlueprintSetter = SetClothAsset, BlueprintGetter = GetClothAsset, Category = Mesh)
+	UPROPERTY(EditAnywhere, Transient, BlueprintSetter = SetClothAsset, BlueprintGetter = GetClothAsset, Category = ClothAsset)
 	TObjectPtr<UChaosClothAsset> ClothAsset;
 #endif
 
