Chaos Cloth - Fix fictitious forces' center of rotation that was offset by the solver simulation space location.

#rnx

diff --git a/Engine/Plugins/Experimental/ChaosCloth/Source/Chaos/Private/ChaosCloth/ChaosClothingSimulationSolver.cpp b/Engine/Plugins/Experimental/ChaosCloth/Source/Chaos/Private/ChaosCloth/ChaosClothingSimulationSolver.cpp
--- a/Engine/Plugins/Experimental/ChaosCloth/Source/Chaos/Private/ChaosCloth/ChaosClothingSimulationSolver.cpp
+++ b/Engine/Plugins/Experimental/ChaosCloth/Source/Chaos/Private/ChaosCloth/ChaosClothingSimulationSolver.cpp
@@ -609,7 +609,7 @@ void FClothingSimulationSolver::SetReferenceVelocityScale(
 	// Save the reference bone relative angular velocity for calculating the fictitious forces
 	const FVec3 FictitiousAngularDisplacement = ReferenceSpaceTransform.TransformVector(Axis * PartialDeltaAngle * FMath::Min((FReal)2., (FReal)FictitiousAngularScale));  // Clamp to 2x the delta angle
 	FictitiousAngularDisplacements[GroupId] = Softs::FSolverVec3(FictitiousAngularDisplacement);
-	ReferenceSpaceLocations[GroupId] = ReferenceSpaceTransform.GetLocation();
+	ReferenceSpaceLocations[GroupId] = ReferenceSpaceTransform.GetLocation() - LocalSpaceLocation;
 }
 
 Softs::FSolverReal FClothingSimulationSolver::SetParticleMassPerArea(int32 Offset, int32 Size, const FTriangleMesh& Mesh)
