Chaos Cloth - Fixed solver velocity scale when a reset is needed, which was causing the garment to flip over a character on the first few frames.

diff --git a/Engine/Plugins/Experimental/ChaosCloth/Source/Chaos/Private/ChaosCloth/ChaosClothingSimulation.cpp b/Engine/Plugins/Experimental/ChaosCloth/Source/Chaos/Private/ChaosCloth/ChaosClothingSimulation.cpp
--- a/Engine/Plugins/Experimental/ChaosCloth/Source/Chaos/Private/ChaosCloth/ChaosClothingSimulation.cpp
+++ b/Engine/Plugins/Experimental/ChaosCloth/Source/Chaos/Private/ChaosCloth/ChaosClothingSimulation.cpp
@@ -543,7 +543,7 @@ void FClothingSimulation::Simulate(IClothingSimulationContext* InContext)
 		Solver->SetWindVelocity(Context->WindVelocity, Context->WindAdaption);
 		Solver->SetGravity(bUseGravityOverride ? GravityOverride : Context->WorldGravity);
 		Solver->EnableClothGravityOverride(!bUseGravityOverride);  // Disable all cloth gravity overrides when the interactor takes over
-		Solver->SetVelocityScale((FReal)Context->VelocityScale * (FReal)SmoothedDeltaTime / DeltaTime);
+		Solver->SetVelocityScale(!bNeedsReset ? (FReal)Context->VelocityScale * (FReal)SmoothedDeltaTime / DeltaTime : 1.f);
 
 		// Check teleport modes
 		for (const TUniquePtr<FClothingSimulationCloth>& Cloth : Cloths)
