Chaos Cloth: correctly handling substepping for animdrive and backstop constraints.

The code assumed there was a single subframe, so velocity was scaled incorrectly and there was no subframe interpolation on the animated positions and normals. Furthermore, the animated velocity for animdrive was actually scaled  by -1. This is fixed.

Update anim drive damping to use latest positions (including this constraint's position-based contribution) in velocity calculations. This prevents instability when using high stiffness and high damping from effectively double-applying the constraint.

Interpolated animated positions and normals as well as velocities are now precalculated and stored by the solver.

Since AnimDrive requires the addition of the velocity input, the old constructor and old behavior has been retained but deprecated.

Tested with various values (including 0 and 1) for stiffness and damping. Note that animdrive constraints are not created if stiffness is zero. Locally modified create code to create if stiffness or damping are nonzero to test damping alone, but did not include this change to check in to not negatively affect existing assets.

diff --git a/Engine/Plugins/Experimental/ChaosCloth/Source/Chaos/Private/ChaosCloth/ChaosClothConstraints.cpp b/Engine/Plugins/Experimental/ChaosCloth/Source/Chaos/Private/ChaosCloth/ChaosClothConstraints.cpp
--- a/Engine/Plugins/Experimental/ChaosCloth/Source/Chaos/Private/ChaosCloth/ChaosClothConstraints.cpp
+++ b/Engine/Plugins/Experimental/ChaosCloth/Source/Chaos/Private/ChaosCloth/ChaosClothConstraints.cpp
@@ -21,8 +21,9 @@ namespace Chaos {
 FClothConstraints::FClothConstraints()
 	: Evolution(nullptr)
 	, AnimationPositions(nullptr)
-	, OldAnimationPositions(nullptr)
+	, OldAnimationPositions_deprecated(nullptr)
 	, AnimationNormals(nullptr)
+	, AnimationVelocities(nullptr)
 	, ParticleOffset(0)
 	, NumParticles(0)
 	, ConstraintInitOffset(INDEX_NONE)
@@ -50,8 +51,27 @@ void FClothConstraints::Initialize(
 {
 	Evolution = InEvolution;
 	AnimationPositions = &InAnimationPositions;
-	OldAnimationPositions = &InOldAnimationPositions;
+	OldAnimationPositions_deprecated = &InOldAnimationPositions;
 	AnimationNormals = &InAnimationNormals;
+	AnimationVelocities = nullptr;
+	ParticleOffset = InParticleOffset;
+	NumParticles = InNumParticles;
+}
+
+void FClothConstraints::Initialize(
+	Softs::FPBDEvolution* InEvolution,
+	const TArray<Softs::FSolverVec3>& InInterpolatedAnimationPositions,
+	const TArray<Softs::FSolverVec3>& /*InOldAnimationPositions*/, // deprecated
+	const TArray<Softs::FSolverVec3>& InInterpolatedAnimationNormals,
+	const TArray<Softs::FSolverVec3>& InAnimationVelocities,
+	int32 InParticleOffset,
+	int32 InNumParticles)
+{
+	Evolution = InEvolution;
+	AnimationPositions = &InInterpolatedAnimationPositions;
+	OldAnimationPositions_deprecated = nullptr;
+	AnimationNormals = &InInterpolatedAnimationNormals;
+	AnimationVelocities = &InAnimationVelocities;
 	ParticleOffset = InParticleOffset;
 	NumParticles = InNumParticles;
 }
@@ -546,13 +566,30 @@ void FClothConstraints::SetBackstopConstraints(const TConstArrayView<FRealSingle
 
 void FClothConstraints::SetAnimDriveConstraints(const TConstArrayView<FRealSingle>& AnimDriveStiffnessMultipliers, const TConstArrayView<FRealSingle>& AnimDriveDampingMultipliers)
 {
-	AnimDriveConstraints = MakeShared<Softs::FPBDAnimDriveConstraint>(
-		ParticleOffset,
-		NumParticles,
-		*AnimationPositions,
-		*OldAnimationPositions,
-		AnimDriveStiffnessMultipliers,
-		AnimDriveDampingMultipliers);
+	if (AnimationVelocities)
+	{
+		AnimDriveConstraints = MakeShared<Softs::FPBDAnimDriveConstraint>(
+			ParticleOffset,
+			NumParticles,
+			*AnimationPositions,
+			TArray<Softs::FSolverVec3>(),
+			*AnimationVelocities,
+			AnimDriveStiffnessMultipliers,
+			AnimDriveDampingMultipliers);
+	}
+	else
+	{
+		PRAGMA_DISABLE_DEPRECATION_WARNINGS
+		// Deprecated behavior until old Initialize can be removed
+		AnimDriveConstraints = MakeShareable( new Softs::FPBDAnimDriveConstraint(
+			ParticleOffset,
+			NumParticles,
+			*AnimationPositions,
+			*OldAnimationPositions_deprecated,
+			AnimDriveStiffnessMultipliers,
+			AnimDriveDampingMultipliers));
+		PRAGMA_ENABLE_DEPRECATION_WARNINGS
+	}
 	++NumConstraintInits;  // Uses init to update the property tables
 	++NumConstraintRules;
 }
diff --git a/Engine/Plugins/Experimental/ChaosCloth/Source/Chaos/Private/ChaosCloth/ChaosClothingSimulationSolver.cpp b/Engine/Plugins/Experimental/ChaosCloth/Source/Chaos/Private/ChaosCloth/ChaosClothingSimulationSolver.cpp
--- a/Engine/Plugins/Experimental/ChaosCloth/Source/Chaos/Private/ChaosCloth/ChaosClothingSimulationSolver.cpp
+++ b/Engine/Plugins/Experimental/ChaosCloth/Source/Chaos/Private/ChaosCloth/ChaosClothingSimulationSolver.cpp
@@ -26,12 +26,15 @@ DECLARE_CYCLE_STAT(TEXT("Chaos Cloth Solver Update Post Solver Step"), STAT_Chao
 DECLARE_CYCLE_STAT(TEXT("Chaos Cloth Solver Calculate Bounds"), STAT_ChaosClothSolverCalculateBounds, STATGROUP_ChaosCloth);
 DECLARE_CYCLE_STAT(TEXT("Chaos Cloth Solver Particle Pre Simulation Transforms"), STAT_ChaosClothParticlePreSimulationTransforms, STATGROUP_ChaosCloth);
 DECLARE_CYCLE_STAT(TEXT("Chaos Cloth Solver Collision Pre Simulation Transforms"), STAT_ChaosClothCollisionPreSimulationTransforms, STATGROUP_ChaosCloth);
+DECLARE_CYCLE_STAT(TEXT("Chaos Cloth Solver Particle Pre Substep Kinematic Interpolation"), STAT_ChaosClothParticlePreSubstepKinematicInterpolation, STATGROUP_ChaosCloth);
 
 #if INTEL_ISPC && !UE_BUILD_SHIPPING
 bool bChaos_PreSimulationTransforms_ISPC_Enabled = CHAOS_PRE_SIMULATION_TRANSFORMS_ISPC_ENABLED_DEFAULT;
 FAutoConsoleVariableRef CVarChaosPreSimulationTransformsISPCEnabled(TEXT("p.Chaos.PreSimulationTransforms.ISPC"), bChaos_PreSimulationTransforms_ISPC_Enabled, TEXT("Whether to use ISPC optimizations in ApplySimulationTransforms"));
 bool bChaos_CalculateBounds_ISPC_Enabled = CHAOS_CALCULATE_BOUNDS_ISPC_ENABLED_DEFAULT;
 FAutoConsoleVariableRef CVarChaosCalculateBoundsISPCEnabled(TEXT("p.Chaos.CalculateBounds.ISPC"), bChaos_CalculateBounds_ISPC_Enabled, TEXT("Whether to use ISPC optimizations in CalculateBounds"));
+bool bChaos_PreSubstepInterpolation_ISPC_Enabled = CHAOS_PRE_SUBSTEP_INTERPOLATION_ISPC_ENABLED_DEFAULT;
+FAutoConsoleVariableRef CVarChaosPreSubstepInterpolationISPCEnabled(TEXT("p.Chaos.PreSubstepInterpolation.ISPC"), bChaos_PreSubstepInterpolation_ISPC_Enabled, TEXT("Whether to use ISPC optimization in PreSubstep"));
 
 static_assert(sizeof(ispc::FVector3f) == sizeof(Chaos::Softs::FSolverVec3), "sizeof(ispc::FVector) != sizeof(Chaos::Softs::FSolverVec3)");
 static_assert(sizeof(ispc::FTransform3f) == sizeof(Chaos::Softs::FSolverRigidTransform3), "sizeof(ispc::FVector) != sizeof(Chaos::Softs::FSolverRigidTransform3)");
@@ -124,7 +127,11 @@ FClothingSimulationSolver::FClothingSimulationSolver()
 	Evolution->Particles().AddArray(&Normals);
 	Evolution->Particles().AddArray(&OldAnimationPositions);
 	Evolution->Particles().AddArray(&AnimationPositions);
+	Evolution->Particles().AddArray(&InterpolatedAnimationPositions);
+	Evolution->Particles().AddArray(&OldAnimationNormals);
 	Evolution->Particles().AddArray(&AnimationNormals);
+	Evolution->Particles().AddArray(&InterpolatedAnimationNormals);
+	Evolution->Particles().AddArray(&AnimationVelocities);
 
 	Evolution->CollisionParticles().AddArray(&CollisionBoneIndices);
 	Evolution->CollisionParticles().AddArray(&CollisionBaseTransforms);
@@ -134,8 +141,7 @@ FClothingSimulationSolver::FClothingSimulationSolver()
 	Evolution->SetKinematicUpdateFunction(
 		[this](Softs::FSolverParticles& ParticlesInput, const Softs::FSolverReal Dt, const Softs::FSolverReal LocalTime, const int32 Index)
 		{
-			const Softs::FSolverReal Alpha = (LocalTime - Time) / DeltaTime;
-			ParticlesInput.P(Index) = Alpha * AnimationPositions[Index] + ((Softs::FSolverReal)1. - Alpha) * OldAnimationPositions[Index];  // X is the step initial condition, here it's P that needs to be updated so that constraints works with the correct step target
+			ParticlesInput.P(Index) = InterpolatedAnimationPositions[Index];  // X is the step initial condition, here it's P that needs to be updated so that constraints works with the correct step target
 		});
 
 	Evolution->SetCollisionKinematicUpdateFunction(
@@ -330,7 +336,7 @@ int32 FClothingSimulationSolver::AddParticles(int32 NumParticles, uint32 GroupId
 	check(!ClothsConstraints.Find(Offset));  // We cannot already have this Offset in the map, particle ranges are always added, never removed (unless reset)
 
 	ClothsConstraints.Emplace(Offset, MakeUnique<FClothConstraints>())
-		->Initialize(Evolution.Get(), AnimationPositions, OldAnimationPositions, AnimationNormals, Offset, NumParticles);
+		->Initialize(Evolution.Get(), InterpolatedAnimationPositions, OldAnimationPositions, InterpolatedAnimationNormals, AnimationVelocities, Offset, NumParticles);
 
 	// Always starts with particles disabled
 	EnableParticles(Offset, false);
@@ -351,11 +357,13 @@ void FClothingSimulationSolver::ResetStartPose(int32 Offset, int32 NumParticles)
 	Softs::FSolverVec3* const Vs = GetParticleVs(Offset);
 	const Softs::FSolverVec3* const Positions = GetAnimationPositions(Offset);
 	Softs::FSolverVec3* const OldPositions = GetOldAnimationPositions(Offset);
+	Softs::FSolverVec3* const InterpolatedPositions = GetInterpolatedAnimationPositions(Offset);
+	Softs::FSolverVec3* const AnimationVs = GetAnimationVelocities(Offset);
 
 	for (int32 Index = 0; Index < NumParticles; ++Index)
 	{
-		PandInvMs[Index].P = Xs[Index] = OldPositions[Index] = Positions[Index];
-		Vs[Index] = Softs::FSolverVec3(0.);
+		PandInvMs[Index].P = Xs[Index] = OldPositions[Index] = InterpolatedPositions[Index] = Positions[Index];
+		Vs[Index] = AnimationVelocities[Index] = Softs::FSolverVec3(0.);
 	}
 }
 
@@ -801,10 +809,13 @@ void FClothingSimulationSolver::ApplyPreSimulationTransforms()
 					(ispc::FVector3f*)Particles.GetV().GetData(),
 					(ispc::FVector3f*)Particles.XArray().GetData(),
 					(ispc::FVector3f*)OldAnimationPositions.GetData(),
+					(ispc::FVector3f*)AnimationVelocities.GetData(),
 					Particles.GetInvM().GetData(),
+					(const ispc::FVector3f*)AnimationPositions.GetData(),
 					ParticleGroupIds.GetData(),
 					(ispc::FTransform3f*)PreSimulationTransforms.GetData(),
 					(ispc::FVector3f&)DeltaLocalSpaceLocation,
+					DeltaTime,
 					Offset,
 					Range);
 			}
@@ -826,6 +837,9 @@ void FClothingSimulationSolver::ApplyPreSimulationTransforms()
 
 						// Update anim initial state (target updated by skinning)
 						OldAnimationPositions[Index] = GroupSpaceTransform.TransformPositionNoScale(OldAnimationPositions[Index]) - DeltaLocalSpaceLocation;
+
+						// Update Animation velocities
+						AnimationVelocities[Index] = (AnimationPositions[Index] - OldAnimationPositions[Index]) / DeltaTime;
 					}, RangeSize < ClothSolverMinParallelBatchSize);
 			}
 		}, /*bForceSingleThreaded =*/ !bClothSolverParallelClothPreUpdate);
@@ -856,6 +870,50 @@ void FClothingSimulationSolver::ApplyPreSimulationTransforms()
 	}
 }
 
+void FClothingSimulationSolver::PreSubstep(const Softs::FSolverReal InterpolationAlpha)
+{
+	TRACE_CPUPROFILER_EVENT_SCOPE(FClothingSimulationSolver_PreSubstep);
+
+	const TPBDActiveView<Softs::FSolverParticles>& ParticlesActiveView = Evolution->ParticlesActiveView();
+	ParticlesActiveView.RangeFor(
+		[this, InterpolationAlpha](Softs::FSolverParticles& Particles, int32 Offset, int32 Range)
+		{
+			TRACE_CPUPROFILER_EVENT_SCOPE(FClothingSimulationSolver_ParticlePreSubstepKinematicInterpolation);
+			SCOPE_CYCLE_COUNTER(STAT_ChaosClothParticlePreSubstepKinematicInterpolation);
+
+			const int32 RangeSize = Range - Offset;
+
+#if INTEL_ISPC
+			if (bRealTypeCompatibleWithISPC && bChaos_PreSubstepInterpolation_ISPC_Enabled)  // TODO: Make the ISPC works with both Single and Double depending on the FSolverReal type
+			{
+				ispc::PreSubstepInterpolation(
+					(ispc::FVector3f*)InterpolatedAnimationPositions.GetData(),
+					(ispc::FVector3f*)InterpolatedAnimationNormals.GetData(),
+					(const ispc::FVector3f*)AnimationPositions.GetData(),
+					(const ispc::FVector3f*)OldAnimationPositions.GetData(),
+					(const ispc::FVector3f*)AnimationNormals.GetData(),
+					(const ispc::FVector3f*)OldAnimationNormals.GetData(),
+					InterpolationAlpha,
+					Offset,
+					Range);
+			}
+			else
+#endif
+			{
+				PhysicsParallelFor(RangeSize,
+					[this, Offset, InterpolationAlpha](int32 i)
+					{
+						const int32 Index = Offset + i;
+						InterpolatedAnimationPositions[Index] = InterpolationAlpha * AnimationPositions[Index] + ((Softs::FSolverReal)1. - InterpolationAlpha) * OldAnimationPositions[Index];
+						InterpolatedAnimationNormals[Index] = (InterpolationAlpha * AnimationNormals[Index] + ((Softs::FSolverReal)1. - InterpolationAlpha) * OldAnimationNormals[Index]).GetSafeNormal();
+
+					}, RangeSize < ClothSolverMinParallelBatchSize);
+			}
+
+
+		}, /*bForceSingleThreaded =*/ !bClothSolverParallelClothPreUpdate);
+}
+
 void FClothingSimulationSolver::UpdateSolverField()
 {
 	if (Evolution && !PerSolverField.IsEmpty())
@@ -917,6 +975,7 @@ void FClothingSimulationSolver::Update(Softs::FSolverReal InDeltaTime)
 
 		Swap(OldCollisionTransforms, CollisionTransforms);
 		Swap(OldAnimationPositions, AnimationPositions);
+		Swap(OldAnimationNormals, AnimationNormals);
 
 		// Clear external collisions so that they can be re-added
 		CollisionParticlesSize = 0;
@@ -983,6 +1042,7 @@ void FClothingSimulationSolver::Update(Softs::FSolverReal InDeltaTime)
 	
 			for (int32 i = 0; i < NumSubsteps; ++i)
 			{
+				PreSubstep(FMath::Clamp((Softs::FSolverReal)(i + 1) / (Softs::FSolverReal)NumSubsteps, (Softs::FSolverReal)0., (Softs::FSolverReal)1.));
 				Evolution->AdvanceOneTimeStep(SubstepDeltaTime);
 			}
 
diff --git a/Engine/Plugins/Experimental/ChaosCloth/Source/Chaos/Private/ChaosCloth/ChaosClothingSimulationSolver.ispc b/Engine/Plugins/Experimental/ChaosCloth/Source/Chaos/Private/ChaosCloth/ChaosClothingSimulationSolver.ispc
--- a/Engine/Plugins/Experimental/ChaosCloth/Source/Chaos/Private/ChaosCloth/ChaosClothingSimulationSolver.ispc
+++ b/Engine/Plugins/Experimental/ChaosCloth/Source/Chaos/Private/ChaosCloth/ChaosClothingSimulationSolver.ispc
@@ -9,10 +9,13 @@ export void ApplyPreSimulationTransforms(uniform FVector4f ParticlesPandInvM[],
 										uniform FVector3f ParticlesV[],
 										uniform FVector3f ParticlesX[],
 										uniform FVector3f OldAnimationPositions[],
+										uniform FVector3f AnimationVelocities[],
 										const uniform float InvM[],
+										const uniform FVector3f AnimationPositions[],
 										const uniform uint32 ParticleGroupIds[],
 										const uniform FTransform3f PreSimulationTransforms[],
 										const uniform FVector3f &DeltaLocalSpaceLocation,
+										const uniform float DeltaTime,
 										const uniform int32 Offset,
 										const uniform int32 Range)
 {
@@ -36,7 +39,38 @@ export void ApplyPreSimulationTransforms(uniform FVector4f ParticlesPandInvM[],
 
 		const varying FVector3f OldAnimationPosition = VectorLoad(&OldAnimationPositions[Index]);
 		const varying FVector3f AP = TransformPosition(GroupSpaceTransform, OldAnimationPosition) - DeltaLocalSpaceLocation;
+
 		VectorStore(&OldAnimationPositions[Index], AP);
+
+		const varying FVector3f AnimationPosition = VectorLoad(&AnimationPositions[Index]);
+		const varying FVector3f AV = (AnimationPosition - AP) / DeltaTime;
+
+		VectorStore(&AnimationVelocities[Index], AV);
+	}
+}
+
+export void PreSubstepInterpolation( uniform FVector3f InterpolatedPositions[],
+									uniform FVector3f InterpolatedNormals[],
+									const uniform FVector3f AnimationPositions[],
+									const uniform FVector3f OldAnimationPositions[],
+									const uniform FVector3f AnimationNormals[],
+									const uniform FVector3f OldAnimationNormals[],
+									const uniform float InterpolationAlpha,
+									const uniform int32 Offset,
+									const uniform int32 Range)
+{
+	foreach(i=Offset ... Range)
+	{
+		const uniform int32 Index = extract(i, 0);
+		const varying FVector3f Position = VectorLoad(&AnimationPositions[Index]);
+		const varying FVector3f OldPosition = VectorLoad(&OldAnimationPositions[Index]);
+		varying FVector3f InterpPosition = InterpolationAlpha * Position + ((uniform float)1.f - InterpolationAlpha) * OldPosition;
+		VectorStore(&InterpolatedPositions[Index], InterpPosition);
+
+		const varying FVector3f Normal = VectorLoad(&AnimationNormals[Index]);
+		const varying FVector3f OldNormal = VectorLoad(&OldAnimationNormals[Index]);
+		varying FVector3f InterpNormal = VectorGetSafeNormal(InterpolationAlpha * Normal + ((uniform float)1.f - InterpolationAlpha) * OldNormal);
+		VectorStore(&InterpolatedNormals[Index], InterpNormal);
 	}
 }
 
diff --git a/Engine/Plugins/Experimental/ChaosCloth/Source/Chaos/Public/ChaosCloth/ChaosClothConstraints.h b/Engine/Plugins/Experimental/ChaosCloth/Source/Chaos/Public/ChaosCloth/ChaosClothConstraints.h
--- a/Engine/Plugins/Experimental/ChaosCloth/Source/Chaos/Public/ChaosCloth/ChaosClothConstraints.h
+++ b/Engine/Plugins/Experimental/ChaosCloth/Source/Chaos/Public/ChaosCloth/ChaosClothConstraints.h
@@ -19,6 +19,7 @@ namespace Chaos
 		~FClothConstraints();
 
 		// ---- Solver interface ----
+		UE_DEPRECATED(5.2, "Use the other Initialize supplying AnimationVelocities for correct AnimDrive behavior")
 		void Initialize(
 			Softs::FPBDEvolution* InEvolution,
 			const TArray<Softs::FSolverVec3>& InAnimationPositions,
@@ -26,6 +27,14 @@ namespace Chaos
 			const TArray<Softs::FSolverVec3>& InAnimationNormals,
 			int32 InParticleOffset,
 			int32 InNumParticles);
+		void Initialize(
+			Softs::FPBDEvolution* InEvolution,
+			const TArray<Softs::FSolverVec3>& InInterpolatedAnimationPositions,
+			const TArray<Softs::FSolverVec3>& /*InOldAnimationPositions*/, // deprecated
+			const TArray<Softs::FSolverVec3>& InInterpolatedAnimationNormals,
+			const TArray<Softs::FSolverVec3>& InAnimationVelocities,
+			int32 InParticleOffset,
+			int32 InNumParticles);
 		// ---- End of Solver interface ----
 
 		// ---- Cloth interface ----
@@ -126,8 +135,9 @@ namespace Chaos
 		
 		Softs::FPBDEvolution* Evolution;
 		const TArray<Softs::FSolverVec3>* AnimationPositions;
-		const TArray<Softs::FSolverVec3>* OldAnimationPositions;
+		const TArray<Softs::FSolverVec3>* OldAnimationPositions_deprecated;
 		const TArray<Softs::FSolverVec3>* AnimationNormals;
+		const TArray<Softs::FSolverVec3>* AnimationVelocities;
 
 		int32 ParticleOffset;
 		int32 NumParticles;
diff --git a/Engine/Plugins/Experimental/ChaosCloth/Source/Chaos/Public/ChaosCloth/ChaosClothingSimulationSolver.h b/Engine/Plugins/Experimental/ChaosCloth/Source/Chaos/Public/ChaosCloth/ChaosClothingSimulationSolver.h
--- a/Engine/Plugins/Experimental/ChaosCloth/Source/Chaos/Public/ChaosCloth/ChaosClothingSimulationSolver.h
+++ b/Engine/Plugins/Experimental/ChaosCloth/Source/Chaos/Public/ChaosCloth/ChaosClothingSimulationSolver.h
@@ -152,9 +152,17 @@ namespace Chaos
 		Softs::FSolverVec3* GetOldAnimationPositions(int32 Offset) { return OldAnimationPositions.GetData() + Offset; }
 		const Softs::FSolverVec3* GetAnimationPositions(int32 Offset) const { return AnimationPositions.GetData() + Offset; }
 		Softs::FSolverVec3* GetAnimationPositions(int32 Offset) { return AnimationPositions.GetData() + Offset; }
+		const Softs::FSolverVec3* GetInterpolatedAnimationPositions(int32 Offset) const { return InterpolatedAnimationPositions.GetData() + Offset; }
+		Softs::FSolverVec3* GetInterpolatedAnimationPositions(int32 Offset) { return InterpolatedAnimationPositions.GetData() + Offset; }
+		const Softs::FSolverVec3* GetOldAnimationNormals(int32 Offset) const { return OldAnimationNormals.GetData() + Offset; }
+		Softs::FSolverVec3* GetOldAnimationNormals(int32 Offset) { return OldAnimationNormals.GetData() + Offset; }
 		const Softs::FSolverVec3* GetAnimationNormals(int32 Offset) const { return AnimationNormals.GetData() + Offset; }
 		Softs::FSolverVec3* GetAnimationNormals(int32 Offset) { return AnimationNormals.GetData() + Offset; }
+		const Softs::FSolverVec3* GetInterpolatedAnimationNormals(int32 Offset) const { return InterpolatedAnimationNormals.GetData() + Offset; }
+		Softs::FSolverVec3* GetInterpolatedAnimationNormals(int32 Offset) { return InterpolatedAnimationNormals.GetData() + Offset; }
 		const Softs::FSolverVec3* GetNormals(int32 Offset) const { return Normals.GetData() + Offset; }
+		const Softs::FSolverVec3* GetAnimationVelocities(int32 Offset) const { return AnimationVelocities.GetData() + Offset; }
+		Softs::FSolverVec3* GetAnimationVelocities(int32 Offset) { return AnimationVelocities.GetData() + Offset; }
 		Softs::FSolverVec3* GetNormals(int32 Offset) { return Normals.GetData() + Offset; }
 		const Softs::FPAndInvM* GetParticlePandInvMs(int32 Offset) const;
 		Softs::FPAndInvM* GetParticlePandInvMs(int32 Offset);
@@ -202,6 +210,7 @@ namespace Chaos
 		void ResetParticles();
 		void ResetCollisionParticles(int32 InCollisionParticlesOffset = 0);
 		void ApplyPreSimulationTransforms();
+		void PreSubstep(const Softs::FSolverReal InterpolationAlpha);
 		Softs::FSolverReal SetParticleMassPerArea(int32 Offset, int32 Size, const FTriangleMesh& Mesh);
 		void ParticleMassUpdateDensity(const FTriangleMesh& Mesh, Softs::FSolverReal Density);
 		void ParticleMassClampAndKinematicStateUpdate(int32 Offset, int32 Size, Softs::FSolverReal MinPerParticleMass, const TFunctionRef<bool(int32)>& KinematicPredicate);
@@ -224,7 +233,11 @@ namespace Chaos
 		TArrayCollectionArray<Softs::FSolverVec3> Normals;
 		TArrayCollectionArray<Softs::FSolverVec3> OldAnimationPositions;
 		TArrayCollectionArray<Softs::FSolverVec3> AnimationPositions;
+		TArrayCollectionArray<Softs::FSolverVec3> InterpolatedAnimationPositions;
+		TArrayCollectionArray<Softs::FSolverVec3> OldAnimationNormals;
 		TArrayCollectionArray<Softs::FSolverVec3> AnimationNormals;
+		TArrayCollectionArray<Softs::FSolverVec3> InterpolatedAnimationNormals;
+		TArrayCollectionArray<Softs::FSolverVec3> AnimationVelocities;
 
 		// Collision particle attributes
 		TArrayCollectionArray<int32> CollisionBoneIndices;
@@ -275,11 +288,17 @@ namespace Chaos
 #define CHAOS_PRE_SIMULATION_TRANSFORMS_ISPC_ENABLED_DEFAULT 1
 #endif
 
+#if !defined(CHAOS_PRE_SUBSTEP_INTERPOLATION_ISPC_ENABLED_DEFAULT)
+#define CHAOS_PRE_SUBSTEP_INTERPOLATION_ISPC_ENABLED_DEFAULT 1
+#endif
+
 // Support run-time toggling on supported platforms in non-shipping configurations
 #if !INTEL_ISPC || UE_BUILD_SHIPPING
 static constexpr bool bChaos_CalculateBounds_ISPC_Enabled = INTEL_ISPC && CHAOS_CALCULATE_BOUNDS_ISPC_ENABLED_DEFAULT;
 static constexpr bool bChaos_PreSimulationTransforms_ISPC_Enabled = INTEL_ISPC && CHAOS_PRE_SIMULATION_TRANSFORMS_ISPC_ENABLED_DEFAULT;
+static constexpr bool bChaos_PreSubstepInterpolation_ISPC_Enabled = INTEL_ISPC && CHAOS_PRE_SUBSTEP_INTERPOLATION_ISPC_ENABLED_DEFAULT;
 #else
 extern bool bChaos_PreSimulationTransforms_ISPC_Enabled;
 extern bool bChaos_CalculateBounds_ISPC_Enabled;
+extern bool bChaos_PreSubstepInterpolation_ISPC_Enabled;
 #endif
diff --git a/Engine/Source/Runtime/Experimental/Chaos/Public/Chaos/PBDAnimDriveConstraint.h b/Engine/Source/Runtime/Experimental/Chaos/Public/Chaos/PBDAnimDriveConstraint.h
--- a/Engine/Source/Runtime/Experimental/Chaos/Public/Chaos/PBDAnimDriveConstraint.h
+++ b/Engine/Source/Runtime/Experimental/Chaos/Public/Chaos/PBDAnimDriveConstraint.h
@@ -14,6 +14,7 @@ namespace Chaos::Softs
 	class FPBDAnimDriveConstraint final
 	{
 	public:
+		UE_DEPRECATED(5.2, "Use the other constructor supplying AnimationVelocities for correct subframe and damping behavior")
 		FPBDAnimDriveConstraint(
 			const int32 InParticleOffset,
 			const int32 InParticleCount,
@@ -23,9 +24,31 @@ namespace Chaos::Softs
 			const TConstArrayView<FRealSingle>& DampingMultipliers  // Use local indexation
 		)
 			: AnimationPositions(InAnimationPositions)
-			, OldAnimationPositions(InOldAnimationPositions)
+			, OldAnimationPositions_deprecated(InOldAnimationPositions)
+			, AnimationVelocities(InOldAnimationPositions) // Unused when using deprecated apply
 			, ParticleOffset(InParticleOffset)
 			, ParticleCount(InParticleCount)
+			, UseDeprecatedApply(true)
+			, Stiffness(FSolverVec2::UnitVector, StiffnessMultipliers, InParticleCount)
+			, Damping(FSolverVec2::UnitVector, DampingMultipliers, InParticleCount)
+		{
+		}
+
+		FPBDAnimDriveConstraint(
+			const int32 InParticleOffset,
+			const int32 InParticleCount,
+			const TArray<FSolverVec3>& InAnimationPositions,  // Use global indexation (will need adding ParticleOffset)
+			const TArray<FSolverVec3>& /*InOldAnimationPositions*/,  // deprecated
+			const TArray<FSolverVec3>& InAnimationVelocities,  // Use global indexation (will need adding ParticleOffset)
+			const TConstArrayView<FRealSingle>& StiffnessMultipliers,  // Use local indexation
+			const TConstArrayView<FRealSingle>& DampingMultipliers  // Use local indexation
+		)
+			: AnimationPositions(InAnimationPositions)
+			, OldAnimationPositions_deprecated(InAnimationVelocities) // Unused when not using apply
+			, AnimationVelocities(InAnimationVelocities)
+			, ParticleOffset(InParticleOffset)
+			, ParticleCount(InParticleCount)
+			, UseDeprecatedApply(false)
 			, Stiffness(FSolverVec2::UnitVector, StiffnessMultipliers, InParticleCount)
 			, Damping(FSolverVec2::UnitVector, DampingMultipliers, InParticleCount)
 		{
@@ -60,21 +83,46 @@ namespace Chaos::Softs
 			{
 				if (Damping.HasWeightMap())
 				{
-					PhysicsParallelFor(ParticleCount, [this, &InParticles, &Dt](int32 Index)  // TODO: profile needed for these parallel loop based on particle count
+					if (UseDeprecatedApply)
 					{
-						const FSolverReal ParticleStiffness = Stiffness[Index];
-						const FSolverReal ParticleDamping = Damping[Index];
-						ApplyHelper(InParticles, ParticleStiffness, ParticleDamping, Dt, Index);
-					});
+						PhysicsParallelFor(ParticleCount, [this, &InParticles, &Dt](int32 Index)  // TODO: profile needed for these parallel loop based on particle count
+							{
+								const FSolverReal ParticleStiffness = Stiffness[Index];
+								const FSolverReal ParticleDamping = Damping[Index];
+								ApplyHelper_Deprecated(InParticles, ParticleStiffness, ParticleDamping, Dt, Index);
+							});
+
+					}
+					else
+					{
+						PhysicsParallelFor(ParticleCount, [this, &InParticles, &Dt](int32 Index)  // TODO: profile needed for these parallel loop based on particle count
+							{
+								const FSolverReal ParticleStiffness = Stiffness[Index];
+								const FSolverReal ParticleDamping = Damping[Index];
+								ApplyHelper(InParticles, ParticleStiffness, ParticleDamping, Dt, Index);
+							});
+					}
 				}
 				else
 				{
 					const FSolverReal ParticleDamping = (FSolverReal)Damping;
-					PhysicsParallelFor(ParticleCount, [this, &InParticles, ParticleDamping, &Dt](int32 Index)
+					if (UseDeprecatedApply)
+					{
+						PhysicsParallelFor(ParticleCount, [this, &InParticles, ParticleDamping, &Dt](int32 Index)
+							{
+								const FSolverReal ParticleStiffness = Stiffness[Index];
+								ApplyHelper_Deprecated(InParticles, ParticleStiffness, ParticleDamping, Dt, Index);
+							});
+
+					}
+					else
 					{
-						const FSolverReal ParticleStiffness = Stiffness[Index];
-						ApplyHelper(InParticles, ParticleStiffness, ParticleDamping, Dt, Index);
-					});
+						PhysicsParallelFor(ParticleCount, [this, &InParticles, ParticleDamping, &Dt](int32 Index)
+							{
+								const FSolverReal ParticleStiffness = Stiffness[Index];
+								ApplyHelper(InParticles, ParticleStiffness, ParticleDamping, Dt, Index);
+							});
+					}
 				}
 			}
 			else
@@ -82,19 +130,40 @@ namespace Chaos::Softs
 				const FSolverReal ParticleStiffness = (FSolverReal)Stiffness;
 				if (Damping.HasWeightMap())
 				{
-					PhysicsParallelFor(ParticleCount, [this, &InParticles, &ParticleStiffness, &Dt](int32 Index)
+					if (UseDeprecatedApply)
 					{
-						const FSolverReal ParticleDamping = Damping[Index];
-						ApplyHelper(InParticles, ParticleStiffness, ParticleDamping, Dt, Index);
-					});
+						PhysicsParallelFor(ParticleCount, [this, &InParticles, &ParticleStiffness, &Dt](int32 Index)
+							{
+								const FSolverReal ParticleDamping = Damping[Index];
+								ApplyHelper_Deprecated(InParticles, ParticleStiffness, ParticleDamping, Dt, Index);
+							});
+					}
+					else
+					{
+						PhysicsParallelFor(ParticleCount, [this, &InParticles, &ParticleStiffness, &Dt](int32 Index)
+							{
+								const FSolverReal ParticleDamping = Damping[Index];
+								ApplyHelper(InParticles, ParticleStiffness, ParticleDamping, Dt, Index);
+							});
+					}
 				}
 				else
 				{
 					const FSolverReal ParticleDamping = (FSolverReal)Damping;
-					PhysicsParallelFor(ParticleCount, [this, &InParticles, &ParticleStiffness, &ParticleDamping, &Dt](int32 Index)
+					if (UseDeprecatedApply)
+					{
+						PhysicsParallelFor(ParticleCount, [this, &InParticles, &ParticleStiffness, &ParticleDamping, &Dt](int32 Index)
+							{
+								ApplyHelper_Deprecated(InParticles, ParticleStiffness, ParticleDamping, Dt, Index);
+							});
+					}
+					else
 					{
-						ApplyHelper(InParticles, ParticleStiffness, ParticleDamping, Dt, Index);
-					});
+						PhysicsParallelFor(ParticleCount, [this, &InParticles, &ParticleStiffness, &ParticleDamping, &Dt](int32 Index)
+							{
+								ApplyHelper(InParticles, ParticleStiffness, ParticleDamping, Dt, Index);
+							});
+					}
 				}
 			}
 		}
@@ -108,9 +177,29 @@ namespace Chaos::Softs
 				return;
 			}
 
+			FSolverVec3& ParticlePosition = Particles.P(ParticleIndex);
+
+			ParticlePosition -= InStiffness * (ParticlePosition - AnimationPositions[ParticleIndex]);
+
+			const FSolverVec3 ParticleDisplacement = ParticlePosition - Particles.X(ParticleIndex);
+			const FSolverVec3 AnimationDisplacement = (AnimationVelocities[ParticleIndex]) * Dt;
+			const FSolverVec3 RelativeDisplacement = ParticleDisplacement - AnimationDisplacement;
+
+			ParticlePosition -= InDamping * RelativeDisplacement;
+		}
+
+		// This method does not have the correct substepping or damping behavior, but retaining for backwards compatibility until deprecated constructor can be removed
+		inline void ApplyHelper_Deprecated(FSolverParticles& Particles, const FSolverReal InStiffness, const FSolverReal InDamping, const FSolverReal Dt, const int32 Index) const
+		{
+			const int32 ParticleIndex = ParticleOffset + Index;
+			if (Particles.InvM(ParticleIndex) == (FSolverReal)0.)
+			{
+				return;
+			}
+
 			FSolverVec3& ParticlePosition = Particles.P(ParticleIndex);
 			const FSolverVec3& AnimationPosition = AnimationPositions[ParticleIndex];
-			const FSolverVec3& OldAnimationPosition = OldAnimationPositions[ParticleIndex];
+			const FSolverVec3& OldAnimationPosition = OldAnimationPositions_deprecated[ParticleIndex];
 
 			const FSolverVec3 ParticleDisplacement = ParticlePosition - Particles.X(ParticleIndex);
 			const FSolverVec3 AnimationDisplacement = OldAnimationPosition - AnimationPosition;
@@ -121,9 +210,11 @@ namespace Chaos::Softs
 
 	private:
 		const TArray<FSolverVec3>& AnimationPositions;  // Use global index (needs adding ParticleOffset)
-		const TArray<FSolverVec3>& OldAnimationPositions;  // Use global index (needs adding ParticleOffset)
+		const TArray<FSolverVec3>& OldAnimationPositions_deprecated;  // Use global index (needs adding ParticleOffset). Only used by ApplyHelper_Deprecated until old constructor can be removed.
+		const TArray<FSolverVec3>& AnimationVelocities;  // Use global index (needs adding ParticleOffset).
 		const int32 ParticleOffset;
 		const int32 ParticleCount;
+		const bool UseDeprecatedApply;
 
 		FPBDStiffness Stiffness;
 		FPBDStiffness Damping;
