Chaos Cloth - Use RoundToInt to calculate the time dependant number of iterations to avoid inadvertently switching between number of iterations when running at fixed framerates.

diff --git a/Engine/Plugins/Experimental/ChaosCloth/Source/Chaos/Private/ChaosCloth/ChaosClothingSimulationSolver.cpp b/Engine/Plugins/Experimental/ChaosCloth/Source/Chaos/Private/ChaosCloth/ChaosClothingSimulationSolver.cpp
--- a/Engine/Plugins/Experimental/ChaosCloth/Source/Chaos/Private/ChaosCloth/ChaosClothingSimulationSolver.cpp
+++ b/Engine/Plugins/Experimental/ChaosCloth/Source/Chaos/Private/ChaosCloth/ChaosClothingSimulationSolver.cpp
@@ -1033,7 +1033,7 @@ void FClothingSimulationSolver::Update(Softs::FSolverReal InDeltaTime)
 
 			const int32 TimeDependentNumIterations = bClothSolverDisableTimeDependentNumIterations ?
 				NumIterations :
-				(int32)(SolverFrequency * DeltaTime * (Softs::FSolverReal)NumIterations);
+				FMath::RoundToInt32(SolverFrequency * DeltaTime * (Softs::FSolverReal)NumIterations);
 
 			Evolution->SetIterations(FMath::Clamp(TimeDependentNumIterations, 1, MaxNumIterations));
 
