Initialize chaos clothing debug materials to null to prevent UB in the case we run ARO while loading the object

diff --git a/Engine/Plugins/Experimental/ChaosCloth/Source/Chaos/Private/ChaosCloth/ChaosClothingSimulation.cpp b/Engine/Plugins/Experimental/ChaosCloth/Source/Chaos/Private/ChaosCloth/ChaosClothingSimulation.cpp
--- a/Engine/Plugins/Experimental/ChaosCloth/Source/Chaos/Private/ChaosCloth/ChaosClothingSimulation.cpp
+++ b/Engine/Plugins/Experimental/ChaosCloth/Source/Chaos/Private/ChaosCloth/ChaosClothingSimulation.cpp
@@ -262,6 +262,10 @@ FClothingSimulation::FClothingSimulation()
 	, bUseGravityOverride(false)
 	, GravityOverride(ChaosClothingSimulationDefault::Gravity)
 	, MaxDistancesMultipliers(ChaosClothingSimulationDefault::MaxDistancesMultipliers)
+#if WITH_EDITOR
+	, DebugClothMaterial(nullptr)
+	, DebugClothMaterialVertex(nullptr)
+#endif
 #if !(UE_BUILD_SHIPPING || UE_BUILD_TEST)
 	, StepCount(0)
 	, ResetCount(0)
