23375251 no longer needed, visualization has been reworked in Main

diff --git a/Engine/Plugins/Experimental/ChaosCloth/Source/Chaos/Private/ChaosCloth/ChaosClothingSimulation.cpp b/Engine/Plugins/Experimental/ChaosCloth/Source/Chaos/Private/ChaosCloth/ChaosClothingSimulation.cpp
--- a/Engine/Plugins/Experimental/ChaosCloth/Source/Chaos/Private/ChaosCloth/ChaosClothingSimulation.cpp
+++ b/Engine/Plugins/Experimental/ChaosCloth/Source/Chaos/Private/ChaosCloth/ChaosClothingSimulation.cpp
@@ -262,10 +262,6 @@ FClothingSimulation::FClothingSimulation()
 	, bUseGravityOverride(false)
 	, GravityOverride(ChaosClothingSimulationDefault::Gravity)
 	, MaxDistancesMultipliers(ChaosClothingSimulationDefault::MaxDistancesMultipliers)
-#if WITH_EDITOR
-	, DebugClothMaterial(nullptr)
-	, DebugClothMaterialVertex(nullptr)
-#endif
 #if !(UE_BUILD_SHIPPING || UE_BUILD_TEST)
 	, StepCount(0)
 	, ResetCount(0)
