Moved using directive at global scope into functions to fix adaptive unity build

diff --git a/Engine/Plugins/Experimental/ChaosCloth/Source/Chaos/Private/ChaosCloth/IspcTestChaosClothingSimulationSolver.cpp b/Engine/Plugins/Experimental/ChaosCloth/Source/Chaos/Private/ChaosCloth/IspcTestChaosClothingSimulationSolver.cpp
--- a/Engine/Plugins/Experimental/ChaosCloth/Source/Chaos/Private/ChaosCloth/IspcTestChaosClothingSimulationSolver.cpp
+++ b/Engine/Plugins/Experimental/ChaosCloth/Source/Chaos/Private/ChaosCloth/IspcTestChaosClothingSimulationSolver.cpp
@@ -6,11 +6,11 @@
 #include "Misc/AutomationTest.h"
 #include "ChaosCloth/ChaosClothingSimulationSolver.h"
 
-using namespace Chaos;
-
 IMPLEMENT_SIMPLE_AUTOMATION_TEST(FIspcTestChaosClothingSimulationSolverAABB, "Ispc.Physics.ChaosClothingSimulationSolver.AABB", EAutomationTestFlags::EditorContext | EAutomationTestFlags::EngineFilter)
 bool FIspcTestChaosClothingSimulationSolverAABB::RunTest(const FString& Parameters)
 {
+	using namespace Chaos;
+	
 	bool InitialState = bChaos_CalculateBounds_ISPC_Enabled;
 
 	const int32 NumParticles = 100;
@@ -69,6 +69,8 @@ bool FIspcTestChaosClothingSimulationSolverAABB::RunTest(const FString& Paramete
 IMPLEMENT_SIMPLE_AUTOMATION_TEST(FIspcTestChaosClothingSimulationSolverPreSimulation, "Ispc.Physics.ChaosClothingSimulationSolver.PreSimulation", EAutomationTestFlags::EditorContext | EAutomationTestFlags::EngineFilter)
 bool FIspcTestChaosClothingSimulationSolverPreSimulation::RunTest(const FString& Parameters)
 {
+	using namespace Chaos;
+	
 	bool InitialState = bChaos_PreSimulationTransforms_ISPC_Enabled;
 
 	const int32 NumParticles = 100;
