Chaos Cloth - Fix the cloth collision thickness when the component is scaled.

diff --git a/Engine/Plugins/Experimental/ChaosCloth/Source/Chaos/Private/ChaosCloth/ChaosClothingSimulationCloth.cpp b/Engine/Plugins/Experimental/ChaosCloth/Source/Chaos/Private/ChaosCloth/ChaosClothingSimulationCloth.cpp
--- a/Engine/Plugins/Experimental/ChaosCloth/Source/Chaos/Private/ChaosCloth/ChaosClothingSimulationCloth.cpp
+++ b/Engine/Plugins/Experimental/ChaosCloth/Source/Chaos/Private/ChaosCloth/ChaosClothingSimulationCloth.cpp
@@ -905,7 +905,10 @@ PRAGMA_ENABLE_DEPRECATION_WARNINGS
 		Solver->SetWindVelocity(GroupId, WindVelocity + Solver->GetWindVelocity());
 
 		// Update general solver properties
-		Solver->SetProperties(GroupId, DampingCoefficient, LocalDampingCoefficient, CollisionThickness, FrictionCoefficient);
+PRAGMA_DISABLE_DEPRECATION_WARNINGS  // TODO: CHAOS_IS_CLOTHINGSIMULATIONMESH_ABSTRACT
+		const Softs::FSolverReal MeshScale = Mesh->GetScale();
+PRAGMA_ENABLE_DEPRECATION_WARNINGS
+		Solver->SetProperties(GroupId, DampingCoefficient, LocalDampingCoefficient, CollisionThickness * MeshScale, FrictionCoefficient);
 
 		// Update use of continuous collision detection
 		Solver->SetUseCCD(GroupId, bUseCCD);
