Chaos Cloth - Fix tether scale clamp that should happen prior to the mesh scale being applied.

diff --git a/Engine/Plugins/Experimental/ChaosCloth/Source/Chaos/Private/ChaosCloth/ChaosClothConstraints.cpp b/Engine/Plugins/Experimental/ChaosCloth/Source/Chaos/Private/ChaosCloth/ChaosClothConstraints.cpp
--- a/Engine/Plugins/Experimental/ChaosCloth/Source/Chaos/Private/ChaosCloth/ChaosClothConstraints.cpp
+++ b/Engine/Plugins/Experimental/ChaosCloth/Source/Chaos/Private/ChaosCloth/ChaosClothConstraints.cpp
@@ -719,11 +719,14 @@ void FClothConstraints::SetVolumeProperties(Softs::FSolverReal VolumeStiffness)
 	}
 }
 
-void FClothConstraints::SetLongRangeAttachmentProperties(const Softs::FSolverVec2& TetherStiffness, const Softs::FSolverVec2& TetherScale)
+void FClothConstraints::SetLongRangeAttachmentProperties(
+	const Softs::FSolverVec2& TetherStiffness,
+	const Softs::FSolverVec2& TetherScale,
+	Softs::FSolverReal MeshScale)
 {
 	if (LongRangeConstraints)
 	{
-		LongRangeConstraints->SetProperties(TetherStiffness, TetherScale);
+		LongRangeConstraints->SetProperties(TetherStiffness, TetherScale, MeshScale);
 	}
 }
 
diff --git a/Engine/Plugins/Experimental/ChaosCloth/Source/Chaos/Private/ChaosCloth/ChaosClothingSimulationCloth.cpp b/Engine/Plugins/Experimental/ChaosCloth/Source/Chaos/Private/ChaosCloth/ChaosClothingSimulationCloth.cpp
--- a/Engine/Plugins/Experimental/ChaosCloth/Source/Chaos/Private/ChaosCloth/ChaosClothingSimulationCloth.cpp
+++ b/Engine/Plugins/Experimental/ChaosCloth/Source/Chaos/Private/ChaosCloth/ChaosClothingSimulationCloth.cpp
@@ -349,7 +349,8 @@ PRAGMA_ENABLE_DEPRECATION_WARNINGS
 	ClothConstraints.SetAreaProperties(Softs::FSolverVec2((FSolverReal)Cloth->AreaStiffness[0], (Softs::FSolverReal)Cloth->AreaStiffness[1]));
 	ClothConstraints.SetLongRangeAttachmentProperties(
 		Softs::FSolverVec2((Softs::FSolverReal)Cloth->TetherStiffness[0], (Softs::FSolverReal)Cloth->TetherStiffness[1]),
-		Softs::FSolverVec2((Softs::FSolverReal)Cloth->TetherScale[0], (Softs::FSolverReal)Cloth->TetherScale[1]) * MeshScale);
+		Softs::FSolverVec2((Softs::FSolverReal)Cloth->TetherScale[0], (Softs::FSolverReal)Cloth->TetherScale[1]),
+		MeshScale);
 	ClothConstraints.SetSelfCollisionProperties((Softs::FSolverReal)Cloth->SelfCollisionThickness, (Softs::FSolverReal)Cloth->SelfCollisionFrictionCoefficient, Cloth->bUseSelfIntersections, Cloth->bUseSelfIntersections);
 	ClothConstraints.SetAnimDriveProperties(
 		Softs::FSolverVec2((Softs::FSolverReal)Cloth->AnimDriveStiffness[0], (Softs::FSolverReal)Cloth->AnimDriveStiffness[1]),
diff --git a/Engine/Plugins/Experimental/ChaosCloth/Source/Chaos/Public/ChaosCloth/ChaosClothConstraints.h b/Engine/Plugins/Experimental/ChaosCloth/Source/Chaos/Public/ChaosCloth/ChaosClothConstraints.h
--- a/Engine/Plugins/Experimental/ChaosCloth/Source/Chaos/Public/ChaosCloth/ChaosClothConstraints.h
+++ b/Engine/Plugins/Experimental/ChaosCloth/Source/Chaos/Public/ChaosCloth/ChaosClothConstraints.h
@@ -73,7 +73,10 @@ namespace Chaos
 		void SetAreaProperties(const Softs::FSolverVec2& AreaStiffness);
 		void SetThinShellVolumeProperties(Softs::FSolverReal VolumeStiffness);
 		void SetVolumeProperties(Softs::FSolverReal VolumeStiffness);
-		void SetLongRangeAttachmentProperties(const Softs::FSolverVec2& TetherStiffness, const Softs::FSolverVec2& TetherScale);
+		void SetLongRangeAttachmentProperties(
+			const Softs::FSolverVec2& TetherStiffness,
+			const Softs::FSolverVec2& TetherScale,
+			Softs::FSolverReal MeshScale = (Softs::FSolverReal)1.);
 		void SetMaximumDistanceProperties(Softs::FSolverReal MaxDistancesMultiplier);
 		void SetAnimDriveProperties(const Softs::FSolverVec2& AnimDriveStiffness, const Softs::FSolverVec2& AnimDriveDamping);
 		void SetSelfCollisionProperties(Softs::FSolverReal SelfCollisionThickness, Softs::FSolverReal SelfCollisionFriction = Softs::FPBDCollisionSpringConstraintsBase::BackCompatFrictionCoefficient, 
diff --git a/Engine/Source/Runtime/Experimental/Chaos/Public/Chaos/PBDLongRangeConstraintsBase.h b/Engine/Source/Runtime/Experimental/Chaos/Public/Chaos/PBDLongRangeConstraintsBase.h
--- a/Engine/Source/Runtime/Experimental/Chaos/Public/Chaos/PBDLongRangeConstraintsBase.h
+++ b/Engine/Source/Runtime/Experimental/Chaos/Public/Chaos/PBDLongRangeConstraintsBase.h
@@ -50,10 +50,10 @@ PRAGMA_ENABLE_DEPRECATION_WARNINGS
 	FSolverVec2 GetStiffness() const { return Stiffness.GetWeightedValue(); }
 
 	// Set the stiffness and scale values used by the constraint
-	void SetProperties(const FSolverVec2& InStiffness, const FSolverVec2& InScale)
+	void SetProperties(const FSolverVec2& InStiffness, const FSolverVec2& InTetherScale, FSolverReal MeshScale = (FSolverReal)1.)
 	{
 		Stiffness.SetWeightedValue(InStiffness);
-		TetherScale.SetWeightedValue(InScale.ClampAxes(MinTetherScale, MaxTetherScale));
+		TetherScale.SetWeightedValue(InTetherScale.ClampAxes(MinTetherScale, MaxTetherScale) * MeshScale);
 	}
 
 	// Set the stiffness input values used by the constraint
