Chaos Cloth - Trivial change in config file comment to clarify that legacy meant PhysX.

#rnx

diff --git a/Engine/Plugins/Experimental/ChaosCloth/Source/Chaos/Private/ChaosCloth/ChaosClothConfig.cpp b/Engine/Plugins/Experimental/ChaosCloth/Source/Chaos/Private/ChaosCloth/ChaosClothConfig.cpp
--- a/Engine/Plugins/Experimental/ChaosCloth/Source/Chaos/Private/ChaosCloth/ChaosClothConfig.cpp
+++ b/Engine/Plugins/Experimental/ChaosCloth/Source/Chaos/Private/ChaosCloth/ChaosClothConfig.cpp
@@ -13,20 +13,6 @@
 
 #include UE_INLINE_GENERATED_CPP_BY_NAME(ChaosClothConfig)
 
-// Legacy parameters not yet migrated to Chaos parameters:
-//  VerticalConstraintConfig.CompressionLimit
-//  VerticalConstraintConfig.StretchLimit
-//  HorizontalConstraintConfig.CompressionLimit
-//  HorizontalConstraintConfig.StretchLimit
-//  BendConstraintConfig.CompressionLimit
-//  BendConstraintConfig.StretchLimit
-//  ShearConstraintConfig.CompressionLimit
-//  ShearConstraintConfig.StretchLimit
-//  SelfCollisionStiffness
-//  SelfCollisionCullScale
-//  LinearDrag
-//  AngularDrag
-//  StiffnessFrequency
 
 UChaosClothConfig::UChaosClothConfig()
 {}
@@ -37,6 +23,21 @@ UChaosClothConfig::~UChaosClothConfig()
 void UChaosClothConfig::MigrateFrom(const FClothConfig_Legacy& ClothConfig)
 {
 #if WITH_EDITORONLY_DATA
+	// Legacy PhysX parameters never migrated to Chaos Cloth Config parameters:
+	//  VerticalConstraintConfig.CompressionLimit
+	//  VerticalConstraintConfig.StretchLimit
+	//  HorizontalConstraintConfig.CompressionLimit
+	//  HorizontalConstraintConfig.StretchLimit
+	//  BendConstraintConfig.CompressionLimit
+	//  BendConstraintConfig.StretchLimit
+	//  ShearConstraintConfig.CompressionLimit
+	//  ShearConstraintConfig.StretchLimit
+	//  SelfCollisionStiffness
+	//  SelfCollisionCullScale
+	//  LinearDrag
+	//  AngularDrag
+	//  StiffnessFrequency
+
 	const float VerticalStiffness =
 		ClothConfig.VerticalConstraintConfig.Stiffness;
 	const float HorizontalStiffness =
