Chaos Cloth - Add new ChaosClothingSimulationConfig class to manage both Chaos Configs and Property Collections at runtime.

#rnx

diff --git a/Engine/Plugins/Experimental/ChaosCloth/Source/Chaos/Private/ChaosCloth/ChaosClothingSimulationConfig.cpp b/Engine/Plugins/Experimental/ChaosCloth/Source/Chaos/Private/ChaosCloth/ChaosClothingSimulationConfig.cpp
new file mode 100644
--- /dev/null
+++ b/Engine/Plugins/Experimental/ChaosCloth/Source/Chaos/Private/ChaosCloth/ChaosClothingSimulationConfig.cpp
@@ -0,0 +1,223 @@
+// Copyright Epic Games, Inc. All Rights Reserved.
+#include "ChaosCloth/ChaosClothingSimulationConfig.h"
+#include "ChaosCloth/ChaosClothConfig.h"
+#include "Chaos/PropertyCollectionAdapter.h"
+#include "Chaos/PBDLongRangeConstraints.h"  // For Tether modes
+#include "GeometryCollection/ManagedArrayCollection.h"
+
+namespace Chaos
+{
+	FClothingSimulationConfig::FClothingSimulationConfig()
+		: PropertyCollection(MakeShared<FManagedArrayCollection>())
+		, Properties(MakeUnique<Softs::FPropertyCollectionMutableAdapter>(PropertyCollection))
+	{
+	}
+
+	FClothingSimulationConfig::FClothingSimulationConfig(const TSharedPtr<const FManagedArrayCollection>& InPropertyCollection)
+		: PropertyCollection(MakeShared<FManagedArrayCollection>())
+		, Properties(MakeUnique<Softs::FPropertyCollectionMutableAdapter>(PropertyCollection))
+	{
+		Initialize(InPropertyCollection);
+	}
+
+	void FClothingSimulationConfig::Initialize(const UChaosClothConfig* ClothConfig, const UChaosClothSharedSimConfig* ClothSharedConfig)
+	{
+		constexpr bool bEnable = true;
+		constexpr bool bAnimatable = true;
+
+		// Clear all properties
+		Properties->Reset();
+
+		// Solver properties
+		if (ClothSharedConfig)
+		{
+			Properties->AddValue(TEXT("NumIterations"), ClothSharedConfig->IterationCount, bEnable, bAnimatable);
+			Properties->AddValue(TEXT("MaxNumIterations"), ClothSharedConfig->MaxIterationCount, bEnable, bAnimatable);
+			Properties->AddValue(TEXT("NumSubsteps"), ClothSharedConfig->SubdivisionCount, bEnable, bAnimatable);
+		}
+
+		// Cloth properties
+		if (ClothConfig)
+		{
+			// Mass
+			{
+				float MassValue;
+				switch (ClothConfig->MassMode)
+				{
+				case EClothMassMode::TotalMass:
+					MassValue = ClothConfig->TotalMass;
+					break;
+				case EClothMassMode::UniformMass:
+					MassValue = ClothConfig->UniformMass;
+					break;
+				default:
+				case EClothMassMode::Density:
+					MassValue = ClothConfig->Density;
+					break;
+				}
+				Properties->AddValue(TEXT("MassMode"), (int32)ClothConfig->MassMode);
+				Properties->AddValue(TEXT("MassValue"), MassValue);
+				Properties->AddValue(TEXT("MinPerParticleMass"), ClothConfig->MinPerParticleMass);
+			}
+
+			// Edge constraint
+			if (ClothConfig->EdgeStiffnessWeighted.Low > 0.f || ClothConfig->EdgeStiffnessWeighted.High > 0.f)
+			{
+				const int32 EdgeSpringStiffnessIndex = Properties->AddProperty(TEXT("EdgeSpringStiffness"), bEnable, bAnimatable);
+				Properties->SetWeightedValue(EdgeSpringStiffnessIndex, ClothConfig->EdgeStiffnessWeighted.Low, ClothConfig->EdgeStiffnessWeighted.High);
+				Properties->SetStringValue(EdgeSpringStiffnessIndex, TEXT("EdgeStiffness"));
+			}
+
+			// Bending constraint
+			if (ClothConfig->BendingStiffnessWeighted.Low > 0.f || ClothConfig->BendingStiffnessWeighted.High > 0.f ||
+				(ClothConfig->bUseBendingElements && (ClothConfig->BucklingStiffnessWeighted.Low > 0.f || ClothConfig->BucklingStiffnessWeighted.High > 0.f)))
+			{
+				Properties->AddValue(TEXT("UseBendingElements"), ClothConfig->bUseBendingElements);
+				if (ClothConfig->bUseBendingElements)
+				{
+					const int32 BendingElementStiffnessIndex = Properties->AddProperty(TEXT("BendingElementStiffness"), bEnable, bAnimatable);
+					Properties->SetWeightedValue(BendingElementStiffnessIndex, ClothConfig->BendingStiffnessWeighted.Low, ClothConfig->BendingStiffnessWeighted.High);
+					Properties->SetStringValue(BendingElementStiffnessIndex, TEXT("BendingStiffness"));
+
+					Properties->AddValue(TEXT("BucklingRatio"), ClothConfig->BucklingRatio);
+
+					if (ClothConfig->BucklingStiffnessWeighted.Low > 0.f && ClothConfig->BucklingStiffnessWeighted.High > 0.f)
+					{
+						const int32 BucklingStiffnessIndex = Properties->AddProperty(TEXT("BucklingStiffness"), bEnable, bAnimatable);
+						Properties->SetWeightedValue(BucklingStiffnessIndex, ClothConfig->BucklingStiffnessWeighted.Low, ClothConfig->BucklingStiffnessWeighted.High);
+						Properties->SetStringValue(BucklingStiffnessIndex, TEXT("BucklingStiffness"));
+					}
+				}
+				else  // Not using bending elements
+				{
+					const int32 BendingSpringStiffnessIndex = Properties->AddProperty(TEXT("BendingSpringStiffness"), bEnable, bAnimatable);
+					Properties->SetWeightedValue(BendingSpringStiffnessIndex, ClothConfig->BendingStiffnessWeighted.Low, ClothConfig->BendingStiffnessWeighted.High);
+					Properties->SetStringValue(BendingSpringStiffnessIndex, TEXT("BendingStiffness"));
+				}
+			}
+
+			// Area constraint
+			if (ClothConfig->AreaStiffnessWeighted.Low > 0.f || ClothConfig->AreaStiffnessWeighted.High > 0.f)
+			{
+				const int32 AreaSpringStiffnessIndex = Properties->AddProperty(TEXT("AreaSpringStiffness"), bEnable, bAnimatable);
+				Properties->SetWeightedValue(AreaSpringStiffnessIndex, ClothConfig->AreaStiffnessWeighted.Low, ClothConfig->AreaStiffnessWeighted.High);
+				Properties->SetStringValue(AreaSpringStiffnessIndex, TEXT("AreaStiffness"));
+			}
+
+			// Long range attachment
+			if (ClothConfig->TetherStiffness.Low > 0.f || ClothConfig->TetherStiffness.High > 0.f)
+			{
+				const Softs::FPBDLongRangeConstraints::EMode TetherMode = ClothConfig->bUseGeodesicDistance ?
+					Softs::FPBDLongRangeConstraints::EMode::Geodesic :
+					Softs::FPBDLongRangeConstraints::EMode::Euclidean;
+				Properties->AddValue(TEXT("TetherMode"), (int32)TetherMode);
+
+				const int32 TetherStiffnessIndex = Properties->AddProperty(TEXT("TetherStiffness"), bEnable, bAnimatable);
+				Properties->SetWeightedValue(TetherStiffnessIndex, ClothConfig->TetherStiffness.Low, ClothConfig->TetherStiffness.High);
+				Properties->SetStringValue(TetherStiffnessIndex, TEXT("TetherStiffness"));
+
+				const int32 TetherScaleIndex = Properties->AddProperty(TEXT("TetherScale"), bEnable, bAnimatable);
+				Properties->SetWeightedValue(TetherScaleIndex, ClothConfig->TetherScale.Low, ClothConfig->TetherScale.High);
+				Properties->SetStringValue(TetherScaleIndex, TEXT("TetherScale"));
+			}
+
+			// AnimDrive
+			if (ClothConfig->AnimDriveStiffness.Low > 0.f || ClothConfig->AnimDriveStiffness.High > 0.f)
+			{
+				const int32 AnimDriveStiffnessIndex = Properties->AddProperty(TEXT("AnimDriveStiffness"), bEnable, bAnimatable);
+				Properties->SetWeightedValue(AnimDriveStiffnessIndex, ClothConfig->AnimDriveStiffness.Low, ClothConfig->AnimDriveStiffness.High);
+				Properties->SetStringValue(AnimDriveStiffnessIndex, TEXT("AnimDriveStiffness"));
+
+				const int32 AnimDriveDampingIndex = Properties->AddProperty(TEXT("AnimDriveDamping"), bEnable, bAnimatable);
+				Properties->SetWeightedValue(AnimDriveDampingIndex, ClothConfig->AnimDriveDamping.Low, ClothConfig->AnimDriveDamping.High);
+				Properties->SetStringValue(AnimDriveDampingIndex, TEXT("AnimDriveDamping"));
+			}
+
+			// Gravity
+			{
+				Properties->AddValue(TEXT("GravityScale"), ClothConfig->GravityScale, bEnable, bAnimatable);
+				Properties->AddValue(TEXT("UseGravityOverride"), ClothConfig->bUseGravityOverride, bEnable, bAnimatable);
+				Properties->AddValue(TEXT("GravityOverride"), FVector3f(ClothConfig->Gravity), bEnable, bAnimatable);
+			}
+
+			// Velocity scale
+			{
+				Properties->AddValue(TEXT("LinearVelocityScale"), FVector3f(ClothConfig->LinearVelocityScale), bEnable, bAnimatable);
+				Properties->AddValue(TEXT("AngularVelocityScale"), ClothConfig->AngularVelocityScale, bEnable, bAnimatable);
+				Properties->AddValue(TEXT("FictitiousAngularScale"), ClothConfig->FictitiousAngularScale, bEnable, bAnimatable);
+			}
+
+			// Aerodynamics
+			Properties->AddValue(TEXT("UsePointBasedWindModel"), ClothConfig->bUsePointBasedWindModel);
+			if (!ClothConfig->bUsePointBasedWindModel && (ClothConfig->Drag.Low > 0.f || ClothConfig->Drag.High > 0.f || ClothConfig->Lift.Low > 0.f || ClothConfig->Lift.High > 0.f))
+			{
+				const int32 DragIndex = Properties->AddProperty(TEXT("Drag"), bEnable, bAnimatable);
+				Properties->SetWeightedValue(DragIndex, ClothConfig->Drag.Low, ClothConfig->Drag.High);
+				Properties->SetStringValue(DragIndex, TEXT("Drag"));
+
+				const int32 LiftIndex = Properties->AddProperty(TEXT("Lift"), bEnable, bAnimatable);
+				Properties->SetWeightedValue(LiftIndex, ClothConfig->Lift.Low, ClothConfig->Lift.High);
+				Properties->SetStringValue(LiftIndex, TEXT("Lift"));
+
+				constexpr float AirDensity = 1.225e-6f;
+				Properties->AddValue(TEXT("AirDensity"), AirDensity, bEnable, bAnimatable);
+			}
+
+			// Pressure
+			if (ClothConfig->Pressure.Low != 0.f || ClothConfig->Pressure.High != 0.f)
+			{
+				const int32 PressureIndex = Properties->AddProperty(TEXT("Pressure"), bEnable, bAnimatable);
+				Properties->SetWeightedValue(PressureIndex, ClothConfig->Pressure.Low, ClothConfig->Pressure.High);
+				Properties->SetStringValue(PressureIndex, TEXT("Pressure"));
+			}
+
+			// Damping
+			Properties->AddValue(TEXT("DampingCoefficient"), ClothConfig->DampingCoefficient, bEnable, bAnimatable);
+			Properties->AddValue(TEXT("LocalDampingCoefficient"), ClothConfig->LocalDampingCoefficient, bEnable, bAnimatable);
+
+			// Collision
+			Properties->AddValue(TEXT("CollisionThickness"), ClothConfig->CollisionThickness, bEnable, bAnimatable);
+			Properties->AddValue(TEXT("FrictionCoefficient"), ClothConfig->FrictionCoefficient, bEnable, bAnimatable);
+			Properties->AddValue(TEXT("UseCCD"), ClothConfig->bUseCCD, bEnable, bAnimatable);
+			Properties->AddValue(TEXT("UseSelfCollisions"), ClothConfig->bUseSelfCollisions);
+			Properties->AddValue(TEXT("SelfCollisionThickness"), ClothConfig->SelfCollisionThickness);
+			Properties->AddValue(TEXT("UseSelfIntersections"), ClothConfig->bUseSelfIntersections);
+			Properties->AddValue(TEXT("SelfCollisionFriction"), ClothConfig->SelfCollisionFriction);
+
+			// Max distance
+			{
+				const int32 MaxDistanceIndex = Properties->AddProperty(TEXT("MaxDistance"));
+				Properties->SetWeightedValue(MaxDistanceIndex, 1.f, 1.f);  // Backward compatibility with legacy mask means it can't be different from 1.
+				Properties->SetStringValue(MaxDistanceIndex, TEXT("MaxDistance"));
+			}
+
+			// Backstop
+			{
+				const int32 BackstopDistanceIndex = Properties->AddProperty(TEXT("BackstopDistance"));
+				Properties->SetWeightedValue(BackstopDistanceIndex, 1.f, 1.f);  // Backward compatibility with legacy mask means it can't be different from 1.
+				Properties->SetStringValue(BackstopDistanceIndex, TEXT("BackstopDistance"));
+
+				const int32 BackstopRadiusIndex = Properties->AddProperty(TEXT("BackstopRadius"));
+				Properties->SetWeightedValue(BackstopRadiusIndex, 1.f, 1.f);  // Backward compatibility with legacy mask means it can't be different from 1.
+				Properties->SetStringValue(BackstopRadiusIndex, TEXT("BackstopRadius"));
+
+				Properties->AddValue(TEXT("UseLegacyBackstop"), ClothConfig->bUseLegacyBackstop);
+			}
+		}
+	}
+
+	void FClothingSimulationConfig::Initialize(const TSharedPtr<const FManagedArrayCollection>& InPropertyCollection)
+	{
+		Properties->Copy(InPropertyCollection);
+	}
+
+	const Softs::FPropertyCollectionConstAdapter& FClothingSimulationConfig::GetProperties() const
+	{
+		return *Properties;
+	}
+
+	Softs::FPropertyCollectionAdapter& FClothingSimulationConfig::GetProperties()
+	{
+		return *Properties;
+	}
+}  // End namespace Chaos
diff --git a/Engine/Plugins/Experimental/ChaosCloth/Source/Chaos/Public/ChaosCloth/ChaosClothingSimulationConfig.h b/Engine/Plugins/Experimental/ChaosCloth/Source/Chaos/Public/ChaosCloth/ChaosClothingSimulationConfig.h
new file mode 100644
--- /dev/null
+++ b/Engine/Plugins/Experimental/ChaosCloth/Source/Chaos/Public/ChaosCloth/ChaosClothingSimulationConfig.h
@@ -0,0 +1,53 @@
+// Copyright Epic Games, Inc. All Rights Reserved.
+#pragma once
+
+#include "Templates/SharedPointer.h"
+#include "Templates/UniquePtr.h"
+
+struct FManagedArrayCollection;
+class UChaosClothConfig;
+class UChaosClothSharedSimConfig;
+
+namespace Chaos
+{
+	namespace Softs
+	{
+		class FPropertyCollectionConstAdapter;
+		class FPropertyCollectionAdapter;
+		class FPropertyCollectionMutableAdapter;
+	}
+
+	// Cloth simulation properties
+	class CHAOSCLOTH_API FClothingSimulationConfig final
+	{
+	public:
+		FClothingSimulationConfig();
+		FClothingSimulationConfig(const TSharedPtr<const FManagedArrayCollection>& InPropertyCollection);
+
+		~FClothingSimulationConfig() = default;
+
+		FClothingSimulationConfig(const FClothingSimulationConfig&) = delete;
+		FClothingSimulationConfig(FClothingSimulationConfig&&) = delete;
+		FClothingSimulationConfig& operator=(const FClothingSimulationConfig&) = delete;
+		FClothingSimulationConfig& operator=(FClothingSimulationConfig&&) = delete;
+
+		/** Initialize configuration from cloth config UObjects. */
+		void Initialize(const UChaosClothConfig* ClothConfig, const UChaosClothSharedSimConfig* ClothSharedConfig);
+
+		/** Initialize config from a property collection. */
+		void Initialize(const TSharedPtr<const FManagedArrayCollection>& InPropertyCollection);
+
+		/** Return a property collection adapter for reading properties from this configuration. */
+		const Softs::FPropertyCollectionConstAdapter& GetProperties() const;
+
+		/** Return a property collection adapter for setting properties for this configuration. */
+		Softs::FPropertyCollectionAdapter& GetProperties();
+
+		/** Return this configuration's internal property collection. */
+		TSharedPtr<const FManagedArrayCollection> GetPropertyCollection() const { return TSharedPtr<const FManagedArrayCollection>(PropertyCollection); }
+
+	private:
+		TSharedPtr<FManagedArrayCollection> PropertyCollection;
+		TUniquePtr<Softs::FPropertyCollectionMutableAdapter> Properties;
+	};
+} // namespace Chaos
diff --git a/Engine/Source/Runtime/Experimental/Chaos/Public/Chaos/PBDSphericalConstraint.h b/Engine/Source/Runtime/Experimental/Chaos/Public/Chaos/PBDSphericalConstraint.h
--- a/Engine/Source/Runtime/Experimental/Chaos/Public/Chaos/PBDSphericalConstraint.h
+++ b/Engine/Source/Runtime/Experimental/Chaos/Public/Chaos/PBDSphericalConstraint.h
@@ -31,7 +31,7 @@ class CHAOS_API FPBDSphericalConstraint final
 public:
 	static bool IsEnabled(const FPropertyCollectionConstAdapter& PropertyCollection)
 	{
-		return IsMaxDistanceScaleEnabled(PropertyCollection, false);
+		return IsMaxDistanceEnabled(PropertyCollection, false);
 	}
 
 	FPBDSphericalConstraint(
@@ -45,7 +45,7 @@ public:
 		: AnimationPositions(InAnimationPositions)
 		, SphereRadii(InSphereRadii)
 		, ParticleOffset(InParticleOffset)
-		, Scale((FSolverReal)FMath::Max(GetMaxDistanceScale(PropertyCollection, 1.f), 0.f) * MeshScale)
+		, Scale(MeshScale)
 	{
 		check(InSphereRadii.Num() == InParticleCount);
 	}
@@ -67,10 +67,8 @@ public:
 
 	void SetProperties(const FPropertyCollectionConstAdapter& PropertyCollection, FSolverReal MeshScale)
 	{
-		if (IsMaxDistanceScaleMutable(PropertyCollection))
-		{
-			SetScale((FSolverReal)GetMaxDistanceScale(PropertyCollection), MeshScale);
-		}
+		SetScale((FSolverReal)1., MeshScale);
+		// TODO: MaxDistance
 	}
 
 	void Apply(FSolverParticles& Particles, const FSolverReal Dt) const
@@ -142,7 +140,7 @@ protected:
 private:
 	FSolverReal Scale;
 
-	UE_CHAOS_DECLARE_PROPERTYCOLLECTION_NAME(MaxDistanceScale, float);
+	UE_CHAOS_DECLARE_PROPERTYCOLLECTION_NAME(MaxDistance, float);
 };
 
 class CHAOS_API FPBDSphericalBackstopConstraint final
@@ -150,8 +148,8 @@ class CHAOS_API FPBDSphericalBackstopConstraint final
 public:
 	static bool IsEnabled(const FPropertyCollectionConstAdapter& PropertyCollection)
 	{
-		return IsBackstopScaleEnabled(PropertyCollection, false) ||
-			IsBackstopScaleAnimatable(PropertyCollection, false);  // Backstop can be re-enabled if animated
+		return IsBackstopDistanceEnabled(PropertyCollection, false) ||
+			IsBackstopDistanceAnimatable(PropertyCollection, false);  // Backstop can be re-enabled if animated
 	}
 
 	FPBDSphericalBackstopConstraint(
@@ -169,7 +167,7 @@ public:
 		, SphereRadii(InSphereRadii)
 		, SphereOffsetDistances(InSphereOffsetDistances)
 		, ParticleOffset(InParticleOffset)
-		, Scale((FSolverReal)FMath::Max(GetBackstopScale(PropertyCollection), 0.f) * MeshScale)
+		, Scale(MeshScale)
 		, bEnabled(true)
 		, bUseLegacyBackstop(false)
 	{
@@ -202,15 +200,8 @@ public:
 
 	void SetProperties(const FPropertyCollectionConstAdapter& PropertyCollection, FSolverReal MeshScale)
 	{
-		if (IsBackstopScaleDirty(PropertyCollection) && IsBackstopScaleAnimatable(PropertyCollection))
-		{
-			bEnabled = IsBackstopScaleEnabled(PropertyCollection);
-		}
-
-		if (IsBackstopScaleMutable(PropertyCollection))
-		{
-			SetScale((FSolverReal)GetBackstopScale(PropertyCollection), MeshScale);
-		}
+		SetScale((FSolverReal)1., MeshScale);
+		// TODO: BackstopDistance and BackstopRadius
 	}
 
 	void SetEnabled(bool bInEnabled) { bEnabled = bInEnabled; }
@@ -251,9 +242,9 @@ public:
 		}
 	}
 
-	void SetScale(FSolverReal MaxDistanceScale, FSolverReal MeshScale)
+	void SetScale(FSolverReal BackstopScale, FSolverReal MeshScale)
 	{
-		Scale = FMath::Max(MaxDistanceScale, (FSolverReal)0.) * MeshScale;
+		Scale = FMath::Max(BackstopScale, (FSolverReal)0.) * MeshScale;
 	}
 
 	FSolverReal GetScale() const { return Scale; }
@@ -360,7 +351,8 @@ private:
 	bool bEnabled;
 	bool bUseLegacyBackstop;
 
-	UE_CHAOS_DECLARE_PROPERTYCOLLECTION_NAME(BackstopScale, float);
+	UE_CHAOS_DECLARE_PROPERTYCOLLECTION_NAME(BackstopDistance, float);
+	UE_CHAOS_DECLARE_PROPERTYCOLLECTION_NAME(BackstopRadius, float);
 };
 
 }  // End namespace Chaos::Softs
