Chaos Cloth - Update the Chaos Clothing simulation and the Cloth Asset simulation configuration to use the new Property Collection system.

#rnx

diff --git a/Engine/Plugins/Experimental/ChaosCloth/Source/Chaos/Private/ChaosCloth/ChaosClothConstraints.cpp b/Engine/Plugins/Experimental/ChaosCloth/Source/Chaos/Private/ChaosCloth/ChaosClothConstraints.cpp
--- a/Engine/Plugins/Experimental/ChaosCloth/Source/Chaos/Private/ChaosCloth/ChaosClothConstraints.cpp
+++ b/Engine/Plugins/Experimental/ChaosCloth/Source/Chaos/Private/ChaosCloth/ChaosClothConstraints.cpp
@@ -1,5 +1,6 @@
 // Copyright Epic Games, Inc. All Rights Reserved.
 #include "ChaosCloth/ChaosClothConstraints.h"
+#include "ChaosCloth/ChaosWeightMapTarget.h"
 #include "Chaos/PBDSpringConstraints.h"
 #include "Chaos/XPBDSpringConstraints.h"
 #include "Chaos/PBDBendingConstraints.h"
@@ -15,13 +16,14 @@
 #include "Chaos/PBDTriangleMeshCollisions.h"
 #include "Chaos/PBDTriangleMeshIntersections.h"
 #include "Chaos/PBDEvolution.h"
+#include "Chaos/PropertyCollectionAdapter.h"
 
 namespace Chaos {
 
 FClothConstraints::FClothConstraints()
 	: Evolution(nullptr)
 	, AnimationPositions(nullptr)
-	, OldAnimationPositions_deprecated(nullptr)
+	, OldAnimationPositions_Deprecated(nullptr)
 	, AnimationNormals(nullptr)
 	, AnimationVelocities(nullptr)
 	, ParticleOffset(0)
@@ -51,7 +53,7 @@ void FClothConstraints::Initialize(
 {
 	Evolution = InEvolution;
 	AnimationPositions = &InAnimationPositions;
-	OldAnimationPositions_deprecated = &InOldAnimationPositions;
+	OldAnimationPositions_Deprecated = &InOldAnimationPositions;
 	AnimationNormals = &InAnimationNormals;
 	AnimationVelocities = nullptr;
 	ParticleOffset = InParticleOffset;
@@ -69,7 +71,7 @@ void FClothConstraints::Initialize(
 {
 	Evolution = InEvolution;
 	AnimationPositions = &InInterpolatedAnimationPositions;
-	OldAnimationPositions_deprecated = nullptr;
+	OldAnimationPositions_Deprecated = nullptr;
 	AnimationNormals = &InInterpolatedAnimationNormals;
 	AnimationVelocities = &InAnimationVelocities;
 	ParticleOffset = InParticleOffset;
@@ -93,6 +95,354 @@ void FClothConstraints::Enable(bool bEnable)
 	}
 }
 
+void FClothConstraints::AddRules(
+	const Softs::FPropertyCollectionConstAdapter& ConfigProperties,
+	const FTriangleMesh& TriangleMesh,
+	const TArray<TConstArrayView<FRealSingle>>& WeightMaps,
+	const TArray<TConstArrayView<TTuple<int32, int32, FRealSingle>>>& Tethers,
+	Softs::FSolverReal MeshScale,
+	bool bEnabled)
+{
+	// Self collisions
+	CreateSelfCollisionConstraints(ConfigProperties, TriangleMesh);
+
+	// Edge constraints
+	CreateEdgeConstraints(ConfigProperties, WeightMaps, TriangleMesh);
+
+	// Bending constraints
+	CreateBendingConstraints(ConfigProperties, WeightMaps, TriangleMesh);
+
+	// Area constraints
+	CreateAreaConstraints(ConfigProperties, WeightMaps, TriangleMesh);
+
+	// Long range constraints
+	CreateLongRangeConstraints(ConfigProperties, WeightMaps, Tethers, MeshScale);
+
+	// Max distances
+	CreateMaxDistanceConstraints(ConfigProperties, WeightMaps, MeshScale);
+
+	// Backstop Constraints
+	CreateBackstopConstraints(ConfigProperties, WeightMaps, MeshScale);
+
+	// Animation Drive Constraints
+	CreateAnimDriveConstraints(ConfigProperties, WeightMaps);
+
+	// Commit rules to solver
+PRAGMA_DISABLE_DEPRECATION_WARNINGS
+	CreateRules();  // TODO: Move CreateRules to private
+PRAGMA_ENABLE_DEPRECATION_WARNINGS
+
+	// Enable or disable constraints as requested
+	Enable(bEnabled);
+}
+
+void FClothConstraints::CreateSelfCollisionConstraints(const Softs::FPropertyCollectionConstAdapter& ConfigProperties, const FTriangleMesh& TriangleMesh)
+{
+	const bool bUseSelfCollisions = ConfigProperties.GetValue<bool>(TEXT("UseSelfCollisions"));
+
+	if (bUseSelfCollisions)
+	{
+		static const int32 DisabledCollisionElementsN = 5;  // TODO: Make this a parameter?
+		TSet<TVec2<int32>> DisabledCollisionElements;  // TODO: Is this needed? Turn this into a bit array?
+
+		const TVec2<int32> Range = TriangleMesh.GetVertexRange();
+		for (int32 Index = Range[0]; Index <= Range[1]; ++Index)
+		{
+			const TSet<int32> Neighbors = TriangleMesh.GetNRing(Index, DisabledCollisionElementsN);
+			for (int32 Element : Neighbors)
+			{
+				check(Index != Element);
+				DisabledCollisionElements.Emplace(TVec2<int32>(Index, Element));
+				DisabledCollisionElements.Emplace(TVec2<int32>(Element, Index));
+			}
+		}
+
+		const bool bUseSelfIntersections = ConfigProperties.GetValue<bool>(TEXT("UseSelfIntersections"));
+
+		SelfCollisionInit = MakeShared<Softs::FPBDTriangleMeshCollisions>(
+			ParticleOffset,
+			NumParticles,
+			TriangleMesh,
+			ConfigProperties);
+
+		SelfCollisionConstraints = MakeShared<Softs::FPBDCollisionSpringConstraints>(
+			ParticleOffset,
+			NumParticles,
+			TriangleMesh,
+			AnimationPositions,
+			MoveTemp(DisabledCollisionElements),
+			ConfigProperties);
+
+		++NumConstraintInits;
+		++NumPostCollisionConstraintRules;
+
+		SelfIntersectionConstraints = MakeShared<Softs::FPBDTriangleMeshIntersections>(
+			ParticleOffset,
+			NumParticles,
+			TriangleMesh);
+		++NumConstraintInits;
+	}
+}
+
+void FClothConstraints::CreateEdgeConstraints(
+	const Softs::FPropertyCollectionConstAdapter& ConfigProperties,
+	const TArray<TConstArrayView<FRealSingle>>& WeightMaps,
+	const FTriangleMesh& TriangleMesh)
+{
+	if (Softs::FXPBDEdgeSpringConstraints::IsEnabled(ConfigProperties))
+	{
+		const TConstArrayView<FRealSingle>& EdgeStiffnessMultipliers = WeightMaps[(int32)EChaosWeightMapTarget::EdgeStiffness];
+		const TConstArrayView<FRealSingle> DampingMultipliers;  // TODO: Damping multiplier
+
+		XEdgeConstraints = MakeShared<Softs::FXPBDEdgeSpringConstraints>(
+			Evolution->Particles(),
+			ParticleOffset,
+			NumParticles,
+			TriangleMesh.GetSurfaceElements(),
+			EdgeStiffnessMultipliers,
+			DampingMultipliers,
+			ConfigProperties,
+			/*bTrimKinematicConstraints =*/ true);
+
+		++NumConstraintInits;  // Uses init to update the property tables
+		++NumConstraintRules;
+	}
+	else if (Softs::FPBDEdgeSpringConstraints::IsEnabled(ConfigProperties))
+	{
+		const TConstArrayView<FRealSingle>& EdgeStiffnessMultipliers = WeightMaps[(int32)EChaosWeightMapTarget::EdgeStiffness];
+
+		EdgeConstraints = MakeShared<Softs::FPBDEdgeSpringConstraints>(
+			Evolution->Particles(),
+			ParticleOffset,
+			NumParticles,
+			TriangleMesh.GetSurfaceElements(),
+			EdgeStiffnessMultipliers,
+			ConfigProperties,
+			/*bTrimKinematicConstraints =*/ true);
+
+		++NumConstraintInits;  // Uses init to update the property tables
+		++NumConstraintRules;
+	}
+}
+
+void FClothConstraints::CreateBendingConstraints(
+	const Softs::FPropertyCollectionConstAdapter& ConfigProperties,
+	const TArray<TConstArrayView<FRealSingle>>& WeightMaps,
+	const FTriangleMesh& TriangleMesh)
+{
+	const bool bUseBendingElements = ConfigProperties.GetValue<bool>(TEXT("UseBendingElements"));
+
+	if (bUseBendingElements)
+	{
+		if (Softs::FXPBDBendingConstraints::IsEnabled(ConfigProperties))
+		{
+			TArray<Chaos::TVec4<int32>> BendingElements = TriangleMesh.GetUniqueAdjacentElements();
+			const TConstArrayView<FRealSingle>& BendingStiffnessMultipliers = WeightMaps[(int32)EChaosWeightMapTarget::BendingStiffness];
+			const TConstArrayView<FRealSingle>& BucklingStiffnessMultipliers = WeightMaps[(int32)EChaosWeightMapTarget::BucklingStiffness];
+			const TConstArrayView<FRealSingle> DampingMultipliers;
+
+			XBendingElementConstraints = MakeShared<Softs::FXPBDBendingConstraints>(
+				Evolution->Particles(),
+				ParticleOffset, NumParticles,
+				MoveTemp(BendingElements),
+				BendingStiffnessMultipliers,
+				BucklingStiffnessMultipliers,
+				DampingMultipliers,
+				ConfigProperties,
+				/*bTrimKinematicConstraints =*/ true);
+
+			++NumConstraintInits;  // Uses init to update the property tables
+			++NumConstraintRules;
+		}
+		else if (Softs::FPBDBendingConstraints::IsEnabled(ConfigProperties))
+		{
+			TArray<Chaos::TVec4<int32>> BendingElements = TriangleMesh.GetUniqueAdjacentElements();
+			const TConstArrayView<FRealSingle>& BendingStiffnessMultipliers = WeightMaps[(int32)EChaosWeightMapTarget::BendingStiffness];
+			const TConstArrayView<FRealSingle>& BucklingStiffnessMultipliers = WeightMaps[(int32)EChaosWeightMapTarget::BucklingStiffness];
+
+			BendingElementConstraints = MakeShared<Softs::FPBDBendingConstraints>(
+				Evolution->Particles(),
+				ParticleOffset, NumParticles,
+				MoveTemp(BendingElements),
+				BendingStiffnessMultipliers,
+				BucklingStiffnessMultipliers,
+				ConfigProperties,
+				/*bTrimKinematicConstraints =*/ true);
+
+			++NumConstraintInits;  // Uses init to update the property tables
+			++NumConstraintRules;
+		}
+	}
+	else
+	{
+		if (Softs::FXPBDBendingSpringConstraints::IsEnabled(ConfigProperties))
+		{
+			const TArray<Chaos::TVec2<int32>> CrossEdges = TriangleMesh.GetUniqueAdjacentPoints();
+			const TConstArrayView<FRealSingle>& BendingStiffnessMultipliers = WeightMaps[(int32)EChaosWeightMapTarget::BendingStiffness];
+			const TConstArrayView<FRealSingle> DampingMultipliers;
+
+			XBendingConstraints = MakeShared<Softs::FXPBDBendingSpringConstraints>(
+				Evolution->Particles(),
+				ParticleOffset, NumParticles,
+				CrossEdges,
+				BendingStiffnessMultipliers,
+				DampingMultipliers,
+				ConfigProperties,
+				/*bTrimKinematicConstraints =*/ true);
+		}
+		else if (Softs::FPBDBendingSpringConstraints::IsEnabled(ConfigProperties))
+		{
+			const TArray<Chaos::TVec2<int32>> CrossEdges = TriangleMesh.GetUniqueAdjacentPoints();
+			const TConstArrayView<FRealSingle>& BendingStiffnessMultipliers = WeightMaps[(int32)EChaosWeightMapTarget::BendingStiffness];
+
+			BendingConstraints = MakeShared<Softs::FPBDBendingSpringConstraints>(
+				Evolution->Particles(),
+				ParticleOffset,
+				NumParticles,
+				CrossEdges,
+				BendingStiffnessMultipliers,
+				ConfigProperties,
+				/*bTrimKinematicConstraints =*/ true);
+		}
+		++NumConstraintInits;  // Uses init to update the property tables
+		++NumConstraintRules;
+	}
+}
+
+void FClothConstraints::CreateAreaConstraints(
+	const Softs::FPropertyCollectionConstAdapter& ConfigProperties,
+	const TArray<TConstArrayView<FRealSingle>>& WeightMaps,
+	const FTriangleMesh& TriangleMesh)
+{
+	if (Softs::FXPBDAreaSpringConstraints::IsEnabled(ConfigProperties))
+	{
+		const TConstArrayView<FRealSingle>& AreaStiffnessMultipliers = WeightMaps[(int32)EChaosWeightMapTarget::AreaStiffness];
+
+		XAreaConstraints = MakeShared<Softs::FXPBDAreaSpringConstraints>(
+			Evolution->Particles(),
+			ParticleOffset,
+			NumParticles,
+			TriangleMesh.GetSurfaceElements(),
+			AreaStiffnessMultipliers,
+			ConfigProperties,
+			/*bTrimKinematicConstraints =*/ true);
+
+		++NumConstraintInits;  // Uses init to update the property tables
+		++NumConstraintRules;
+	}
+	else if (Softs::FPBDAreaSpringConstraints::IsEnabled(ConfigProperties))
+	{
+		const TConstArrayView<FRealSingle>& AreaStiffnessMultipliers = WeightMaps[(int32)EChaosWeightMapTarget::AreaStiffness];
+
+		AreaConstraints = MakeShared<Softs::FPBDAreaSpringConstraints>(
+			Evolution->Particles(),
+			ParticleOffset,
+			NumParticles,
+			TriangleMesh.GetSurfaceElements(),
+			AreaStiffnessMultipliers,
+			ConfigProperties,
+			/*bTrimKinematicConstraints =*/ true);
+
+		++NumConstraintInits;  // Uses init to update the property tables
+		++NumConstraintRules;
+	}
+}
+
+void FClothConstraints::CreateLongRangeConstraints(
+	const Softs::FPropertyCollectionConstAdapter& ConfigProperties,
+	const TArray<TConstArrayView<FRealSingle>>& WeightMaps,
+	const TArray<TConstArrayView<TTuple<int32, int32, FRealSingle>>>& Tethers,
+	Softs::FSolverReal MeshScale)
+{
+	if (Softs::FPBDLongRangeConstraints::IsEnabled(ConfigProperties))
+	{
+		const TConstArrayView<FRealSingle>& TetherStiffnessMultipliers = WeightMaps[(int32)EChaosWeightMapTarget::TetherStiffness];
+		const TConstArrayView<FRealSingle>& TetherScaleMultipliers = WeightMaps[(int32)EChaosWeightMapTarget::TetherScale];
+
+		//  Now that we're only doing a single iteration of Long range constraints, and they're more of a fake constraint to jump start our initial guess, it's not clear that using XPBD makes sense here.
+		LongRangeConstraints = MakeShared<Softs::FPBDLongRangeConstraints>(
+			Evolution->Particles(),
+			ParticleOffset,
+			NumParticles,
+			Tethers,
+			TetherStiffnessMultipliers,
+			TetherScaleMultipliers,
+			ConfigProperties,
+			MeshScale);
+
+		++NumConstraintInits;  // Uses init to both update the property tables and apply the constraint
+	}
+}
+
+void FClothConstraints::CreateMaxDistanceConstraints(
+	const Softs::FPropertyCollectionConstAdapter& ConfigProperties,
+	const TArray<TConstArrayView<FRealSingle>>& WeightMaps,
+	Softs::FSolverReal MeshScale)
+{
+	const TConstArrayView<FRealSingle>& MaxDistances = WeightMaps[(int32)EChaosWeightMapTarget::MaxDistance];
+	if (MaxDistances.Num() == NumParticles && Softs::FPBDSphericalConstraint::IsEnabled(ConfigProperties))
+	{
+		MaximumDistanceConstraints = MakeShared<Softs::FPBDSphericalConstraint>(
+			ParticleOffset,
+			NumParticles,
+			*AnimationPositions,
+			MaxDistances,
+			ConfigProperties,
+			MeshScale);
+
+		++NumConstraintRules;
+	}
+}
+
+void FClothConstraints::CreateBackstopConstraints(
+	const Softs::FPropertyCollectionConstAdapter& ConfigProperties,
+	const TArray<TConstArrayView<FRealSingle>>& WeightMaps,
+	Softs::FSolverReal MeshScale)
+{
+	if (Softs::FPBDSphericalBackstopConstraint::IsEnabled(ConfigProperties))
+	{
+		const TConstArrayView<FRealSingle>& BackstopDistances = WeightMaps[(int32)EChaosWeightMapTarget::BackstopDistance];
+		const TConstArrayView<FRealSingle>& BackstopRadiuses = WeightMaps[(int32)EChaosWeightMapTarget::BackstopRadius];
+
+		if (BackstopRadiuses.Num() == NumParticles && BackstopDistances.Num() == NumParticles)
+		{
+			BackstopConstraints = MakeShared<Softs::FPBDSphericalBackstopConstraint>(
+				ParticleOffset,
+				NumParticles,
+				*AnimationPositions,
+				*AnimationNormals,
+				BackstopRadiuses,
+				BackstopDistances,
+				ConfigProperties,
+				MeshScale);
+
+			++NumConstraintRules;
+		}
+	}
+}
+
+void FClothConstraints::CreateAnimDriveConstraints(const Softs::FPropertyCollectionConstAdapter& ConfigProperties, const TArray<TConstArrayView<FRealSingle>>& WeightMaps)
+{
+	if (Softs::FPBDAnimDriveConstraint::IsEnabled(ConfigProperties))
+	{
+		const TConstArrayView<FRealSingle>& AnimDriveStiffnessMultipliers = WeightMaps[(int32)EChaosWeightMapTarget::AnimDriveStiffness];
+		const TConstArrayView<FRealSingle>& AnimDriveDampingMultipliers = WeightMaps[(int32)EChaosWeightMapTarget::AnimDriveDamping];
+		check(AnimationVelocities); // Legacy code didn't use to have AnimationVelocities
+
+		AnimDriveConstraints = MakeShared<Softs::FPBDAnimDriveConstraint>(
+			ParticleOffset,
+			NumParticles,
+			*AnimationPositions,
+			*AnimationVelocities,
+			AnimDriveStiffnessMultipliers,
+			AnimDriveDampingMultipliers,
+			ConfigProperties);
+
+		++NumConstraintInits;  // Uses init to update the property tables
+		++NumConstraintRules;
+	}
+}
+
 void FClothConstraints::CreateRules()
 {
 	check(Evolution);
@@ -135,6 +485,21 @@ void FClothConstraints::CreateRules()
 				XEdgeConstraints->Apply(Particles, Dt);
 			};
 	}
+	if (XEdgeConstraints_Deprecated)  // TODO: Remove for 5.4
+	{
+		ConstraintInits[ConstraintInitIndex++] =
+			[this](Softs::FSolverParticles& /*Particles*/, const Softs::FSolverReal Dt)
+		{
+			XEdgeConstraints_Deprecated->Init();
+			XEdgeConstraints_Deprecated->ApplyProperties(Dt, Evolution->GetIterations());
+		};
+
+		ConstraintRules[ConstraintRuleIndex++] =
+			[this](Softs::FSolverParticles& Particles, const Softs::FSolverReal Dt)
+		{
+			XEdgeConstraints_Deprecated->Apply(Particles, Dt);
+		};
+	}
 	if (EdgeConstraints)
 	{
 		ConstraintInits[ConstraintInitIndex++] =
@@ -148,6 +513,19 @@ void FClothConstraints::CreateRules()
 				EdgeConstraints->Apply(Particles, Dt);
 			};
 	}
+	if (EdgeConstraints_Deprecated)  // TODO: Remove for 5.4
+	{
+		ConstraintInits[ConstraintInitIndex++] =
+			[this](Softs::FSolverParticles& /*Particles*/, const Softs::FSolverReal Dt)
+		{
+			EdgeConstraints_Deprecated->ApplyProperties(Dt, Evolution->GetIterations());
+		};
+		ConstraintRules[ConstraintRuleIndex++] =
+			[this](Softs::FSolverParticles& Particles, const Softs::FSolverReal Dt)
+		{
+			EdgeConstraints_Deprecated->Apply(Particles, Dt);
+		};
+	}
 	if (XBendingConstraints)
 	{
 		ConstraintInits[ConstraintInitIndex++] =
@@ -162,6 +540,20 @@ void FClothConstraints::CreateRules()
 				XBendingConstraints->Apply(Particles, Dt);
 			};
 	}
+	if (XBendingConstraints_Deprecated)  // TODO: Remove for 5.4
+	{
+		ConstraintInits[ConstraintInitIndex++] =
+			[this](Softs::FSolverParticles& /*Particles*/, const Softs::FSolverReal Dt)
+		{
+			XBendingConstraints_Deprecated->Init();
+			XBendingConstraints_Deprecated->ApplyProperties(Dt, Evolution->GetIterations());
+		};
+		ConstraintRules[ConstraintRuleIndex++] =
+			[this](Softs::FSolverParticles& Particles, const Softs::FSolverReal Dt)
+		{
+			XBendingConstraints_Deprecated->Apply(Particles, Dt);
+		};
+	}
 	if (BendingConstraints)
 	{
 		ConstraintInits[ConstraintInitIndex++] =
@@ -175,6 +567,19 @@ void FClothConstraints::CreateRules()
 				BendingConstraints->Apply(Particles, Dt);
 			};
 	}
+	if (BendingConstraints_Deprecated)  // TODO: Remove for 5.4
+	{
+		ConstraintInits[ConstraintInitIndex++] =
+			[this](Softs::FSolverParticles& /*Particles*/, const Softs::FSolverReal Dt)
+		{
+			BendingConstraints_Deprecated->ApplyProperties(Dt, Evolution->GetIterations());
+		};
+		ConstraintRules[ConstraintRuleIndex++] =
+			[this](Softs::FSolverParticles& Particles, const Softs::FSolverReal Dt)
+		{
+			BendingConstraints_Deprecated->Apply(Particles, Dt);
+		};
+	}
 	if (BendingElementConstraints)
 	{
 		ConstraintInits[ConstraintInitIndex++] =
@@ -217,6 +622,20 @@ void FClothConstraints::CreateRules()
 				XAreaConstraints->Apply(Particles, Dt);
 			};
 	}
+	if (XAreaConstraints_Deprecated)  // TODO: Remove for 5.4
+	{
+		ConstraintInits[ConstraintInitIndex++] =
+			[this](Softs::FSolverParticles& /*Particles*/, const Softs::FSolverReal Dt)
+		{
+			XAreaConstraints_Deprecated->Init();
+			XAreaConstraints_Deprecated->ApplyProperties(Dt, Evolution->GetIterations());
+		};
+		ConstraintRules[ConstraintRuleIndex++] =
+			[this](Softs::FSolverParticles& Particles, const Softs::FSolverReal Dt)
+		{
+			XAreaConstraints_Deprecated->Apply(Particles, Dt);
+		};
+	}
 	if (AreaConstraints)
 	{
 		ConstraintInits[ConstraintInitIndex++] =
@@ -230,21 +649,34 @@ void FClothConstraints::CreateRules()
 				AreaConstraints->Apply(Particles, Dt);
 			};
 	}
-	if (ThinShellVolumeConstraints)
+	if (AreaConstraints_Deprecated)  // TODO: Remove for 5.4
+	{
+		ConstraintInits[ConstraintInitIndex++] =
+			[this](Softs::FSolverParticles& /*Particles*/, const Softs::FSolverReal Dt)
+		{
+			AreaConstraints_Deprecated->ApplyProperties(Dt, Evolution->GetIterations());
+		};
+		ConstraintRules[ConstraintRuleIndex++] =
+			[this](Softs::FSolverParticles& Particles, const Softs::FSolverReal Dt)
+		{
+			AreaConstraints_Deprecated->Apply(Particles, Dt);
+		};
+	}
+	if (ThinShellVolumeConstraints_Deprecated)  // TODO: Remove for 5.4
 	{
 		ConstraintRules[ConstraintRuleIndex++] =
 			[this](Softs::FSolverParticles& Particles, const Softs::FSolverReal Dt)
-			{
-				ThinShellVolumeConstraints->Apply(Particles, Dt);
-			};
+		{
+			ThinShellVolumeConstraints_Deprecated->Apply(Particles, Dt);
+		};
 	}
-	if (VolumeConstraints)
+	if (VolumeConstraints_Deprecated)  // TODO: Remove for 5.4
 	{
 		ConstraintRules[ConstraintRuleIndex++] =
 			[this](Softs::FSolverParticles& Particles, const Softs::FSolverReal Dt)
-			{
-				VolumeConstraints->Apply(Particles, Dt);
-			};
+		{
+			VolumeConstraints_Deprecated->Apply(Particles, Dt);
+		};
 	}
 	if (MaximumDistanceConstraints)
 	{
@@ -338,7 +770,7 @@ void FClothConstraints::SetEdgeConstraints(const TArray<TVec3<int32>>& SurfaceEl
 	{
 		const TConstArrayView<FRealSingle> DampingMultipliers;
 
-		XEdgeConstraints = MakeShared<Softs::FXPBDSpringConstraints>(
+		XEdgeConstraints_Deprecated = MakeShared<Softs::FXPBDSpringConstraints>(
 			Evolution->Particles(),
 			ParticleOffset, NumParticles,
 			SurfaceElements,
@@ -349,7 +781,7 @@ void FClothConstraints::SetEdgeConstraints(const TArray<TVec3<int32>>& SurfaceEl
 	}
 	else
 	{
-		EdgeConstraints = MakeShared<Softs::FPBDSpringConstraints>(
+		EdgeConstraints_Deprecated = MakeShared<Softs::FPBDSpringConstraints>(
 			Evolution->Particles(),
 			ParticleOffset,
 			NumParticles,
@@ -365,7 +797,7 @@ void FClothConstraints::SetEdgeConstraints(const TArray<TVec3<int32>>& SurfaceEl
 void FClothConstraints::SetXPBDEdgeConstraints(const TArray<TVec3<int32>>& SurfaceElements, const TConstArrayView<FRealSingle>& StiffnessMultipliers, const TConstArrayView<FRealSingle>& DampingRatioMultipliers)
 {
 	check(Evolution);
-	XEdgeConstraints = MakeShared<Softs::FXPBDSpringConstraints>(
+	XEdgeConstraints_Deprecated = MakeShared<Softs::FXPBDSpringConstraints>(
 		Evolution->Particles(),
 		ParticleOffset, NumParticles,
 		SurfaceElements,
@@ -386,7 +818,7 @@ void FClothConstraints::SetBendingConstraints(const TArray<TVec2<int32>>& Edges,
 	{
 		const TConstArrayView<FRealSingle> DampingMultipliers;
 
-		XBendingConstraints = MakeShared<Softs::FXPBDSpringConstraints>(
+		XBendingConstraints_Deprecated = MakeShared<Softs::FXPBDSpringConstraints>(
 			Evolution->Particles(),
 			ParticleOffset, NumParticles,
 			Edges,
@@ -397,7 +829,7 @@ void FClothConstraints::SetBendingConstraints(const TArray<TVec2<int32>>& Edges,
 	}
 	else
 	{
-		BendingConstraints = MakeShared<Softs::FPBDSpringConstraints>(
+		BendingConstraints_Deprecated = MakeShared<Softs::FPBDSpringConstraints>(
 			Evolution->Particles(),
 			ParticleOffset,
 			NumParticles,
@@ -495,7 +927,7 @@ void FClothConstraints::SetAreaConstraints(const TArray<TVec3<int32>>& SurfaceEl
 
 	if (bUseXPBDConstraints)
 	{
-		XAreaConstraints = MakeShared<Softs::FXPBDAxialSpringConstraints>(
+		XAreaConstraints_Deprecated = MakeShared<Softs::FXPBDAxialSpringConstraints>(
 			Evolution->Particles(),
 			ParticleOffset,
 			NumParticles,
@@ -506,7 +938,7 @@ void FClothConstraints::SetAreaConstraints(const TArray<TVec3<int32>>& SurfaceEl
 	}
 	else
 	{
-		AreaConstraints = MakeShared<Softs::FPBDAxialSpringConstraints>(
+		AreaConstraints_Deprecated = MakeShared<Softs::FPBDAxialSpringConstraints>(
 			Evolution->Particles(),
 			ParticleOffset,
 			NumParticles,
@@ -523,7 +955,7 @@ void FClothConstraints::SetVolumeConstraints(const TArray<TVec2<int32>>& DoubleB
 {
 	check(Evolution);
 
-	ThinShellVolumeConstraints = MakeShared<Softs::FPBDSpringConstraints>(
+	ThinShellVolumeConstraints_Deprecated = MakeShared<Softs::FPBDSpringConstraints>(
 		Evolution->Particles(),
 		ParticleOffset,
 		NumParticles,
@@ -539,7 +971,7 @@ void FClothConstraints::SetVolumeConstraints(TArray<TVec3<int32>>&& SurfaceEleme
 	check(Evolution);
 	check(VolumeStiffness > 0.f && VolumeStiffness <= 1.f);
 
-	VolumeConstraints = MakeShared<Softs::FPBDVolumeConstraint>(Evolution->Particles(), MoveTemp(SurfaceElements), VolumeStiffness);
+	VolumeConstraints_Deprecated = MakeShared<Softs::FPBDVolumeConstraint>(Evolution->Particles(), MoveTemp(SurfaceElements), VolumeStiffness);
 	++NumConstraintRules;
 }
 
@@ -547,7 +979,8 @@ void FClothConstraints::SetLongRangeConstraints(
 	const TArray<TConstArrayView<TTuple<int32, int32, FRealSingle>>>& Tethers,
 	const TConstArrayView<FRealSingle>& TetherStiffnessMultipliers,
 	const TConstArrayView<FRealSingle>& TetherScaleMultipliers,
-	const Softs::FSolverVec2& TetherScale)
+	const Softs::FSolverVec2& TetherScale,
+	Softs::FSolverReal MeshScale)
 {
 	check(Evolution);
 	//  Now that we're only doing a single iteration of Long range constraints, and they're more of a fake constraint to jump start our initial guess, it's not clear that using XPBD makes sense here.
@@ -559,10 +992,26 @@ void FClothConstraints::SetLongRangeConstraints(
 		TetherStiffnessMultipliers,
 		TetherScaleMultipliers,
 		/*InStiffness =*/ Softs::FSolverVec2::UnitVector,
-		TetherScale);
+		TetherScale,
+		MeshScale);
 	++NumConstraintInits;  // Uses init to both update the property tables and apply the constraint
 }
 
+// Deprecated in 5.1
+void FClothConstraints::SetLongRangeConstraints(
+	const TArray<TConstArrayView<TTuple<int32,
+	int32, FRealSingle>>>& Tethers,
+	const TConstArrayView<FRealSingle>& TetherStiffnessMultipliers,
+	const TConstArrayView<FRealSingle>& TetherScaleMultipliers,
+	const Softs::FSolverVec2& TetherScale,
+	bool bUseXPBDConstraints,
+	Softs::FSolverReal MeshScale)
+{
+PRAGMA_DISABLE_DEPRECATION_WARNINGS
+	return SetLongRangeConstraints(Tethers, TetherStiffnessMultipliers, TetherScaleMultipliers, TetherScale, MeshScale);
+PRAGMA_ENABLE_DEPRECATION_WARNINGS
+}
+
 void FClothConstraints::SetMaximumDistanceConstraints(const TConstArrayView<FRealSingle>& MaxDistances)
 {
 	MaximumDistanceConstraints = MakeShared<Softs::FPBDSphericalConstraint>(
@@ -607,7 +1056,7 @@ void FClothConstraints::SetAnimDriveConstraints(const TConstArrayView<FRealSingl
 			ParticleOffset,
 			NumParticles,
 			*AnimationPositions,
-			*OldAnimationPositions_deprecated,
+			*OldAnimationPositions_Deprecated,
 			AnimDriveStiffnessMultipliers,
 			AnimDriveDampingMultipliers));
 		PRAGMA_ENABLE_DEPRECATION_WARNINGS
@@ -659,15 +1108,74 @@ void FClothConstraints::SetSelfCollisionConstraints(const FTriangleMesh& Triangl
 	++NumConstraintInits;
 }
 
+void FClothConstraints::Update(
+	const Softs::FPropertyCollectionConstAdapter& ConfigProperties,
+	Softs::FSolverReal MeshScale,
+	Softs::FSolverReal MaxDistancesScale)
+{
+	if (XEdgeConstraints)
+	{
+		XEdgeConstraints->SetProperties(ConfigProperties);
+	}
+	else if (EdgeConstraints)
+	{
+		EdgeConstraints->SetProperties(ConfigProperties);
+	}
+	if (XBendingConstraints)
+	{
+		XBendingConstraints->SetProperties(ConfigProperties);
+	}
+	else if (BendingConstraints)
+	{
+		BendingConstraints->SetProperties(ConfigProperties);
+	}
+	if (XBendingElementConstraints)
+	{
+		XBendingElementConstraints->SetProperties(ConfigProperties);
+	}
+	else if (BendingElementConstraints)
+	{
+		BendingElementConstraints->SetProperties(ConfigProperties);
+	}
+	if (XAreaConstraints)
+	{
+		XAreaConstraints->SetProperties(ConfigProperties);
+	}
+	else if (AreaConstraints)
+	{
+		AreaConstraints->SetProperties(ConfigProperties);
+	}
+	if (LongRangeConstraints)
+	{
+		LongRangeConstraints->SetProperties(ConfigProperties, MeshScale);
+	}
+	if (MaximumDistanceConstraints)
+	{
+		MaximumDistanceConstraints->SetProperties(ConfigProperties, MeshScale * MaxDistancesScale);
+	}
+	if (AnimDriveConstraints)
+	{
+		AnimDriveConstraints->SetProperties(ConfigProperties);
+	}
+	if (SelfCollisionConstraints)
+	{
+		SelfCollisionConstraints->SetProperties(ConfigProperties);
+	}
+	if (BackstopConstraints)
+	{
+		BackstopConstraints->SetProperties(ConfigProperties, MeshScale);
+	}
+}
+
 void FClothConstraints::SetEdgeProperties(const Softs::FSolverVec2& EdgeStiffness, const Softs::FSolverVec2& DampingRatio)
 {
 	if (EdgeConstraints)
 	{
-		EdgeConstraints->SetProperties(EdgeStiffness);
+		static_cast<Softs::FPBDSpringConstraints*>(EdgeConstraints.Get())->SetProperties(EdgeStiffness);
 	}
 	if (XEdgeConstraints)
 	{
-		XEdgeConstraints->SetProperties(EdgeStiffness, DampingRatio);
+		static_cast<Softs::FXPBDSpringConstraints*>(XEdgeConstraints.Get())->SetProperties(EdgeStiffness, DampingRatio);
 	}
 }
 
@@ -675,11 +1183,11 @@ void FClothConstraints::SetBendingProperties(const Softs::FSolverVec2& BendingSt
 {
 	if (BendingConstraints)
 	{
-		BendingConstraints->SetProperties(BendingStiffness);
+		static_cast<Softs::FPBDSpringConstraints*>(BendingConstraints.Get())->SetProperties(BendingStiffness);
 	}
 	if (XBendingConstraints)
 	{
-		XBendingConstraints->SetProperties(BendingStiffness);
+		static_cast<Softs::FXPBDSpringConstraints*>(XBendingConstraints.Get())->SetProperties(BendingStiffness);
 	}
 	if (BendingElementConstraints)
 	{
@@ -695,27 +1203,27 @@ void FClothConstraints::SetAreaProperties(const Softs::FSolverVec2& AreaStiffnes
 {
 	if (AreaConstraints)
 	{
-		AreaConstraints->SetProperties(AreaStiffness);
+		static_cast<Softs::FPBDAxialSpringConstraints*>(AreaConstraints.Get())->SetProperties(AreaStiffness);
 	}
 	if (XAreaConstraints)
 	{
-		XAreaConstraints->SetProperties(AreaStiffness);
+		static_cast<Softs::FXPBDAxialSpringConstraints*>(XAreaConstraints.Get())->SetProperties(AreaStiffness);
 	}
 }
 
 void FClothConstraints::SetThinShellVolumeProperties(Softs::FSolverReal VolumeStiffness)
 {
-	if (ThinShellVolumeConstraints)
+	if (ThinShellVolumeConstraints_Deprecated)
 	{
-		ThinShellVolumeConstraints->SetProperties(VolumeStiffness);
+		ThinShellVolumeConstraints_Deprecated->SetProperties(VolumeStiffness);
 	}
 }
 
 void FClothConstraints::SetVolumeProperties(Softs::FSolverReal VolumeStiffness)
 {
-	if (VolumeConstraints)
+	if (VolumeConstraints_Deprecated)
 	{
-		VolumeConstraints->SetStiffness(VolumeStiffness);
+		VolumeConstraints_Deprecated->SetStiffness(VolumeStiffness);
 	}
 }
 
diff --git a/Engine/Plugins/Experimental/ChaosCloth/Source/Chaos/Private/ChaosCloth/ChaosClothVisualization.cpp b/Engine/Plugins/Experimental/ChaosCloth/Source/Chaos/Private/ChaosCloth/ChaosClothVisualization.cpp
--- a/Engine/Plugins/Experimental/ChaosCloth/Source/Chaos/Private/ChaosCloth/ChaosClothVisualization.cpp
+++ b/Engine/Plugins/Experimental/ChaosCloth/Source/Chaos/Private/ChaosCloth/ChaosClothVisualization.cpp
@@ -1198,7 +1198,7 @@ namespace Chaos
 
 			const TConstArrayView<Softs::FSolverVec3> Positions = Cloth->GetParticlePositions(Solver);
 
-			if (const Softs::FPBDSpringConstraints* const EdgeConstraints = ClothConstraints.GetEdgeConstraints().Get())
+			if (const Softs::FPBDEdgeSpringConstraints* const EdgeConstraints = ClothConstraints.GetEdgeSpringConstraints().Get())
 			{
 				DrawSpringConstraintColors(PDI, Positions, LocalSpaceLocation, EdgeConstraints);
 			}
@@ -1227,7 +1227,7 @@ namespace Chaos
 
 			const TConstArrayView<Softs::FSolverVec3> Positions = Cloth->GetParticlePositions(Solver);
 
-			if (const Softs::FPBDSpringConstraints* const BendingConstraints = ClothConstraints.GetBendingConstraints().Get())
+			if (const Softs::FPBDBendingSpringConstraints* const BendingConstraints = ClothConstraints.GetBendingSpringConstraints().Get())
 			{
 				DrawSpringConstraintColors(PDI, Positions, LocalSpaceLocation, BendingConstraints);
 			}
diff --git a/Engine/Plugins/Experimental/ChaosCloth/Source/Chaos/Private/ChaosCloth/ChaosClothingSimulation.cpp b/Engine/Plugins/Experimental/ChaosCloth/Source/Chaos/Private/ChaosCloth/ChaosClothingSimulation.cpp
--- a/Engine/Plugins/Experimental/ChaosCloth/Source/Chaos/Private/ChaosCloth/ChaosClothingSimulation.cpp
+++ b/Engine/Plugins/Experimental/ChaosCloth/Source/Chaos/Private/ChaosCloth/ChaosClothingSimulation.cpp
@@ -10,6 +10,7 @@
 #include "ChaosCloth/ChaosClothingSimulationSkeletalMesh.h"
 #include "ChaosCloth/ChaosClothingSimulationCloth.h"
 #include "ChaosCloth/ChaosClothingSimulationCollider.h"
+#include "ChaosCloth/ChaosClothingSimulationConfig.h"
 
 #include "PhysicsField/PhysicsFieldComponent.h"
 #include "PhysicsProxy/PerSolverFieldSystem.h"
@@ -258,7 +259,7 @@ namespace ChaosClothingSimulationDefault
 
 FClothingSimulation::FClothingSimulation()
 	: ClothSharedSimConfig(nullptr)
-	, bUseLocalSpaceSimulation(false)
+	, bUseLocalSpaceSimulation(true)
 	, bUseGravityOverride(false)
 	, GravityOverride(ChaosClothingSimulationDefault::Gravity)
 	, MaxDistancesMultipliers(ChaosClothingSimulationDefault::MaxDistancesMultipliers)
@@ -330,8 +331,6 @@ void FClothingSimulation::CreateActor(USkeletalMeshComponent* InOwnerComponent,
 	{
 		ClothSharedSimConfig = Asset->GetClothConfig<UChaosClothSharedSimConfig>();
 
-		UpdateSimulationFromSharedSimConfig();
-
 		// Must set the local space location prior to adding any mesh/cloth, as otherwise the start poses would be in the wrong local space
 		const FClothingSimulationContext* const Context = static_cast<const FClothingSimulationContext*>(InOwnerComponent->GetClothingSimulationContext());
 		check(Context);
@@ -352,12 +351,12 @@ void FClothingSimulation::CreateActor(USkeletalMeshComponent* InOwnerComponent,
 		return;
 	}
 
-	// Create mesh node
+	// Create mesh runtime simulation object
 	const int32 MeshIndex = Meshes.Emplace(MakeUnique<FClothingSimulationSkeletalMesh>(
 		Asset,
 		InOwnerComponent));
 
-	// Create collider node
+	// Create collider runtime simulation object
 	const int32 ColliderIndex = Colliders.Emplace(MakeUnique<FClothingSimulationCollider>(
 		Asset->PhysicsAsset,
 		&CastChecked<USkeletalMesh>(Asset->GetOuter())->GetRefSkeleton()));
@@ -365,62 +364,27 @@ void FClothingSimulation::CreateActor(USkeletalMeshComponent* InOwnerComponent,
 	// Set the external collision data to get updated at every frame
 	Colliders[ColliderIndex]->SetCollisionData(&ExternalCollisionData);
 
-	// Create cloth node
+	// Create cloth config runtime simulation object
+	const int32 ClothConfigIndex = Configs.Emplace(MakeUnique<FClothingSimulationConfig>());
+	Configs[ClothConfigIndex]->Initialize(ClothConfig, nullptr);
+
+	// Create cloth runtime simulation object
 	const int32 ClothIndex = Cloths.Emplace(MakeUnique<FClothingSimulationCloth>(
-		PRAGMA_DISABLE_DEPRECATION_WARNINGS  // TODO: CHAOS_IS_CLOTHINGSIMULATIONMESH_ABSTRACT
+		Configs[ClothConfigIndex].Get(),
+PRAGMA_DISABLE_DEPRECATION_WARNINGS  // TODO: CHAOS_IS_CLOTHINGSIMULATIONMESH_ABSTRACT
 		Meshes[MeshIndex].Get(),
-		PRAGMA_ENABLE_DEPRECATION_WARNINGS
+PRAGMA_ENABLE_DEPRECATION_WARNINGS
 		TArray<FClothingSimulationCollider*>({ Colliders[ColliderIndex].Get() }),
-		InSimDataIndex,
-		(FClothingSimulationCloth::EMassMode)ClothConfig->MassMode,
-		ClothConfig->GetMassValue(),
-		ClothConfig->MinPerParticleMass,
-		TVec2<FRealSingle>(ClothConfig->EdgeStiffnessWeighted.Low, ClothConfig->EdgeStiffnessWeighted.High),
-		TVec2<FRealSingle>(ClothConfig->BendingStiffnessWeighted.Low, ClothConfig->BendingStiffnessWeighted.High),
-		ClothConfig->BucklingRatio,
-		TVec2<FRealSingle>(ClothConfig->BucklingStiffnessWeighted.Low, ClothConfig->BucklingStiffnessWeighted.High),
-		ClothConfig->bUseBendingElements,
-		TVec2<FRealSingle>(ClothConfig->AreaStiffnessWeighted.Low, ClothConfig->AreaStiffnessWeighted.High),
-		ClothConfig->VolumeStiffness,
-		ClothConfig->bUseThinShellVolumeConstraints,
-		TVec2<FRealSingle>(ClothConfig->TetherStiffness.Low, ClothConfig->TetherStiffness.High),  // Animatable
-		TVec2<FRealSingle>(ClothConfig->TetherScale.Low, ClothConfig->TetherScale.High),  // Animatable
-		ClothConfig->bUseGeodesicDistance ? FClothingSimulationCloth::ETetherMode::Geodesic : FClothingSimulationCloth::ETetherMode::Euclidean,
-		/*MaxDistancesMultiplier =*/ 1.f,  // Animatable
-		TVec2<FRealSingle>(ClothConfig->AnimDriveStiffness.Low, ClothConfig->AnimDriveStiffness.High),  // Animatable
-		TVec2<FRealSingle>(ClothConfig->AnimDriveDamping.Low, ClothConfig->AnimDriveDamping.High),  // Animatable
-		ClothConfig->ShapeTargetStiffness,  // TODO: This is now deprecated
-		false, // bUseXPBDEdgeSprings
-		false, // bUseXPBDBendingElements
-		false, // bUseXPBDAreaSprings
-		ClothConfig->GravityScale,
-		ClothConfig->bUseGravityOverride,
-		ClothConfig->Gravity,
-		ClothConfig->LinearVelocityScale,
-		ClothConfig->AngularVelocityScale,
-		ClothConfig->FictitiousAngularScale,
-		TVec2<FRealSingle>(ClothConfig->Drag.Low, ClothConfig->Drag.High),  // Animatable
-		TVec2<FRealSingle>(ClothConfig->Lift.Low, ClothConfig->Lift.High),  // Animatable
-		ClothConfig->bUsePointBasedWindModel,
-		TVec2<FRealSingle>(ClothConfig->Pressure.Low, ClothConfig->Pressure.High),  // Animatable
-		ClothConfig->DampingCoefficient,
-		ClothConfig->LocalDampingCoefficient,
-		ClothConfig->CollisionThickness,
-		ClothConfig->FrictionCoefficient,
-		ClothConfig->bUseCCD,
-		ClothConfig->bUseSelfCollisions,
-		ClothConfig->SelfCollisionThickness,
-		ClothConfig->SelfCollisionFriction,
-		ClothConfig->bUseSelfIntersections,
-		ClothConfig->bUseLegacyBackstop,
-		/*bUseLODIndexOverride =*/ false,
-		/*LODIndexOverride =*/ INDEX_NONE,
-		TVec2<FRealSingle>(0.f), // EdgeDampingRatio
-		TVec2<FRealSingle>(0.f))); // BendDampingRatio
+		InSimDataIndex));
 
 	// Add cloth to solver
 	Solver->AddCloth(Cloths[ClothIndex].Get());
 
+	// Create solver config runtime simulation object
+	const int32 SolverConfigIndex = Configs.Emplace(MakeUnique<FClothingSimulationConfig>());
+	Configs[SolverConfigIndex]->Initialize(nullptr, ClothSharedSimConfig);
+	Solver->SetConfig(Configs[SolverConfigIndex].Get());
+
 	// Update stats
 	UpdateStats(Cloths[ClothIndex].Get());
 
@@ -461,21 +425,6 @@ void FClothingSimulation::UpdateStats(const FClothingSimulationCloth* Cloth)
 	NumDynamicParticles += Cloth->GetNumActiveDynamicParticles();
 }
 
-void FClothingSimulation::UpdateSimulationFromSharedSimConfig()
-{
-	check(Solver);
-	if (ClothSharedSimConfig) // ClothSharedSimConfig will be a null pointer if all cloth instances are disabled in which case we will use default Evolution parameters
-	{
-		// Update local space simulation switch
-		bUseLocalSpaceSimulation = ClothSharedSimConfig->bUseLocalSpaceSimulation;
-
-		// Set common simulation parameters
-		Solver->SetNumSubsteps(ClothSharedSimConfig->SubdivisionCount);
-		Solver->SetNumIterations(ClothSharedSimConfig->IterationCount);
-		Solver->SetMaxNumIterations(ClothSharedSimConfig->MaxIterationCount);
-	}
-}
-
 void FClothingSimulation::SetNumIterations(int32 InNumIterations)
 {
 	Solver->SetNumIterations(InNumIterations);
@@ -791,8 +740,6 @@ void FClothingSimulation::GetCollisions(FClothCollisionData& OutCollisions, bool
 
 void FClothingSimulation::RefreshClothConfig(const IClothingSimulationContext* InContext)
 {
-	UpdateSimulationFromSharedSimConfig();
-
 	// Update new space location
 	const FClothingSimulationContext* const Context = static_cast<const FClothingSimulationContext*>(InContext);
 	static const bool bReset = true;
@@ -813,55 +760,19 @@ void FClothingSimulation::RefreshClothConfig(const IClothingSimulationContext* I
 		const uint32 GroupId = Cloth->GetGroupId();
 		const UChaosClothConfig* const ClothConfig = Mesh->GetAsset()->GetClothConfig<UChaosClothConfig>();
 
+		// Update cloth config runtime simulation object
+		FClothingSimulationConfig* const Config = Cloth->GetConfig();
+		check(Config);
+		Config->Initialize(ClothConfig, nullptr);
+
+		// Recreate cloth runtime simulation object
 		Cloth = MakeUnique<FClothingSimulationCloth>(
+			Config,
+			PRAGMA_DISABLE_DEPRECATION_WARNINGS  // TODO: CHAOS_IS_CLOTHINGSIMULATIONMESH_ABSTRACT
 			Mesh,
+			PRAGMA_ENABLE_DEPRECATION_WARNINGS
 			MoveTemp(ClothColliders),
-			GroupId,
-			(FClothingSimulationCloth::EMassMode)ClothConfig->MassMode,
-			ClothConfig->GetMassValue(),
-			ClothConfig->MinPerParticleMass,
-			TVec2<FRealSingle>(ClothConfig->EdgeStiffnessWeighted.Low, ClothConfig->EdgeStiffnessWeighted.High),
-			TVec2<FRealSingle>(ClothConfig->BendingStiffnessWeighted.Low, ClothConfig->BendingStiffnessWeighted.High),
-			ClothConfig->BucklingRatio,
-			TVec2<FRealSingle>(ClothConfig->BucklingStiffnessWeighted.Low, ClothConfig->BucklingStiffnessWeighted.High),
-			ClothConfig->bUseBendingElements,
-			TVec2<FRealSingle>(ClothConfig->AreaStiffnessWeighted.Low, ClothConfig->AreaStiffnessWeighted.High),
-			ClothConfig->VolumeStiffness,
-			ClothConfig->bUseThinShellVolumeConstraints,
-			TVec2<FRealSingle>(ClothConfig->TetherStiffness.Low, ClothConfig->TetherStiffness.High),  // Animatable
-			TVec2<FRealSingle>(ClothConfig->TetherScale.Low, ClothConfig->TetherScale.High),  // Animatable
-			ClothConfig->bUseGeodesicDistance ? FClothingSimulationCloth::ETetherMode::Geodesic : FClothingSimulationCloth::ETetherMode::Euclidean,
-			/*MaxDistancesMultiplier =*/ 1.f,  // Animatable
-			TVec2<FRealSingle>(ClothConfig->AnimDriveStiffness.Low, ClothConfig->AnimDriveStiffness.High),  // Animatable
-			TVec2<FRealSingle>(ClothConfig->AnimDriveDamping.Low, ClothConfig->AnimDriveDamping.High),  // Animatable
-			ClothConfig->ShapeTargetStiffness,
-			false, // bUseXPBDEdgeSprings
-			false, // bUseXPBDBendingElements
-			false, // bUseXPBDAreaSprings
-			ClothConfig->GravityScale,
-			ClothConfig->bUseGravityOverride,
-			ClothConfig->Gravity,
-			ClothConfig->LinearVelocityScale,
-			ClothConfig->AngularVelocityScale,
-			ClothConfig->FictitiousAngularScale,
-			TVec2<FRealSingle>(ClothConfig->Drag.Low, ClothConfig->Drag.High),  // Animatable
-			TVec2<FRealSingle>(ClothConfig->Lift.Low, ClothConfig->Lift.High),  // Animatable
-			ClothConfig->bUsePointBasedWindModel,
-			TVec2<FRealSingle>(ClothConfig->Pressure.Low, ClothConfig->Pressure.High),  // Animatable
-			ClothConfig->DampingCoefficient,
-			ClothConfig->LocalDampingCoefficient,
-			ClothConfig->CollisionThickness,
-			ClothConfig->FrictionCoefficient,
-			ClothConfig->bUseCCD,
-			ClothConfig->bUseSelfCollisions,
-			ClothConfig->SelfCollisionThickness,
-			ClothConfig->SelfCollisionFriction,
-			ClothConfig->bUseSelfIntersections,
-			ClothConfig->bUseLegacyBackstop,
-			/*bUseLODIndexOverride =*/ false,
-			/*LODIndexOverride =*/ INDEX_NONE,
-			TVec2<FRealSingle>(0.f), // EdgeDampingRatio
-			TVec2<FRealSingle>(0.f)); // BendDampingRatio
+			GroupId);
 
 		// Re-add cloth to the solver
 		Solver->AddCloth(Cloth.Get());
@@ -869,6 +780,12 @@ void FClothingSimulation::RefreshClothConfig(const IClothingSimulationContext* I
 		// Update stats
 		UpdateStats(Cloth.Get());
 	}
+
+	// Update solver config runtime simulation object
+	FClothingSimulationConfig* const SolverConfig = Solver->GetConfig();
+	check(SolverConfig);
+	SolverConfig->Initialize(nullptr, ClothSharedSimConfig);
+
 	UE_LOG(LogChaosCloth, VeryVerbose, TEXT("RefreshClothConfig, all constraints and self-collisions have been updated for all clothing assets and LODs."));
 }
 
diff --git a/Engine/Plugins/Experimental/ChaosCloth/Source/Chaos/Private/ChaosCloth/ChaosClothingSimulationCloth.cpp b/Engine/Plugins/Experimental/ChaosCloth/Source/Chaos/Private/ChaosCloth/ChaosClothingSimulationCloth.cpp
--- a/Engine/Plugins/Experimental/ChaosCloth/Source/Chaos/Private/ChaosCloth/ChaosClothingSimulationCloth.cpp
+++ b/Engine/Plugins/Experimental/ChaosCloth/Source/Chaos/Private/ChaosCloth/ChaosClothingSimulationCloth.cpp
@@ -3,15 +3,36 @@
 #include "ChaosCloth/ChaosClothingSimulationSolver.h"
 #include "ChaosCloth/ChaosClothingSimulationMesh.h"
 #include "ChaosCloth/ChaosClothingSimulationCollider.h"
+#include "ChaosCloth/ChaosClothingSimulationConfig.h"
 #include "ChaosCloth/ChaosWeightMapTarget.h"
 #include "ChaosCloth/ChaosClothPrivate.h"
+#include "Chaos/PropertyCollectionAdapter.h"
 #include "Containers/ArrayView.h"
+#include "GeometryCollection/ManagedArrayCollection.h"
 #include "HAL/IConsoleManager.h"
 #include "ClothingSimulation.h"
 
 namespace Chaos
 {
 
+namespace ClothingSimulationClothDefault
+{
+	constexpr int32 MassMode = (int32)FClothingSimulationCloth::EMassMode::Density;
+	constexpr int32 TetherMode = (int32)FClothingSimulationCloth::ETetherMode::Geodesic;
+	constexpr float MassValue = 0.35f;
+	constexpr float MinPerParticleMass = 0.0001f;
+	constexpr float CollisionThickness = 1.0f;
+	constexpr float FrictionCoefficient = 0.8f;
+	constexpr float DampingCoefficient = 0.01f;
+	constexpr float Drag = 0.035f;
+	constexpr float Lift = 0.035f;
+	constexpr float AirDensity = 1.225e-6f;
+	constexpr float GravityScale = 1.f;
+	constexpr float GravityZOverride = -980.665f;
+	constexpr float VelocityScale = 0.75f;
+	constexpr float FictitiousAngularScale = 1.f;
+}
+
 namespace ClothingSimulationClothConsoleVariables
 {
 	TAutoConsoleVariable<bool> CVarLegacyDisablesAccurateWind(TEXT("p.ChaosCloth.LegacyDisablesAccurateWind"), true, TEXT("Whether using the Legacy wind model switches off the accurate wind model, or adds up to it"));
@@ -87,9 +108,12 @@ void FClothingSimulationCloth::FLODData::Add(FClothingSimulationSolver* Solver,
 	// Add particles
 	Offset = Solver->AddParticles(NumParticles, Cloth->GroupId);  // TODO: Have a per solver map of offset
 
-	// Update source mesh for this LOD, this is required prior to reset the start pose
 PRAGMA_DISABLE_DEPRECATION_WARNINGS  // TODO: CHAOS_IS_CLOTHINGSIMULATIONMESH_ABSTRACT
+	// Update source mesh for this LOD, this is required prior to reset the start pose
 	Cloth->Mesh->Update(Solver, INDEX_NONE, InLODIndex, 0, Offset);
+	
+	// Retrieve the component's scale
+	const Softs::FSolverReal MeshScale = Cloth->Mesh->GetScale();
 PRAGMA_ENABLE_DEPRECATION_WARNINGS
 
 	// Reset the particles start pose before setting up mass and constraints
@@ -109,7 +133,7 @@ PRAGMA_ENABLE_DEPRECATION_WARNINGS
 			 static_cast<int32>(Offset + Indices[Index + 2]) });
 	}
 
-	TriangleMesh.Init(MoveTemp(Elements));
+	TriangleMesh.Init(MoveTemp(Elements), Offset, Offset + NumParticles - 1);  // Init with the Offset to avoid discrepancies since the Triangle Mesh only relies on the used indices
 	TriangleMesh.GetPointToTriangleMap(); // Builds map for later use by GetPointNormals(), and the velocity fields
 
 	// Initialize the normals, in case the sim data is queried before the simulation steps
@@ -124,190 +148,37 @@ PRAGMA_ENABLE_DEPRECATION_WARNINGS
 			return MaxDistances.IsValidIndex(Index) && MaxDistances[Index] < KinematicDistanceThreshold;
 		};
 
-	switch (Cloth->MassMode)
+	check(Cloth->Config);
+	const Softs::FPropertyCollectionConstAdapter& ConfigProperties = Cloth->Config->GetProperties();
+
+	const int32 MassMode = ConfigProperties.GetValue<int32>(TEXT("MassMode"), ClothingSimulationClothDefault::MassMode);
+	const FRealSingle MassValue = (FRealSingle)ConfigProperties.GetValue<float>(TEXT("MassValue"), ClothingSimulationClothDefault::MassValue);
+	const FRealSingle MinPerParticleMass = (FRealSingle)ConfigProperties.GetValue<float>(TEXT("MinPerParticleMass"), ClothingSimulationClothDefault::MinPerParticleMass);
+
+	switch (MassMode)
 	{
 	default:
 		check(false);
 	case EMassMode::UniformMass:
-		Solver->SetParticleMassUniform(Offset, Cloth->MassValue, Cloth->MinPerParticleMass, TriangleMesh, KinematicPredicate);
+		Solver->SetParticleMassUniform(Offset, MassValue, MinPerParticleMass, TriangleMesh, KinematicPredicate);
 		break;
 	case EMassMode::TotalMass:
-		Solver->SetParticleMassFromTotalMass(Offset, Cloth->MassValue, Cloth->MinPerParticleMass, TriangleMesh, KinematicPredicate);
+		Solver->SetParticleMassFromTotalMass(Offset, MassValue, MinPerParticleMass, TriangleMesh, KinematicPredicate);
 		break;
 	case EMassMode::Density:
-		Solver->SetParticleMassFromDensity(Offset, Cloth->MassValue, Cloth->MinPerParticleMass, TriangleMesh, KinematicPredicate);
+		Solver->SetParticleMassFromDensity(Offset, MassValue, MinPerParticleMass, TriangleMesh, KinematicPredicate);
 		break;
 	}
-	const TConstArrayView<Softs::FSolverReal> InvMasses(Solver->GetParticleInvMasses(Offset), NumParticles);
 
 	// Setup solver constraints
 	FClothConstraints& ClothConstraints = Solver->GetClothConstraints(Offset);
-	const TArray<TVec3<int32>>& SurfaceElements = TriangleMesh.GetSurfaceElements();
-
-	// Self collisions
-	if (Cloth->bUseSelfCollisions)
-	{
-		static const int32 DisabledCollisionElementsN = 5;  // TODO: Make this a parameter?
-		TSet<TVec2<int32>> DisabledCollisionElements;  // TODO: Is this needed? Turn this into a bit array?
-	
-		const int32 Range = Offset + NumParticles;
-		for (int32 Index = Offset; Index < Range; ++Index)
-		{
-			const TSet<int32> Neighbors = TriangleMesh.GetNRing(Index, DisabledCollisionElementsN);
-			for (int32 Element : Neighbors)
-			{
-				check(Index != Element);
-				DisabledCollisionElements.Emplace(TVec2<int32>(Index, Element));
-				DisabledCollisionElements.Emplace(TVec2<int32>(Element, Index));
-			}
-		}
-		ClothConstraints.SetSelfCollisionConstraints(TriangleMesh, MoveTemp(DisabledCollisionElements), (Softs::FSolverReal)Cloth->SelfCollisionThickness, (Softs::FSolverReal)Cloth->SelfCollisionFrictionCoefficient, Cloth->bUseSelfIntersections, Cloth->bUseSelfIntersections);
-	}
-
-	// Edge constraints
-	const TConstArrayView<FRealSingle>& EdgeStiffnessMultipliers = WeightMaps[(int32)EChaosWeightMapTarget::EdgeStiffness];
-	if (Cloth->EdgeStiffness[0] > (FRealSingle)0. || (Cloth->EdgeStiffness[1] > (FRealSingle)0. && EdgeStiffnessMultipliers.Num() == NumParticles))
-	{
-		if (Cloth->bUseXPBDEdgeConstraints)
-		{
-			const TConstArrayView<FRealSingle>& EdgeDampingRatioMultipliers = TConstArrayView<FRealSingle>(); // XPBD is not currently exposed to the UI.
-			ClothConstraints.SetXPBDEdgeConstraints(SurfaceElements, EdgeStiffnessMultipliers, EdgeDampingRatioMultipliers);
-		}
-		else
-		{
-			ClothConstraints.SetEdgeConstraints(SurfaceElements, EdgeStiffnessMultipliers, Cloth->bUseXPBDEdgeConstraints);
-		}
-	}
 
-	// Bending constraints
-	const TConstArrayView<FRealSingle>& BendingStiffnessMultipliers = WeightMaps[(int32)EChaosWeightMapTarget::BendingStiffness];
-	const TConstArrayView<FRealSingle>& BucklingStiffnessMultipliers = WeightMaps[(int32)EChaosWeightMapTarget::BucklingStiffness];
-	if (Cloth->BendingStiffness[0] > (FRealSingle)0. || (Cloth->BendingStiffness[1] > (FRealSingle)0. && BendingStiffnessMultipliers.Num() == NumParticles) ||
-		(Cloth->bUseBendingElements && (Cloth->BucklingStiffness[0] > (FRealSingle)0. || (Cloth->BucklingStiffness[1] > (FRealSingle)0. && BucklingStiffnessMultipliers.Num() == NumParticles))))
-	{
-		if (Cloth->bUseBendingElements)
-		{
-			TArray<Chaos::TVec4<int32>> BendingElements = TriangleMesh.GetUniqueAdjacentElements();
-			if (Cloth->bUseXPBDBendingConstraints)
-			{
-				const TConstArrayView<FRealSingle>& BendingDampingRatioMultipliers = TConstArrayView<FRealSingle>(); // XPBD is not currently exposed to the UI.
-				ClothConstraints.SetXPBDBendingConstraints(MoveTemp(BendingElements), BendingStiffnessMultipliers, BucklingStiffnessMultipliers, BendingDampingRatioMultipliers);
-			}
-			else
-			{
-				ClothConstraints.SetBendingConstraints(MoveTemp(BendingElements), BendingStiffnessMultipliers, BucklingStiffnessMultipliers, Cloth->bUseXPBDBendingConstraints);
-			}
-		}
-		else
-		{
-			const TArray<Chaos::TVec2<int32>> Edges = TriangleMesh.GetUniqueAdjacentPoints();
-			ClothConstraints.SetBendingConstraints(Edges, BendingStiffnessMultipliers, Cloth->bUseXPBDBendingConstraints);
-		}
-	}
-
-	// Area constraints
-	const TConstArrayView<FRealSingle>& AreaStiffnessMultipliers = WeightMaps[(int32)EChaosWeightMapTarget::AreaStiffness];
-	if (Cloth->AreaStiffness[0] > (FRealSingle)0. || (Cloth->AreaStiffness[1] > (FRealSingle)0. && AreaStiffnessMultipliers.Num() == NumParticles))
-	{
-		ClothConstraints.SetAreaConstraints(SurfaceElements, AreaStiffnessMultipliers, Cloth->bUseXPBDAreaConstraints);
-	}
-
-	// Volume constraints
-	if (Cloth->VolumeStiffness > (FRealSingle)0.)
-	{
-		if (Cloth->bUseThinShellVolumeConstraints)
-		{
-			TArray<Chaos::TVec2<int32>> BendingConstraints = TriangleMesh.GetUniqueAdjacentPoints();
-			TArray<Chaos::TVec2<int32>> DoubleBendingConstraints;
-			{
-				TMap<int32, TArray<int32>> BendingHash;
-				for (int32 i = 0; i < BendingConstraints.Num(); ++i)
-				{
-					BendingHash.FindOrAdd(BendingConstraints[i][0]).Add(BendingConstraints[i][1]);
-					BendingHash.FindOrAdd(BendingConstraints[i][1]).Add(BendingConstraints[i][0]);
-				}
-				TSet<Chaos::TVec2<int32>> Visited;
-				for (auto Elem : BendingHash)
-				{
-					for (int32 i = 0; i < Elem.Value.Num(); ++i)
-					{
-						for (int32 j = i + 1; j < Elem.Value.Num(); ++j)
-						{
-							if (Elem.Value[i] == Elem.Value[j])
-								continue;
-							auto NewElem = Chaos::TVec2<int32>(Elem.Value[i], Elem.Value[j]);
-							if (!Visited.Contains(NewElem))
-							{
-								DoubleBendingConstraints.Add(NewElem);
-								Visited.Add(NewElem);
-								Visited.Add(Chaos::TVec2<int32>(Elem.Value[j], Elem.Value[i]));
-							}
-						}
-					}
-				}
-			}
-
-			ClothConstraints.SetVolumeConstraints(DoubleBendingConstraints, (Softs::FSolverReal)Cloth->VolumeStiffness);
-		}
-		else
-		{
-			TArray<Chaos::TVec3<int32>> SurfaceConstraints = SurfaceElements;
-			ClothConstraints.SetVolumeConstraints(MoveTemp(SurfaceConstraints), (Softs::FSolverReal)Cloth->VolumeStiffness);
-		}
-	}
-
-	// Long range constraints
-	const TConstArrayView<FRealSingle>& TetherStiffnessMultipliers = WeightMaps[(int32)EChaosWeightMapTarget::TetherStiffness];
-	if (Cloth->TetherStiffness[0] > (FRealSingle)0. || (Cloth->TetherStiffness[1] > (FRealSingle)0. && TetherStiffnessMultipliers.Num() == NumParticles))
-	{
-		const TConstArrayView<FRealSingle>& TetherScaleMultipliers = WeightMaps[(int32)EChaosWeightMapTarget::TetherScale];
-
-PRAGMA_DISABLE_DEPRECATION_WARNINGS  // TODO: CHAOS_IS_CLOTHINGSIMULATIONMESH_ABSTRACT
-		const Softs::FSolverReal MeshScale = Cloth->Mesh->GetScale();
-PRAGMA_ENABLE_DEPRECATION_WARNINGS
-
-		ClothConstraints.SetLongRangeConstraints(
-			Tethers,
-			TetherStiffnessMultipliers,
-			TetherScaleMultipliers,
-			Softs::FSolverVec2((Softs::FSolverReal)Cloth->TetherScale[0], (Softs::FSolverReal)Cloth->TetherScale[1]) * MeshScale);
-	}
-
-	// Max distances
-	if (MaxDistances.Num() == NumParticles)
-	{
-		ClothConstraints.SetMaximumDistanceConstraints(MaxDistances);
-	}
-
-	// Backstop Constraints
-	const TConstArrayView<FRealSingle>& BackstopDistances = WeightMaps[(int32)EChaosWeightMapTarget::BackstopDistance];
-	const TConstArrayView<FRealSingle>& BackstopRadiuses = WeightMaps[(int32)EChaosWeightMapTarget::BackstopRadius];
-	if (BackstopRadiuses.Num() == NumParticles && BackstopDistances.Num() == NumParticles)
-	{
-		ClothConstraints.SetBackstopConstraints(BackstopDistances, BackstopRadiuses, Cloth->bUseLegacyBackstop);
-	}
-
-	// Animation Drive Constraints
-	const TConstArrayView<FRealSingle>& AnimDriveStiffnessMultipliers = WeightMaps[(int32)EChaosWeightMapTarget::AnimDriveStiffness];
-	if (Cloth->AnimDriveStiffness[0] > (FRealSingle)0. || (Cloth->AnimDriveStiffness[1] > (FRealSingle)0. && AnimDriveStiffnessMultipliers.Num() == NumParticles))
-	{
-		const TConstArrayView<FRealSingle>& AnimDriveDampingMultipliers = WeightMaps[(int32)EChaosWeightMapTarget::AnimDriveDamping];
-		ClothConstraints.SetAnimDriveConstraints(AnimDriveStiffnessMultipliers, AnimDriveDampingMultipliers);
-	}
-
-	// Shape target constraint
-	if (Cloth->ShapeTargetStiffness > (FRealSingle)0.)
-	{
-		ClothConstraints.SetShapeTargetConstraints((Softs::FSolverReal)Cloth->ShapeTargetStiffness);
-	}
-
-	// Commit rules to solver
-	ClothConstraints.CreateRules();
-
-	// Disable constraints by default
-	ClothConstraints.Enable(false);
+	// Create constraints
+	const bool bEnabled = false;  // Set constraint disabled by default
+	ClothConstraints.AddRules(ConfigProperties, TriangleMesh, WeightMaps, Tethers, MeshScale, bEnabled);
 
 	// Update LOD stats
+	const TConstArrayView<Softs::FSolverReal> InvMasses(Solver->GetParticleInvMasses(Offset), NumParticles);
 	NumKinenamicParticles = 0;
 	NumDynammicParticles = 0;
 	for (int32 Index = 0; Index < NumParticles; ++Index)
@@ -331,33 +202,20 @@ void FClothingSimulationCloth::FLODData::Remove(FClothingSimulationSolver* Solve
 void FClothingSimulationCloth::FLODData::Update(FClothingSimulationSolver* Solver, FClothingSimulationCloth* Cloth)
 {
 	check(Solver);
+	check(Cloth);
+
 	const int32 Offset = SolverData.FindChecked(Solver).Offset;
 	check(Offset != INDEX_NONE);
 
+	// Update the animatable constraint parameters
+	FClothConstraints& ClothConstraints = Solver->GetClothConstraints(Offset);
+
+	check(Cloth->Config);
 PRAGMA_DISABLE_DEPRECATION_WARNINGS  // TODO: CHAOS_IS_CLOTHINGSIMULATIONMESH_ABSTRACT
 	const Softs::FSolverReal MeshScale = Cloth->Mesh->GetScale();
 PRAGMA_ENABLE_DEPRECATION_WARNINGS
-
-	// Update the animatable constraint parameters
-	FClothConstraints& ClothConstraints = Solver->GetClothConstraints(Offset);
-	ClothConstraints.SetMaximumDistanceProperties((Softs::FSolverReal)Cloth->MaxDistancesMultiplier * MeshScale);
-	ClothConstraints.SetEdgeProperties(Softs::FSolverVec2((Softs::FSolverReal)Cloth->EdgeStiffness[0], (Softs::FSolverReal)Cloth->EdgeStiffness[1]), 
-		Softs::FSolverVec2((Softs::FSolverReal)Cloth->EdgeDampingRatio[0], (Softs::FSolverReal)Cloth->EdgeDampingRatio[1]));
-	ClothConstraints.SetBendingProperties(Softs::FSolverVec2((Softs::FSolverReal)Cloth->BendingStiffness[0], (Softs::FSolverReal)Cloth->BendingStiffness[1]), 
-		(Softs::FSolverReal)Cloth->BucklingRatio, Softs::FSolverVec2((Softs::FSolverReal)Cloth->BucklingStiffness[0], (Softs::FSolverReal)Cloth->BucklingStiffness[1]), 
-		Softs::FSolverVec2((Softs::FSolverReal)Cloth->BendingDampingRatio[0], (Softs::FSolverReal)Cloth->BendingDampingRatio[1]));
-	ClothConstraints.SetAreaProperties(Softs::FSolverVec2((FSolverReal)Cloth->AreaStiffness[0], (Softs::FSolverReal)Cloth->AreaStiffness[1]));
-	ClothConstraints.SetLongRangeAttachmentProperties(
-		Softs::FSolverVec2((Softs::FSolverReal)Cloth->TetherStiffness[0], (Softs::FSolverReal)Cloth->TetherStiffness[1]),
-		Softs::FSolverVec2((Softs::FSolverReal)Cloth->TetherScale[0], (Softs::FSolverReal)Cloth->TetherScale[1]),
-		MeshScale);
-	ClothConstraints.SetSelfCollisionProperties((Softs::FSolverReal)Cloth->SelfCollisionThickness, (Softs::FSolverReal)Cloth->SelfCollisionFrictionCoefficient, Cloth->bUseSelfIntersections, Cloth->bUseSelfIntersections);
-	ClothConstraints.SetAnimDriveProperties(
-		Softs::FSolverVec2((Softs::FSolverReal)Cloth->AnimDriveStiffness[0], (Softs::FSolverReal)Cloth->AnimDriveStiffness[1]),
-		Softs::FSolverVec2((Softs::FSolverReal)Cloth->AnimDriveDamping[0], (Softs::FSolverReal)Cloth->AnimDriveDamping[1]));
-	ClothConstraints.SetThinShellVolumeProperties((Softs::FSolverReal)Cloth->VolumeStiffness);
-	ClothConstraints.SetVolumeProperties((Softs::FSolverReal)Cloth->VolumeStiffness);
-	ClothConstraints.SetBackstopProperties(Cloth->bEnableBackstop, MeshScale);
+	const Softs::FSolverReal MaxDistancesScale = (Softs::FSolverReal)Cloth->MaxDistancesMultiplier;
+	ClothConstraints.Update(Cloth->Config->GetProperties(), MeshScale, MaxDistancesScale);
 }
 
 void FClothingSimulationCloth::FLODData::Enable(FClothingSimulationSolver* Solver, bool bEnable) const
@@ -397,6 +255,20 @@ void FClothingSimulationCloth::FLODData::UpdateNormals(FClothingSimulationSolver
 	TriangleMesh.GetPointNormals(Normals, TConstArrayView<Softs::FSolverVec3>(FaceNormals), /*bUseGlobalArray =*/ false);
 }
 
+FClothingSimulationCloth::FClothingSimulationCloth(
+	FClothingSimulationConfig* InConfig,
+PRAGMA_DISABLE_DEPRECATION_WARNINGS  // TODO: CHAOS_IS_CLOTHINGSIMULATIONMESH_ABSTRACT
+	FClothingSimulationMesh* InMesh,
+PRAGMA_ENABLE_DEPRECATION_WARNINGS
+	TArray<FClothingSimulationCollider*>&& InColliders,
+	uint32 InGroupId)
+	: GroupId(InGroupId)
+{
+	SetConfig(InConfig);
+	SetMesh(InMesh);
+	SetColliders(MoveTemp(InColliders));
+}
+
 FClothingSimulationCloth::FClothingSimulationCloth(
 PRAGMA_DISABLE_DEPRECATION_WARNINGS  // TODO: CHAOS_IS_CLOTHINGSIMULATIONMESH_ABSTRACT
 	FClothingSimulationMesh* InMesh,
@@ -412,27 +284,27 @@ PRAGMA_ENABLE_DEPRECATION_WARNINGS
 	const TVec2<FRealSingle>& InBucklingStiffness,
 	bool bInUseBendingElements,
 	const TVec2<FRealSingle>& InAreaStiffness,
-	FRealSingle InVolumeStiffness,
-	bool bInUseThinShellVolumeConstraints,
+	FRealSingle /*InVolumeStiffness*/,  // Deprecated
+	bool /*bInUseThinShellVolumeConstraints*/,  // Deprecated
 	const TVec2<FRealSingle>& InTetherStiffness,
 	const TVec2<FRealSingle>& InTetherScale,
 	ETetherMode InTetherMode,
 	FRealSingle InMaxDistancesMultiplier,
 	const TVec2<FRealSingle>& InAnimDriveStiffness,
 	const TVec2<FRealSingle>& InAnimDriveDamping,
-	FRealSingle InShapeTargetStiffness,
+	FRealSingle /*InShapeTargetStiffness*/,  // Deprecated
 	bool bInUseXPBDEdgeConstraints,
 	bool bInUseXPBDBendingConstraints,
-	bool bInUseXPBDAreaConstraints,
+	bool /*bInUseXPBDAreaConstraints*/,  // Deprecated
 	FRealSingle InGravityScale,
-	bool bInIsGravityOverridden,
+	bool bInUseGravityOverride,
 	const TVec3<FRealSingle>& InGravityOverride,
 	const TVec3<FRealSingle>& InLinearVelocityScale,
 	FRealSingle InAngularVelocityScale,
 	FRealSingle InFictitiousAngularScale,
 	const TVec2<FRealSingle>& InDrag,
 	const TVec2<FRealSingle>& InLift,
-	bool bInUseLegacyWind,
+	bool bInUsePointBasedWindModel,
 	const TVec2<FRealSingle>& InPressure,
 	FRealSingle InDampingCoefficient,
 	FRealSingle InLocalDampingCoefficient,
@@ -451,52 +323,7 @@ PRAGMA_ENABLE_DEPRECATION_WARNINGS
 	: Mesh(nullptr)
 	, Colliders()
 	, GroupId(InGroupId)
-	, MassMode(InMassMode)
-	, MassValue(InMassValue)
-	, MinPerParticleMass(InMinPerParticleMass)
-	, EdgeStiffness(InEdgeStiffness)
-	, EdgeDampingRatio(InEdgeDampingRatio)
-	, BendingStiffness(InBendingStiffness)
-	, BendingDampingRatio(InBendingDampingRatio)
-	, BucklingRatio(InBucklingRatio)
-	, BucklingStiffness(InBucklingStiffness)
-	, bUseBendingElements(bInUseBendingElements)
-	, AreaStiffness(InAreaStiffness)
-	, VolumeStiffness(InVolumeStiffness)
-	, bUseThinShellVolumeConstraints(bInUseThinShellVolumeConstraints)
-	, TetherStiffness(InTetherStiffness)
-	, TetherScale(InTetherScale)
-	, TetherMode(InTetherMode)
-	, MaxDistancesMultiplier(InMaxDistancesMultiplier)
-	, AnimDriveStiffness(InAnimDriveStiffness)
-	, AnimDriveDamping(InAnimDriveDamping)
-	, ShapeTargetStiffness(InShapeTargetStiffness)
-	, bUseXPBDEdgeConstraints(bInUseXPBDEdgeConstraints)
-	, bUseXPBDBendingConstraints(bInUseXPBDBendingConstraints)
-	, bUseXPBDAreaConstraints(bInUseXPBDAreaConstraints)
-	, GravityScale(InGravityScale)
-	, bIsGravityOverridden(bInIsGravityOverridden)
-	, GravityOverride(InGravityOverride)
-	, LinearVelocityScale(InLinearVelocityScale)
-	, AngularVelocityScale(InAngularVelocityScale)
-	, FictitiousAngularScale(InFictitiousAngularScale)
-	, Drag(InDrag)
-	, Lift(InLift)
-	, WindVelocity((FReal)0., (FReal)0., (FReal)0.)  // Set by clothing interactor
-	, AirDensity(1.225e-6f)  // Set by clothing interactor
-	, bUseLegacyWind(bInUseLegacyWind)
-	, Pressure(InPressure)
-	, DampingCoefficient(InDampingCoefficient)
-	, LocalDampingCoefficient(InLocalDampingCoefficient)
-	, CollisionThickness(InCollisionThickness)
-	, FrictionCoefficient(InFrictionCoefficient)
-	, bUseCCD(bInUseCCD)
-	, bUseSelfCollisions(bInUseSelfCollisions)
-	, SelfCollisionThickness(InSelfCollisionThickness)
-	, SelfCollisionFrictionCoefficient(InSelfCollisionFrictionCoefficient)
-	, bUseSelfIntersections(bInUseSelfIntersections)
-	, bEnableBackstop(true)  // Set by clothing interactor
-	, bUseLegacyBackstop(bInUseLegacyBackstop)
+	, PropertyCollection(MakeShared<FManagedArrayCollection>())
 	, bUseLODIndexOverride(bInUseLODIndexOverride)
 	, LODIndexOverride(InLODIndexOverride)
 	, bNeedsReset(false)
@@ -504,12 +331,174 @@ PRAGMA_ENABLE_DEPRECATION_WARNINGS
 	, NumActiveKinematicParticles(0)
 	, NumActiveDynamicParticles(0)
 {
+	// Turn parameters into a config
+	PropertyCollection = MakeShared<FManagedArrayCollection>();
+	Softs::FPropertyCollectionMutableAdapter Properties(PropertyCollection);
+
+	constexpr bool bEnable = true;
+	constexpr bool bAnimatable = true;
+	
+	// Mass
+	{
+		Properties.AddValue(TEXT("MassMode"), (int32)InMassMode);
+		Properties.AddValue(TEXT("MassValue"), InMassValue);
+		Properties.AddValue(TEXT("MinPerParticleMass"), InMinPerParticleMass);
+	}
+	
+	// Edge constraint
+	if (InEdgeStiffness[0] > 0.f || InEdgeStiffness[1] > 0.f)
+	{
+		const int32 EdgeSpringStiffnessIndex = Properties.AddProperty(TEXT("EdgeSpringStiffness"), bEnable, bAnimatable);
+		Properties.SetWeightedValue(EdgeSpringStiffnessIndex, InEdgeStiffness[0], InEdgeStiffness[1]);
+		Properties.SetStringValue(EdgeSpringStiffnessIndex, TEXT("EdgeStiffness"));
+	}
+	
+	// Bending constraint
+	if (InBendingStiffness[0] > 0.f || InBendingStiffness[1] > 0.f ||
+		(bInUseBendingElements && (InBucklingStiffness[0] > 0.f || InBucklingStiffness[1] > 0.f)))
+	{
+		Properties.AddValue(TEXT("UseBendingElements"), bInUseBendingElements);
+		if (bInUseBendingElements)
+		{
+			const int32 BendingElementStiffnessIndex = Properties.AddProperty(TEXT("BendingElementStiffness"), bEnable, bAnimatable);
+			Properties.SetWeightedValue(BendingElementStiffnessIndex, InBendingStiffness[0], InBendingStiffness[1]);
+			Properties.SetStringValue(BendingElementStiffnessIndex, TEXT("BendingStiffness"));
+
+			Properties.AddValue(TEXT("BucklingRatio"), InBucklingRatio);
+
+			if (InBucklingStiffness[0] > 0.f && InBucklingStiffness[1] > 0.f)
+			{
+				const int32 BucklingStiffnessIndex = Properties.AddProperty(TEXT("BucklingStiffness"), bEnable, bAnimatable);
+				Properties.SetWeightedValue(BucklingStiffnessIndex, InBucklingStiffness[0], InBucklingStiffness[1]);
+				Properties.SetStringValue(BucklingStiffnessIndex, TEXT("BucklingStiffness"));
+			}
+		}
+		else  // Not using bending elements
+		{
+			const int32 BendingSpringStiffnessIndex = Properties.AddProperty(TEXT("BendingSpringStiffness"), bEnable, bAnimatable);
+			Properties.SetWeightedValue(BendingSpringStiffnessIndex, InBendingStiffness[0], InBendingStiffness[1]);
+			Properties.SetStringValue(BendingSpringStiffnessIndex, TEXT("BendingStiffness"));
+		}
+	}
+
+	// Area constraint
+	if (InAreaStiffness[0] > 0.f || InAreaStiffness[1] > 0.f)
+	{
+		const int32 AreaSpringStiffnessIndex = Properties.AddProperty(TEXT("AreaSpringStiffness"), bEnable, bAnimatable);
+		Properties.SetWeightedValue(AreaSpringStiffnessIndex, InAreaStiffness[0], InAreaStiffness[1]);
+		Properties.SetStringValue(AreaSpringStiffnessIndex, TEXT("AreaStiffness"));
+	}
+
+	// Long range attachment
+	if (InTetherStiffness[0] > 0.f || InTetherStiffness[1] > 0.f)
+	{
+		Properties.AddValue(TEXT("TetherMode"), (int32)InTetherMode);
+
+		const int32 TetherStiffnessIndex = Properties.AddProperty(TEXT("TetherStiffness"), bEnable , bAnimatable);
+		Properties.SetWeightedValue(TetherStiffnessIndex, InTetherStiffness[0], InTetherStiffness[1]);
+		Properties.SetStringValue(TetherStiffnessIndex, TEXT("TetherStiffness"));
+	
+		const int32 TetherScaleIndex = Properties.AddProperty(TEXT("TetherScale"), bEnable, bAnimatable);
+		Properties.SetWeightedValue(TetherScaleIndex, InTetherScale[0], InTetherScale[1]);
+		Properties.SetStringValue(TetherScaleIndex, TEXT("TetherScale"));
+	}
+
+	// AnimDrive
+	if (InAnimDriveStiffness[0] > 0.f || InAnimDriveStiffness[1] > 0.f)
+	{
+		const int32 AnimDriveStiffnessIndex = Properties.AddProperty(TEXT("AnimDriveStiffness"), bEnable, bAnimatable);
+		Properties.SetWeightedValue(AnimDriveStiffnessIndex, InAnimDriveStiffness[0], InAnimDriveStiffness[1]);
+		Properties.SetStringValue(AnimDriveStiffnessIndex, TEXT("AnimDriveStiffness"));
+
+		const int32 AnimDriveDampingIndex = Properties.AddProperty(TEXT("AnimDriveDamping"), bEnable, bAnimatable);
+		Properties.SetWeightedValue(AnimDriveDampingIndex, InAnimDriveDamping[0], InAnimDriveDamping[1]);
+		Properties.SetStringValue(AnimDriveDampingIndex, TEXT("AnimDriveDamping"));
+	}
+
+	// Gravity
+	{
+		Properties.AddValue(TEXT("GravityScale"), InGravityScale, bEnable, bAnimatable);
+		Properties.AddValue(TEXT("UseGravityOverride"), bInUseGravityOverride, bEnable, bAnimatable);
+		Properties.AddValue(TEXT("GravityOverride"), FVector3f(InGravityOverride), bEnable, bAnimatable);
+	}
+
+	// Velocity scale
+	{
+		Properties.AddValue(TEXT("LinearVelocityScale"), FVector3f(InLinearVelocityScale), bEnable, bAnimatable);
+		Properties.AddValue(TEXT("AngularVelocityScale"), InAngularVelocityScale, bEnable, bAnimatable);
+		Properties.AddValue(TEXT("FictitiousAngularScale"), InFictitiousAngularScale, bEnable, bAnimatable);
+	}
+
+	// Aerodynamics
+	Properties.AddValue(TEXT("UsePointBasedWindModel"), bInUsePointBasedWindModel);
+	if (!bInUsePointBasedWindModel && (InDrag[0] > 0.f || InDrag[1] > 0.f || InLift[0] > 0.f || InLift[1] > 0.f))
+	{
+		const int32 DragIndex = Properties.AddProperty(TEXT("Drag"), bEnable, bAnimatable);
+		Properties.SetWeightedValue(DragIndex, InDrag[0], InDrag[1]);
+		Properties.SetStringValue(DragIndex, TEXT("Drag"));
+
+		const int32 LiftIndex = Properties.AddProperty(TEXT("Lift"), bEnable, bAnimatable);
+		Properties.SetWeightedValue(LiftIndex, InLift[0], InLift[1]);
+		Properties.SetStringValue(LiftIndex, TEXT("Lift"));
+
+		Properties.AddValue(TEXT("AirDensity"), ClothingSimulationClothDefault::AirDensity, bEnable, bAnimatable);
+	}
+
+	// Pressure
+	if (InPressure[0] != 0.f || InPressure[1] != 0.f)
+	{
+		const int32 PressureIndex = Properties.AddProperty(TEXT("Pressure"), bEnable, bAnimatable);
+		Properties.SetWeightedValue(PressureIndex, InPressure[0], InPressure[1]);
+		Properties.SetStringValue(PressureIndex, TEXT("Pressure"));
+	}
+
+	// Damping
+	Properties.AddValue(TEXT("DampingCoefficient"), InDampingCoefficient, bEnable, bAnimatable);
+	Properties.AddValue(TEXT("LocalDampingCoefficient"), InLocalDampingCoefficient, bEnable, bAnimatable);
+
+	// Collision
+	Properties.AddValue(TEXT("CollisionThickness"), InCollisionThickness, bEnable, bAnimatable);
+	Properties.AddValue(TEXT("FrictionCoefficient"), InFrictionCoefficient, bEnable, bAnimatable);
+	Properties.AddValue(TEXT("UseCCD"), bInUseCCD, bEnable, bAnimatable);
+	Properties.AddValue(TEXT("UseSelfCollisions"), bInUseSelfCollisions);
+	Properties.AddValue(TEXT("SelfCollisionThickness"), InSelfCollisionThickness);
+	Properties.AddValue(TEXT("SelfCollisionFrictionCoefficient"), InSelfCollisionFrictionCoefficient);
+	Properties.AddValue(TEXT("UseSelfIntersections"), bInUseSelfIntersections);
+
+	// Max distance
+	{
+		const int32 MaxDistanceIndex = Properties.AddProperty(TEXT("MaxDistance"));
+		Properties.SetWeightedValue(MaxDistanceIndex, 1.f, 1.f);  // Backward compatibility with legacy mask means can't be different from 1.
+		Properties.SetStringValue(MaxDistanceIndex, TEXT("MaxDistance"));
+	}
+
+	// Backstop
+	{
+		const int32 BackstopScaleIndex = Properties.AddProperty(TEXT("BackstopDistance"));
+		Properties.SetWeightedValue(BackstopScaleIndex, 1.f, 1.f);  // Backward compatibility with legacy mask means can't be different from 1.
+		Properties.SetStringValue(BackstopScaleIndex, TEXT("BackstopDistance"));
+
+		const int32 BackstopRadiusIndex = Properties.AddProperty(TEXT("BackstopRadius"));
+		Properties.SetWeightedValue(BackstopRadiusIndex, 1.f, 1.f);  // Backward compatibility with legacy mask means can't be different from 1.
+		Properties.SetStringValue(BackstopRadiusIndex, TEXT("BackstopRadius"));
+
+		Properties.AddValue(TEXT("UseLegacyBackstop"), bInUseLegacyBackstop);
+	}
+
+	Config = new FClothingSimulationConfig(PropertyCollection);
+
+	// Set mesh and colliders
 	SetMesh(InMesh);
 	SetColliders(MoveTemp(InColliders));
 }
 
 FClothingSimulationCloth::~FClothingSimulationCloth()
 {
+	// If the PropertyCollection is owned by this object, so does the current config object
+	if (PropertyCollection.IsValid())
+	{
+		delete Config;
+	}
 }
 
 PRAGMA_DISABLE_DEPRECATION_WARNINGS  // TODO: CHAOS_IS_CLOTHINGSIMULATIONMESH_ABSTRACT
@@ -520,6 +509,8 @@ void FClothingSimulationCloth::SetMesh(FClothingSimulationMesh* InMesh)
 	// Reset LODs
 	const int32 NumLODs = Mesh ? Mesh->GetNumLODs() : 0;
 	LODData.Reset(NumLODs);
+
+	const ETetherMode TetherMode = (ETetherMode)Config->GetProperties().GetValue<int32>(TEXT("TetherMode"), ClothingSimulationClothDefault::TetherMode);
 	const bool bUseGeodesicTethers = (TetherMode == ETetherMode::Geodesic);
 	for (int32 Index = 0; Index < NumLODs; ++Index)
 	{
@@ -541,6 +532,26 @@ void FClothingSimulationCloth::SetMesh(FClothingSimulationMesh* InMesh)
 }
 PRAGMA_ENABLE_DEPRECATION_WARNINGS
 
+void FClothingSimulationCloth::SetConfig(FClothingSimulationConfig* InConfig)
+{
+	if (InConfig)
+	{
+		// If the PropertyCollection is owned by this object, so does the current config object
+		if (InConfig && PropertyCollection.IsValid())
+		{
+			delete Config;
+			PropertyCollection.Reset();
+		}
+		Config = InConfig;
+	}
+	else
+	{
+		// Create a default empty config object for coherence
+		PropertyCollection = MakeShared<FManagedArrayCollection>();
+		Config = new FClothingSimulationConfig(PropertyCollection);
+	}
+}
+
 void FClothingSimulationCloth::SetColliders(TArray<FClothingSimulationCollider*>&& InColliders)
 {
 	// Empty the collider list, but keep the pointers around for the removal operation below
@@ -684,7 +695,14 @@ int32 FClothingSimulationCloth::GetOffset(const FClothingSimulationSolver* Solve
 TVec3<FRealSingle> FClothingSimulationCloth::GetGravity(const FClothingSimulationSolver* Solver) const
 {
 	check(Solver);
-	return Solver->IsClothGravityOverrideEnabled() && bIsGravityOverridden ? GravityOverride : Solver->GetGravity() * GravityScale;
+	check(Config);
+	const Softs::FPropertyCollectionConstAdapter& ConfigProperties = Config->GetProperties();
+
+	const bool bUseGravityOverride = ConfigProperties.GetValue<bool>(TEXT("UseGravityOverride"));
+	const TVec3<FRealSingle> GravityOverride = (TVec3<FRealSingle>)ConfigProperties.GetValue<FVector3f>(TEXT("GravityOverride"), FVector3f(0.f, 0.f, ClothingSimulationClothDefault::GravityZOverride));
+	const FRealSingle GravityScale = (FRealSingle)ConfigProperties.GetValue<float>(TEXT("GravityScale"), 1.f);
+
+	return Solver->IsClothGravityOverrideEnabled() && bUseGravityOverride ? GravityOverride : Solver->GetGravity() * GravityScale;
 }
 
 FAABB3 FClothingSimulationCloth::CalculateBoundingBox(const FClothingSimulationSolver* Solver) const
@@ -845,19 +863,22 @@ PRAGMA_ENABLE_DEPRECATION_WARNINGS
 
 	if (LODIndex != INDEX_NONE)
 	{
+		check(Config);
+		const Softs::FPropertyCollectionConstAdapter& ConfigProperties = Config->GetProperties();
+
 		// Update Cloth group parameters  TODO: Cloth groups should exist as their own node object so that they can be used by several cloth objects
 		LODData[LODIndex]->Update(Solver, this);
 
 		// TODO: Move all groupID updates out of the cloth update to allow to use of the same GroupId with different cloths
 
 		// Set the reference input velocity and deal with teleport & reset; external forces depends on these values, so they must be initialized before then
-		FVec3 OutLinearVelocityScale;
-		FReal OutAngularVelocityScale;
+		FVec3f OutLinearVelocityScale;
+		FRealSingle OutAngularVelocityScale;
 
 		if (bNeedsReset)
 		{
 			// Make sure not to do any pre-sim transform just after a reset
-			OutLinearVelocityScale = FVec3(1.f);
+			OutLinearVelocityScale = FVec3f(1.f);
 			OutAngularVelocityScale = 1.f;
 
 			// Reset to start pose
@@ -867,17 +888,19 @@ PRAGMA_ENABLE_DEPRECATION_WARNINGS
 		else if (bNeedsTeleport)
 		{
 			// Remove all impulse velocity from the last frame
-			OutLinearVelocityScale = FVec3(0.f);
+			OutLinearVelocityScale = FVec3f(0.f);
 			OutAngularVelocityScale = 0.f;
 			UE_LOG(LogChaosCloth, VeryVerbose, TEXT("Cloth in group Id %d Needs teleport."), GroupId);
 		}
 		else
 		{
-			// Use the cloth group's parameters
-			OutLinearVelocityScale = LinearVelocityScale;
-			OutAngularVelocityScale = AngularVelocityScale;
+			// Use the cloth config parameters
+			OutLinearVelocityScale = ConfigProperties.GetValue<FVector3f>(TEXT("LinearVelocityScale"), FVector3f(ClothingSimulationClothDefault::VelocityScale));
+			OutAngularVelocityScale = ConfigProperties.GetValue<float>(TEXT("AngularVelocityScale"), ClothingSimulationClothDefault::VelocityScale);
 		}
 
+		const FRealSingle FictitiousAngularScale = ConfigProperties.GetValue<float>(TEXT("FictitiousAngularScale"), ClothingSimulationClothDefault::FictitiousAngularScale);
+
 		Solver->SetReferenceVelocityScale(
 			GroupId,
 			OldReferenceSpaceTransform,
@@ -893,25 +916,37 @@ PRAGMA_ENABLE_DEPRECATION_WARNINGS
 		Solver->SetGravity(GroupId, GetGravity(Solver));
 
 		// External forces (legacy wind+field)
-		Solver->AddExternalForces(GroupId, bUseLegacyWind);
+		const bool bUsePointBasedWindModel = ConfigProperties.GetValue<bool>(TEXT("UsePointBasedWindModel"));
+		Solver->AddExternalForces(GroupId, bUsePointBasedWindModel);
 
-		if (bUseLegacyWind && ClothingSimulationClothConsoleVariables::CVarLegacyDisablesAccurateWind.GetValueOnAnyThread())
+		const FVec2f Pressure = ConfigProperties.GetWeightedFloatValue(TEXT("Pressure"));
+		if (bUsePointBasedWindModel && ClothingSimulationClothConsoleVariables::CVarLegacyDisablesAccurateWind.GetValueOnAnyThread())
 		{
 			Solver->SetWindAndPressureProperties(GroupId, TVec2<FRealSingle>(0.), TVec2<FRealSingle>(0.), (FRealSingle)0., Pressure);  // Disable the wind velocity field
 		}
 		else
 		{
+			const FVec2f Drag = ConfigProperties.GetWeightedFloatValue(TEXT("Drag"), ClothingSimulationClothDefault::Drag);
+			const FVec2f Lift = ConfigProperties.GetWeightedFloatValue(TEXT("Lift"), ClothingSimulationClothDefault::Lift);
+			const FRealSingle AirDensity = ConfigProperties.GetValue<float>(TEXT("AirDensity"), ClothingSimulationClothDefault::AirDensity);
 			Solver->SetWindAndPressureProperties(GroupId, Drag, Lift, AirDensity, Pressure);
 		}
+		const FVec3f WindVelocity = ConfigProperties.GetValue<FVector3f>(TEXT("WindVelocity"));
 		Solver->SetWindVelocity(GroupId, WindVelocity + Solver->GetWindVelocity());
 
 		// Update general solver properties
 PRAGMA_DISABLE_DEPRECATION_WARNINGS  // TODO: CHAOS_IS_CLOTHINGSIMULATIONMESH_ABSTRACT
 		const Softs::FSolverReal MeshScale = Mesh->GetScale();
 PRAGMA_ENABLE_DEPRECATION_WARNINGS
+
+		const FRealSingle DampingCoefficient = ConfigProperties.GetValue<float>(TEXT("DampingCoefficient"), ClothingSimulationClothDefault::DampingCoefficient);
+		const FRealSingle LocalDampingCoefficient = ConfigProperties.GetValue<float>(TEXT("LocalDampingCoefficient"));
+		const FRealSingle CollisionThickness = ConfigProperties.GetValue<float>(TEXT("CollisionThickness"), ClothingSimulationClothDefault::CollisionThickness);
+		const FRealSingle FrictionCoefficient = ConfigProperties.GetValue<float>(TEXT("FrictionCoefficient"), ClothingSimulationClothDefault::FrictionCoefficient);
 		Solver->SetProperties(GroupId, DampingCoefficient, LocalDampingCoefficient, CollisionThickness * MeshScale, FrictionCoefficient);
 
 		// Update use of continuous collision detection
+		const bool bUseCCD = ConfigProperties.GetValue<bool>(TEXT("UseCCD"));
 		Solver->SetUseCCD(GroupId, bUseCCD);
 	}
 
@@ -980,4 +1015,75 @@ TConstArrayView<Softs::FSolverReal> FClothingSimulationCloth::GetParticleInvMass
 	return TConstArrayView<Softs::FSolverReal>(Solver->GetParticleInvMasses(GetOffset(Solver, LODIndex)), GetNumParticles(LODIndex));
 }
 
+void FClothingSimulationCloth::SetMaterialProperties(const TVec2<FRealSingle>& InEdgeStiffness, const TVec2<FRealSingle>& InBendingStiffness, const TVec2<FRealSingle>& InAreaStiffness)
+{
+	Config->GetProperties().SetWeightedFloatValue(TEXT("EdgeSpringStiffness"), FVector2f(InEdgeStiffness));
+	Config->GetProperties().SetWeightedFloatValue(TEXT("BendingSpringStiffness"), FVector2f(InBendingStiffness));
+	Config->GetProperties().SetWeightedFloatValue(TEXT("AreaSpringStiffness"), FVector2f(InAreaStiffness));
+}
+
+void FClothingSimulationCloth::SetLongRangeAttachmentProperties(const TVec2<FRealSingle>& InTetherStiffness, const TVec2<FRealSingle>& InTetherScale)
+{
+	Config->GetProperties().SetWeightedFloatValue(TEXT("TetherStiffness"), FVector2f(InTetherStiffness));
+	Config->GetProperties().SetWeightedFloatValue(TEXT("TetherScale"), FVector2f(InTetherScale));
+}
+
+void FClothingSimulationCloth::SetCollisionProperties(FRealSingle InCollisionThickness, FRealSingle InFrictionCoefficient, bool bInUseCCD, FRealSingle InSelfCollisionThickness)
+{
+	Config->GetProperties().SetValue(TEXT("CollisionThickness"), (float)InCollisionThickness);
+	Config->GetProperties().SetValue(TEXT("FrictionCoefficient"), (float)InFrictionCoefficient);
+	Config->GetProperties().SetValue(TEXT("bUseCCD"), bInUseCCD);
+	Config->GetProperties().SetValue(TEXT("SelfCollisionThickness"), (float)InSelfCollisionThickness);
+}
+
+void FClothingSimulationCloth::SetBackstopProperties(bool bInEnableBackstop)
+{
+	Config->GetProperties().SetEnabled(TEXT("BackstopScale"), bInEnableBackstop);  // BackstopScale controls whether the backstop is enabled or not
+}
+
+void FClothingSimulationCloth::SetDampingProperties(FRealSingle InDampingCoefficient, FRealSingle InLocalDampingCoefficient)
+{
+	Config->GetProperties().SetValue(TEXT("DampingCoefficient"), (float)InDampingCoefficient);
+	Config->GetProperties().SetValue(TEXT("LocalDampingCoefficient"), (float)InLocalDampingCoefficient);
+}
+
+void FClothingSimulationCloth::SetAerodynamicsProperties(const TVec2<FRealSingle>& InDrag, const TVec2<FRealSingle>& InLift, FRealSingle InAirDensity, const FVec3& InWindVelocity)
+{
+	Config->GetProperties().SetWeightedFloatValue(TEXT("Drag"), FVector2f(InDrag));
+	Config->GetProperties().SetWeightedFloatValue(TEXT("Lift"), FVector2f(InLift));
+	Config->GetProperties().SetValue(TEXT("AirDensity"), (float)InAirDensity);
+	Config->GetProperties().SetValue(TEXT("WindVelocity"), FVector3f(InWindVelocity));
+}
+
+void FClothingSimulationCloth::SetPressureProperties(const TVec2<FRealSingle>& InPressure)
+{
+	Config->GetProperties().SetWeightedFloatValue(TEXT("Pressure"), FVector2f(InPressure));
+}
+
+void FClothingSimulationCloth::SetGravityProperties(FRealSingle InGravityScale, bool bInUseGravityOverride, const FVec3& InGravityOverride)
+{
+	Config->GetProperties().SetValue(TEXT("GravityScale"), (float)InGravityScale);
+	Config->GetProperties().SetValue(TEXT("UseGravityOverride"), bInUseGravityOverride);
+	Config->GetProperties().SetValue(TEXT("GravityOverride"), FVector3f(InGravityOverride));
+}
+
+void FClothingSimulationCloth::SetAnimDriveProperties(const TVec2<FRealSingle>& InAnimDriveStiffness, const TVec2<FRealSingle>& InAnimDriveDamping)
+{
+	Config->GetProperties().SetWeightedFloatValue(TEXT("AnimDriveStiffness"), FVector2f(InAnimDriveStiffness));
+	Config->GetProperties().SetWeightedFloatValue(TEXT("AnimDriveDamping"), FVector2f(InAnimDriveDamping));
+}
+
+void FClothingSimulationCloth::GetAnimDriveProperties(TVec2<FRealSingle>& OutAnimDriveStiffness, TVec2<FRealSingle>& OutAnimDriveDamping)
+{
+	OutAnimDriveStiffness = TVec2<FRealSingle>(Config->GetProperties().GetWeightedFloatValue(TEXT("AnimDriveStiffness")));
+	OutAnimDriveDamping = TVec2<FRealSingle>(Config->GetProperties().GetWeightedFloatValue(TEXT("AnimDriveDamping")));
+}
+
+void FClothingSimulationCloth::SetVelocityScaleProperties(const FVec3& InLinearVelocityScale, FRealSingle InAngularVelocityScale, FRealSingle InFictitiousAngularScale)
+{
+	Config->GetProperties().SetValue(TEXT("LinearVelocityScale"), FVector3f(InLinearVelocityScale));
+	Config->GetProperties().SetValue(TEXT("AngularVelocityScale"), (float)InAngularVelocityScale);
+	Config->GetProperties().SetValue(TEXT("FictitiousAngularScale"), (float)InFictitiousAngularScale);
+}
+
 }  // End namespace Chaos
diff --git a/Engine/Plugins/Experimental/ChaosCloth/Source/Chaos/Private/ChaosCloth/ChaosClothingSimulationConfig.cpp b/Engine/Plugins/Experimental/ChaosCloth/Source/Chaos/Private/ChaosCloth/ChaosClothingSimulationConfig.cpp
--- a/Engine/Plugins/Experimental/ChaosCloth/Source/Chaos/Private/ChaosCloth/ChaosClothingSimulationConfig.cpp
+++ b/Engine/Plugins/Experimental/ChaosCloth/Source/Chaos/Private/ChaosCloth/ChaosClothingSimulationConfig.cpp
@@ -158,9 +158,6 @@ namespace Chaos
 				const int32 LiftIndex = Properties->AddProperty(TEXT("Lift"), bEnable, bAnimatable);
 				Properties->SetWeightedValue(LiftIndex, ClothConfig->Lift.Low, ClothConfig->Lift.High);
 				Properties->SetStringValue(LiftIndex, TEXT("Lift"));
-
-				constexpr float AirDensity = 1.225e-6f;
-				Properties->AddValue(TEXT("AirDensity"), AirDensity, bEnable, bAnimatable);
 			}
 
 			// Pressure
diff --git a/Engine/Plugins/Experimental/ChaosCloth/Source/Chaos/Private/ChaosCloth/ChaosClothingSimulationSolver.cpp b/Engine/Plugins/Experimental/ChaosCloth/Source/Chaos/Private/ChaosCloth/ChaosClothingSimulationSolver.cpp
--- a/Engine/Plugins/Experimental/ChaosCloth/Source/Chaos/Private/ChaosCloth/ChaosClothingSimulationSolver.cpp
+++ b/Engine/Plugins/Experimental/ChaosCloth/Source/Chaos/Private/ChaosCloth/ChaosClothingSimulationSolver.cpp
@@ -3,6 +3,7 @@
 #include "ChaosCloth/ChaosClothingSimulationCloth.h"
 #include "ChaosCloth/ChaosClothingSimulationCollider.h"
 #include "ChaosCloth/ChaosClothingSimulationMesh.h"
+#include "ChaosCloth/ChaosClothingSimulationConfig.h"
 #include "ChaosCloth/ChaosClothingSimulation.h"
 #include "ChaosCloth/ChaosClothPrivate.h"
 #include "Chaos/PBDEvolution.h"
@@ -72,8 +73,11 @@ namespace ClothingSimulationSolverDefault
 	static const Softs::FSolverVec3 Gravity((Softs::FSolverReal)0., (Softs::FSolverReal)0., (Softs::FSolverReal)-980.665);  // cm/s^2
 	static const Softs::FSolverVec3 WindVelocity((Softs::FSolverReal)0.);
 	static const int32 NumIterations = 1;
+	static const int32 MinNumIterations = 1;
 	static const int32 MaxNumIterations = 10;
 	static const int32 NumSubsteps = 1;
+	static const int32 MinNumSubsteps = 1;
+	static const FRealSingle SolverFrequency = 60.f;
 	static const FRealSingle SelfCollisionThickness = 2.f;
 	static const FRealSingle CollisionThickness = 1.2f;
 	static const FRealSingle FrictionCoefficient = 0.2f;
@@ -696,6 +700,63 @@ void FClothingSimulationSolver::SetWindVelocity(const TVec3<FRealSingle>& InWind
 	LegacyWindAdaption = InLegacyWindAdaption;
 }
 
+void FClothingSimulationSolver::SetNumIterations(int32 InNumIterations)
+{
+	if (Config)
+	{
+		Config->GetProperties().SetValue(TEXT("NumIterations"), InNumIterations);
+	}
+	else
+	{
+		NumIterations = InNumIterations;
+	}
+}
+
+int32 FClothingSimulationSolver::GetNumIterations() const
+{
+	return Config ?
+		Config->GetProperties().GetValue<int32>(TEXT("NumIterations"), ClothingSimulationSolverDefault::NumIterations) :
+		NumIterations;
+}
+
+void FClothingSimulationSolver::SetMaxNumIterations(int32 InMaxNumIterations)
+{
+	if (Config)
+	{
+		Config->GetProperties().SetValue(TEXT("MaxNumIterations"), InMaxNumIterations);
+	}
+	else
+	{
+		MaxNumIterations = InMaxNumIterations;
+	}
+}
+
+int32 FClothingSimulationSolver::GetMaxNumIterations() const
+{
+	return Config?
+		Config->GetProperties().GetValue<int32>(TEXT("MaxNumIterations"), ClothingSimulationSolverDefault::MaxNumIterations) :
+		MaxNumIterations;
+}
+
+void FClothingSimulationSolver::SetNumSubsteps(int32 InNumSubsteps)
+{
+	if (Config)
+	{
+		Config->GetProperties().SetValue(TEXT("NumSubsteps"), InNumSubsteps);
+	}
+	else
+	{
+		NumSubsteps = InNumSubsteps;
+	}
+}
+
+int32 FClothingSimulationSolver::GetNumSubsteps() const
+{
+	return Config?
+		Config->GetProperties().GetValue<int32>(TEXT("NumSubsteps"), ClothingSimulationSolverDefault::NumSubsteps) :
+		NumSubsteps;
+}
+
 void FClothingSimulationSolver::SetWindVelocity(uint32 GroupId, const TVec3<FRealSingle>& InWindVelocity)
 {
 	Softs::FVelocityAndPressureField& VelocityField = Evolution->GetVelocityAndPressureField(GroupId);
@@ -1028,21 +1089,48 @@ void FClothingSimulationSolver::Update(Softs::FSolverReal InDeltaTime)
 			SCOPE_CYCLE_COUNTER(STAT_ChaosClothSolverUpdateSolverStep);
 			SCOPE_CYCLE_COUNTER(STAT_ClothInternalSolve);
 
+			// Read config properties if any, otherwise go back to using the property setters for this object
+			int32 ConfigNumIterations;
+			int32 ConfigMaxNumIterations;
+			int32 ConfigNumSubsteps;
+			if (Config)
+			{
+				const Softs::FPropertyCollectionConstAdapter& Properties = Config->GetProperties();
+
+				ConfigMaxNumIterations = FMath::Max(
+					Properties.GetValue<int32>(TEXT("MaxNumIterations"), ClothingSimulationSolverDefault::MaxNumIterations),
+					ClothingSimulationSolverDefault::MinNumIterations);
+
+				ConfigNumIterations = FMath::Clamp(
+					Properties.GetValue<int32>(TEXT("NumIterations"), ClothingSimulationSolverDefault::NumIterations),
+					ClothingSimulationSolverDefault::MinNumIterations,
+					ConfigMaxNumIterations);
+
+				ConfigNumSubsteps = FMath::Max(
+					Properties.GetValue<int32>(TEXT("NumSubsteps"), ClothingSimulationSolverDefault::NumSubsteps),
+					ClothingSimulationSolverDefault::MinNumSubsteps);
+			}
+			else
+			{
+				ConfigMaxNumIterations = MaxNumIterations;
+				ConfigNumIterations = NumIterations;
+				ConfigNumSubsteps = NumSubsteps;
+			}
 			// Update solver time dependent parameters
-			constexpr Softs::FSolverReal SolverFrequency = 60.f;  // TODO: This could become a solver property
+			const Softs::FSolverReal SolverFrequency = (Softs::FSolverReal)ClothingSimulationSolverDefault::SolverFrequency;  // 60Hz default TODO: Should this become a solver property?
 
 			const int32 TimeDependentNumIterations = bClothSolverDisableTimeDependentNumIterations ?
-				NumIterations :
-				FMath::RoundToInt32(SolverFrequency * DeltaTime * (Softs::FSolverReal)NumIterations);
+				ConfigNumIterations :
+				FMath::RoundToInt32(SolverFrequency * DeltaTime * (Softs::FSolverReal)ConfigNumIterations);
 
-			Evolution->SetIterations(FMath::Clamp(TimeDependentNumIterations, 1, MaxNumIterations));
+			Evolution->SetIterations(FMath::Clamp(TimeDependentNumIterations, 1, ConfigMaxNumIterations));
 
 			// Advance substeps
-			const Softs::FSolverReal SubstepDeltaTime = DeltaTime / (Softs::FSolverReal)NumSubsteps;
+			const Softs::FSolverReal SubstepDeltaTime = DeltaTime / (Softs::FSolverReal)ConfigNumSubsteps;
 	
-			for (int32 i = 0; i < NumSubsteps; ++i)
+			for (int32 i = 0; i < ConfigNumSubsteps; ++i)
 			{
-				PreSubstep(FMath::Clamp((Softs::FSolverReal)(i + 1) / (Softs::FSolverReal)NumSubsteps, (Softs::FSolverReal)0., (Softs::FSolverReal)1.));
+				PreSubstep(FMath::Clamp((Softs::FSolverReal)(i + 1) / (Softs::FSolverReal)ConfigNumSubsteps, (Softs::FSolverReal)0., (Softs::FSolverReal)1.));
 				Evolution->AdvanceOneTimeStep(SubstepDeltaTime);
 			}
 
diff --git a/Engine/Plugins/Experimental/ChaosCloth/Source/Chaos/Public/ChaosCloth/ChaosClothConstraints.h b/Engine/Plugins/Experimental/ChaosCloth/Source/Chaos/Public/ChaosCloth/ChaosClothConstraints.h
--- a/Engine/Plugins/Experimental/ChaosCloth/Source/Chaos/Public/ChaosCloth/ChaosClothConstraints.h
+++ b/Engine/Plugins/Experimental/ChaosCloth/Source/Chaos/Public/ChaosCloth/ChaosClothConstraints.h
@@ -38,73 +38,129 @@ namespace Chaos
 		// ---- End of Solver interface ----
 
 		// ---- Cloth interface ----
+		void AddRules(
+			const Softs::FPropertyCollectionConstAdapter& ConfigProperties,
+			const FTriangleMesh& TriangleMesh,
+			const TArray<TConstArrayView<FRealSingle>>& WeightMaps,
+			const TArray<TConstArrayView<TTuple<int32, int32, FRealSingle>>>& Tethers,
+			Softs::FSolverReal MeshScale,
+			bool bEnabled);
+
+		void Update(
+			const Softs::FPropertyCollectionConstAdapter& ConfigProperties,
+			Softs::FSolverReal MeshScale,
+			Softs::FSolverReal MaxDistancesScale = (Softs::FSolverReal)1.);
+
+		void Enable(bool bEnable);
+
+		UE_DEPRECATED(5.2, "Use AddRules() instead.")
 		void SetEdgeConstraints(const TArray<TVec3<int32>>& SurfaceElements, const TConstArrayView<FRealSingle>& StiffnessMultipliers, bool bUseXPBDConstraints);
+		UE_DEPRECATED(5.2, "Use AddRules() instead.")
 		void SetXPBDEdgeConstraints(const TArray<TVec3<int32>>& SurfaceElements, const TConstArrayView<FRealSingle>& StiffnessMultipliers, const TConstArrayView<FRealSingle>& DampingRatioMultipliers);
+		UE_DEPRECATED(5.2, "Use AddRules() instead.")
 		void SetBendingConstraints(const TArray<TVec2<int32>>& Edges, const TConstArrayView<FRealSingle>& StiffnessMultipliers, bool bUseXPBDConstraints);
+		UE_DEPRECATED(5.2, "Use AddRules() instead.")
 		void SetBendingConstraints(TArray<TVec4<int32>>&& BendingElements, const TConstArrayView<FRealSingle>& StiffnessMultipliers, const TConstArrayView<FRealSingle>& BucklingStiffnessMultipliers, bool bUseXPBDConstraints);
+		UE_DEPRECATED(5.2, "Use AddRules() instead.")
 		void SetXPBDBendingConstraints(TArray<TVec4<int32>>&& BendingElements, const TConstArrayView<FRealSingle>& StiffnessMultipliers, const TConstArrayView<FRealSingle>& BucklingStiffnessMultipliers, const TConstArrayView<FRealSingle>& DampingRatioMultipliers);
 		UE_DEPRECATED(5.1, "Use SetBendingConstraints(TArray<TVec4<int32>>&& BendingElements, const TConstArrayView<FRealSingle>& StiffnessMultipliers, const TConstArrayView<FRealSingle>& BucklingStiffnessMultipliers, bool bUseXPBDConstraints) instead.")
 		void SetBendingConstraints(TArray<TVec4<int32>>&& BendingElements, Softs::FSolverReal BendingStiffness);
+		UE_DEPRECATED(5.2, "Use AddRules() instead.")
 		void SetAreaConstraints(const TArray<TVec3<int32>>& SurfaceElements, const TConstArrayView<FRealSingle>& StiffnessMultipliers, bool bUseXPBDConstraints);
+		UE_DEPRECATED(5.2, "Use AddRules() instead.")
 		void SetVolumeConstraints(const TArray<TVec2<int32>>& DoubleBendingEdges, Softs::FSolverReal VolumeStiffness);
+		UE_DEPRECATED(5.2, "Use AddRules() instead.")
 		void SetVolumeConstraints(TArray<TVec3<int32>>&& SurfaceElements, Softs::FSolverReal VolumeStiffness);
-		void SetLongRangeConstraints(const TArray<TConstArrayView<TTuple<int32, int32, FRealSingle>>>& Tethers,
-			const TConstArrayView<FRealSingle>& TetherStiffnessMultipliers, const TConstArrayView<FRealSingle>& TetherScaleMultipliers,
-			const Softs::FSolverVec2& TetherScale);
+		UE_DEPRECATED(5.2, "Use AddRules() instead.")
+		void SetLongRangeConstraints(
+			const TArray<TConstArrayView<TTuple<int32,int32, FRealSingle>>>& Tethers,
+			const TConstArrayView<FRealSingle>& TetherStiffnessMultipliers,
+			const TConstArrayView<FRealSingle>& TetherScaleMultipliers,
+			const Softs::FSolverVec2& TetherScale,
+			Softs::FSolverReal MeshScale = (Softs::FSolverReal)1.);
 		UE_DEPRECATED(5.1, "LongRange XPBDConstraints are not supported.")
-		void SetLongRangeConstraints(const TArray<TConstArrayView<TTuple<int32, int32, FRealSingle>>>& Tethers,
-			const TConstArrayView<FRealSingle>& TetherStiffnessMultipliers, const TConstArrayView<FRealSingle>& TetherScaleMultipliers,
-			const Softs::FSolverVec2& TetherScale, bool bUseXPBDConstraints) {
-			return SetLongRangeConstraints(Tethers, TetherStiffnessMultipliers, TetherScaleMultipliers, TetherScale);
-		}
+		void SetLongRangeConstraints(
+			const TArray<TConstArrayView<TTuple<int32, int32, FRealSingle>>>& Tethers,
+			const TConstArrayView<FRealSingle>& TetherStiffnessMultipliers,
+			const TConstArrayView<FRealSingle>& TetherScaleMultipliers,
+			const Softs::FSolverVec2& TetherScale,
+			bool bUseXPBDConstraints,
+			Softs::FSolverReal MeshScale = (Softs::FSolverReal)1.);
+		UE_DEPRECATED(5.2, "Use AddRules() instead.")
 		void SetMaximumDistanceConstraints(const TConstArrayView<FRealSingle>& MaxDistances);
+		UE_DEPRECATED(5.2, "Use AddRules() instead.")
 		void SetBackstopConstraints(const TConstArrayView<FRealSingle>& BackstopDistances, const TConstArrayView<FRealSingle>& BackstopRadiuses, bool bUseLegacyBackstop);
+		UE_DEPRECATED(5.2, "Use AddRules() instead.")
 		void SetAnimDriveConstraints(const TConstArrayView<FRealSingle>& AnimDriveStiffnessMultipliers, const TConstArrayView<FRealSingle>& AnimDriveDampingMultipliers);
+		UE_DEPRECATED(5.2, "Use AddRules() instead.")
 		void SetShapeTargetConstraints(Softs::FSolverReal ShapeTargetStiffness);
-		void SetSelfCollisionConstraints(const class FTriangleMesh& TriangleMesh, TSet<TVec2<int32>>&& DisabledCollisionElements, Softs::FSolverReal SelfCollisionThickness,
-			Softs::FSolverReal SelfCollisionFriction = Softs::FPBDCollisionSpringConstraintsBase::BackCompatFrictionCoefficient, bool bGlobalIntersectionAnalysis = false,
+		UE_DEPRECATED(5.2, "Use AddRules() instead.")
+		void SetSelfCollisionConstraints(
+			const class FTriangleMesh& TriangleMesh,
+			TSet<TVec2<int32>>&& DisabledCollisionElements,
+			Softs::FSolverReal SelfCollisionThickness,
+			Softs::FSolverReal SelfCollisionFriction = Softs::FPBDCollisionSpringConstraintsBase::BackCompatFrictionCoefficient,
+			bool bGlobalIntersectionAnalysis = false,
 			bool bContourMinimization = false);
 
+		UE_DEPRECATED(5.2, "CreateRules will soon be made private, use AddRules() instead.")
 		void CreateRules();
-		void Enable(bool bEnable);
 
+		UE_DEPRECATED(5.2, "Use Update() instead.")
 		void SetEdgeProperties(const Softs::FSolverVec2& EdgeStiffness, const Softs::FSolverVec2& DampingRatio = Softs::FSolverVec2::ZeroVector);
+		UE_DEPRECATED(5.2, "Use Update() instead.")
 		void SetBendingProperties(const Softs::FSolverVec2& BendingStiffness, Softs::FSolverReal BucklingRatio = 0.f, const Softs::FSolverVec2& BucklingStiffness = Softs::FSolverVec2::UnitVector, const Softs::FSolverVec2& DampingRatio = Softs::FSolverVec2::ZeroVector);
+		UE_DEPRECATED(5.2, "Use Update() instead.")
 		void SetAreaProperties(const Softs::FSolverVec2& AreaStiffness);
+		UE_DEPRECATED(5.2, "Use Update() instead.")
 		void SetThinShellVolumeProperties(Softs::FSolverReal VolumeStiffness);
+		UE_DEPRECATED(5.2, "Use Update() instead.")
 		void SetVolumeProperties(Softs::FSolverReal VolumeStiffness);
+		UE_DEPRECATED(5.2, "Use Update() instead.")
 		void SetLongRangeAttachmentProperties(
 			const Softs::FSolverVec2& TetherStiffness,
 			const Softs::FSolverVec2& TetherScale,
 			Softs::FSolverReal MeshScale = (Softs::FSolverReal)1.);
+		UE_DEPRECATED(5.2, "Use Update() instead.")
 		void SetMaximumDistanceProperties(Softs::FSolverReal MaxDistancesMultiplier);
+		UE_DEPRECATED(5.2, "Use Update() instead.")
 		void SetAnimDriveProperties(const Softs::FSolverVec2& AnimDriveStiffness, const Softs::FSolverVec2& AnimDriveDamping);
+		UE_DEPRECATED(5.2, "Use Update() instead.")
 		void SetSelfCollisionProperties(Softs::FSolverReal SelfCollisionThickness, Softs::FSolverReal SelfCollisionFriction = Softs::FPBDCollisionSpringConstraintsBase::BackCompatFrictionCoefficient, 
 			bool bGlobalIntersectionAnalysis = false, bool bContourMinimization = false);
-		UE_DEPRECATED(5.0, "Use SetBackstopProperties(bool, FSolverReal) instead.")
-		void SetBackstopProperties(bool bEnabled) { SetBackstopProperties(bEnabled, (Softs::FSolverReal)1.); }
+		UE_DEPRECATED(5.2, "Use Update() instead.")
 		void SetBackstopProperties(bool bEnabled, Softs::FSolverReal BackstopDistancesMultiplier);
 		// ---- End of Cloth interface ----
 
 		// ---- Debug functions ----
-		const TSharedPtr<Softs::FPBDSpringConstraints> GetEdgeConstraints() const { return EdgeConstraints; }
-		const TSharedPtr<Softs::FXPBDSpringConstraints>& GetXEdgeConstraints() const { return XEdgeConstraints; }
-		const TSharedPtr<Softs::FPBDSpringConstraints>& GetBendingConstraints() const { return BendingConstraints; }  
-		const TSharedPtr<Softs::FXPBDSpringConstraints>& GetXBendingConstraints() const { return XBendingConstraints; }
+		UE_DEPRECATED(5.2, "Use GetEdgeSpringConstraints() instead.")
+		const TSharedPtr<Softs::FPBDSpringConstraints>& GetEdgeConstraints() const { return EdgeConstraints_Deprecated; }
+		UE_DEPRECATED(5.2, "Use GetXEdgeSpringConstraints() instead.")
+		const TSharedPtr<Softs::FXPBDSpringConstraints>& GetXEdgeConstraints() const { return XEdgeConstraints_Deprecated; }
+		UE_DEPRECATED(5.2, "Use GetBendingSpringConstraints() instead.")
+		const TSharedPtr<Softs::FPBDSpringConstraints>& GetBendingConstraints() const { return BendingConstraints_Deprecated; }
+		UE_DEPRECATED(5.2, "Use GetXBendingSpringConstraints() instead.")
+		const TSharedPtr<Softs::FXPBDSpringConstraints>& GetXBendingConstraints() const { return XBendingConstraints_Deprecated; }
+		UE_DEPRECATED(5.2, "Use GetAreaSpringConstraints() instead.")
+		const TSharedPtr<Softs::FPBDAxialSpringConstraints>& GetAreaConstraints() const { return AreaConstraints_Deprecated; }
+		UE_DEPRECATED(5.2, "Use GetXAreaSpringConstraints() instead.")
+		const TSharedPtr<Softs::FXPBDAxialSpringConstraints>& GetXAreaConstraints() const { return XAreaConstraints_Deprecated; }
+		UE_DEPRECATED(5.2, "ThinShellVolume constraints are no longer supported.")
+		const TSharedPtr<Softs::FPBDSpringConstraints>& GetThinShellVolumeConstraints() const { return ThinShellVolumeConstraints_Deprecated; }
+		UE_DEPRECATED(5.2, "Volume constraints are no longer supported.")
+		const TSharedPtr<Softs::FPBDVolumeConstraint>& GetVolumeConstraints() const { return VolumeConstraints_Deprecated; }
+		UE_DEPRECATED(5.1,"LongRange XPBDConstraints are no longer supported.")
+		const TSharedPtr<Softs::FXPBDLongRangeConstraints>& GetXLongRangeConstraints() const { return XLongRangeConstraints_Deprecated; }
+
+		const TSharedPtr<Softs::FPBDEdgeSpringConstraints>& GetEdgeSpringConstraints() const { return EdgeConstraints; }
+		const TSharedPtr<Softs::FXPBDEdgeSpringConstraints>& GetXEdgeSpringConstraints() const { return XEdgeConstraints; }
+		const TSharedPtr<Softs::FPBDBendingSpringConstraints>& GetBendingSpringConstraints() const { return BendingConstraints; }
+		const TSharedPtr<Softs::FXPBDBendingSpringConstraints>& GetXBendingSpringConstraints() const { return XBendingConstraints; }
 		const TSharedPtr<Softs::FPBDBendingConstraints>& GetBendingElementConstraints() const { return BendingElementConstraints; }
 		const TSharedPtr<Softs::FXPBDBendingConstraints>& GetXBendingElementConstraints() const { return XBendingElementConstraints; }
-		const TSharedPtr<Softs::FPBDAxialSpringConstraints>& GetAreaConstraints() const { return AreaConstraints; }
-		const TSharedPtr<Softs::FXPBDAxialSpringConstraints>& GetXAreaConstraints() const { return XAreaConstraints; }
-		const TSharedPtr<Softs::FPBDSpringConstraints>& GetThinShellVolumeConstraints() const { return ThinShellVolumeConstraints; }
-		const TSharedPtr<Softs::FPBDVolumeConstraint>& GetVolumeConstraints() const { return VolumeConstraints; }
+		const TSharedPtr<Softs::FPBDAreaSpringConstraints>& GetAreaSpringConstraints() const { return AreaConstraints; }
+		const TSharedPtr<Softs::FXPBDAreaSpringConstraints>& GetXAreaSpringConstraints() const { return XAreaConstraints; }
 		const TSharedPtr<Softs::FPBDLongRangeConstraints>& GetLongRangeConstraints() const { return LongRangeConstraints; }
-		UE_DEPRECATED(5.1,"LongRange XPBDConstraints are not supported.")
-		const TSharedPtr<Softs::FXPBDLongRangeConstraints>& GetXLongRangeConstraints() const 
-		{
-			PRAGMA_DISABLE_DEPRECATION_WARNINGS
-			return XLongRangeConstraints; 
-			PRAGMA_ENABLE_DEPRECATION_WARNINGS
-		}
 		const TSharedPtr<Softs::FPBDSphericalConstraint>& GetMaximumDistanceConstraints() const { return MaximumDistanceConstraints; }
 		const TSharedPtr<Softs::FPBDSphericalBackstopConstraint>& GetBackstopConstraints() const { return BackstopConstraints; }
 		const TSharedPtr<Softs::FPBDAnimDriveConstraint>& GetAnimDriveConstraints() const { return AnimDriveConstraints; }
@@ -115,19 +171,59 @@ namespace Chaos
 		// ---- End of debug functions ----
 
 	private:
-		TSharedPtr<Softs::FPBDSpringConstraints> EdgeConstraints;
-		TSharedPtr<Softs::FXPBDSpringConstraints> XEdgeConstraints;
-		TSharedPtr<Softs::FPBDSpringConstraints> BendingConstraints;
-		TSharedPtr<Softs::FXPBDSpringConstraints> XBendingConstraints;
+		void CreateSelfCollisionConstraints(
+			const Softs::FPropertyCollectionConstAdapter& ConfigProperties,
+			const FTriangleMesh& TriangleMesh);
+		void CreateEdgeConstraints(
+			const Softs::FPropertyCollectionConstAdapter& ConfigProperties,
+			const TArray<TConstArrayView<FRealSingle>>& WeightMaps,
+			const FTriangleMesh& TriangleMesh);
+		void CreateBendingConstraints(
+			const Softs::FPropertyCollectionConstAdapter& ConfigProperties,
+			const TArray<TConstArrayView<FRealSingle>>& WeightMaps,
+			const FTriangleMesh& TriangleMesh);
+		void CreateAreaConstraints(
+			const Softs::FPropertyCollectionConstAdapter& ConfigProperties,
+			const TArray<TConstArrayView<FRealSingle>>& WeightMaps,
+			const FTriangleMesh& TriangleMesh);
+		void CreateLongRangeConstraints(
+			const Softs::FPropertyCollectionConstAdapter& ConfigProperties,
+			const TArray<TConstArrayView<FRealSingle>>& WeightMaps,
+			const TArray<TConstArrayView<TTuple<int32, int32, FRealSingle>>>& Tethers,
+			Softs::FSolverReal MeshScale);
+		void CreateMaxDistanceConstraints(
+			const Softs::FPropertyCollectionConstAdapter& ConfigProperties,
+			const TArray<TConstArrayView<FRealSingle>>& WeightMaps,
+			Softs::FSolverReal MeshScale);
+		void CreateBackstopConstraints(
+			const Softs::FPropertyCollectionConstAdapter& ConfigProperties,
+			const TArray<TConstArrayView<FRealSingle>>& WeightMaps,
+			Softs::FSolverReal MeshScale);
+		void CreateAnimDriveConstraints(
+			const Softs::FPropertyCollectionConstAdapter& ConfigProperties,
+			const TArray<TConstArrayView<FRealSingle>>& WeightMaps);
+
+		//~ Begin deprecated constraints
+		TSharedPtr<Softs::FPBDSpringConstraints> EdgeConstraints_Deprecated;
+		TSharedPtr<Softs::FXPBDSpringConstraints> XEdgeConstraints_Deprecated;
+		TSharedPtr<Softs::FPBDSpringConstraints> BendingConstraints_Deprecated;
+		TSharedPtr<Softs::FXPBDSpringConstraints> XBendingConstraints_Deprecated;
+		TSharedPtr<Softs::FPBDAxialSpringConstraints> AreaConstraints_Deprecated;
+		TSharedPtr<Softs::FXPBDAxialSpringConstraints> XAreaConstraints_Deprecated;
+		TSharedPtr<Softs::FPBDSpringConstraints> ThinShellVolumeConstraints_Deprecated;;
+		TSharedPtr<Softs::FPBDVolumeConstraint> VolumeConstraints_Deprecated;
+		TSharedPtr<Softs::FXPBDLongRangeConstraints> XLongRangeConstraints_Deprecated;
+		//~ End deprecated constraints
+
+		TSharedPtr<Softs::FPBDEdgeSpringConstraints> EdgeConstraints;
+		TSharedPtr<Softs::FXPBDEdgeSpringConstraints> XEdgeConstraints;
+		TSharedPtr<Softs::FPBDBendingSpringConstraints> BendingConstraints;
+		TSharedPtr<Softs::FXPBDBendingSpringConstraints> XBendingConstraints;
 		TSharedPtr<Softs::FPBDBendingConstraints> BendingElementConstraints;
 		TSharedPtr<Softs::FXPBDBendingConstraints> XBendingElementConstraints;
-		TSharedPtr<Softs::FPBDAxialSpringConstraints> AreaConstraints;
-		TSharedPtr<Softs::FXPBDAxialSpringConstraints> XAreaConstraints;
-		TSharedPtr<Softs::FPBDSpringConstraints> ThinShellVolumeConstraints;
-		TSharedPtr<Softs::FPBDVolumeConstraint> VolumeConstraints;
+		TSharedPtr<Softs::FPBDAreaSpringConstraints> AreaConstraints;
+		TSharedPtr<Softs::FXPBDAreaSpringConstraints> XAreaConstraints;
 		TSharedPtr<Softs::FPBDLongRangeConstraints> LongRangeConstraints; 
-		UE_DEPRECATED(5.1, "LongRange XPBDConstraints are not supported.")
-		TSharedPtr<Softs::FXPBDLongRangeConstraints> XLongRangeConstraints;
 		TSharedPtr<Softs::FPBDSphericalConstraint> MaximumDistanceConstraints;
 		TSharedPtr<Softs::FPBDSphericalBackstopConstraint> BackstopConstraints;
 		TSharedPtr<Softs::FPBDAnimDriveConstraint> AnimDriveConstraints;
@@ -138,7 +234,7 @@ namespace Chaos
 		
 		Softs::FPBDEvolution* Evolution;
 		const TArray<Softs::FSolverVec3>* AnimationPositions;
-		const TArray<Softs::FSolverVec3>* OldAnimationPositions_deprecated;
+		const TArray<Softs::FSolverVec3>* OldAnimationPositions_Deprecated;
 		const TArray<Softs::FSolverVec3>* AnimationNormals;
 		const TArray<Softs::FSolverVec3>* AnimationVelocities;
 
diff --git a/Engine/Plugins/Experimental/ChaosCloth/Source/Chaos/Public/ChaosCloth/ChaosClothingSimulation.h b/Engine/Plugins/Experimental/ChaosCloth/Source/Chaos/Public/ChaosCloth/ChaosClothingSimulation.h
--- a/Engine/Plugins/Experimental/ChaosCloth/Source/Chaos/Public/ChaosCloth/ChaosClothingSimulation.h
+++ b/Engine/Plugins/Experimental/ChaosCloth/Source/Chaos/Public/ChaosCloth/ChaosClothingSimulation.h
@@ -19,6 +19,7 @@ namespace Chaos
 	class FClothingSimulationMesh;
 	class FClothingSimulationCloth;
 	class FClothingSimulationCollider;
+	class FClothingSimulationConfig;
 	class FSkeletalMeshCacheAdapter;
 
 	typedef FClothingSimulationContextCommon FClothingSimulationContext;
@@ -118,8 +119,6 @@ namespace Chaos
 		void ResetStats();
 		void UpdateStats(const FClothingSimulationCloth* Cloth);
 
-		void UpdateSimulationFromSharedSimConfig();
-
 	private:
 		// Visualization object
 		FClothVisualization Visualization;
@@ -131,6 +130,7 @@ PRAGMA_DISABLE_DEPRECATION_WARNINGS  // TODO: CHAOS_IS_CLOTHINGSIMULATIONMESH_AB
 PRAGMA_ENABLE_DEPRECATION_WARNINGS
 		TArray<TUniquePtr<FClothingSimulationCloth>> Cloths;
 		TArray<TUniquePtr<FClothingSimulationCollider>> Colliders;
+		TArray<TUniquePtr<FClothingSimulationConfig>> Configs;
 
 		// External collision Data
 		FClothCollisionData ExternalCollisionData;
diff --git a/Engine/Plugins/Experimental/ChaosCloth/Source/Chaos/Public/ChaosCloth/ChaosClothingSimulationCloth.h b/Engine/Plugins/Experimental/ChaosCloth/Source/Chaos/Public/ChaosCloth/ChaosClothingSimulationCloth.h
--- a/Engine/Plugins/Experimental/ChaosCloth/Source/Chaos/Public/ChaosCloth/ChaosClothingSimulationCloth.h
+++ b/Engine/Plugins/Experimental/ChaosCloth/Source/Chaos/Public/ChaosCloth/ChaosClothingSimulationCloth.h
@@ -7,11 +7,14 @@
 #include "Containers/ContainersFwd.h"
 #include "ChaosCloth/ChaosClothConstraints.h"
 
+struct FManagedArrayCollection;
+
 namespace Chaos
 {
 	class FClothingSimulationSolver;
 	class FClothingSimulationMesh;
 	class FClothingSimulationCollider;
+	class FClothingSimulationConfig;
 
 	// Cloth simulation node
 	class CHAOSCLOTH_API FClothingSimulationCloth final
@@ -27,6 +30,15 @@ namespace Chaos
 		typedef FClothConstraints::ETetherMode ETetherMode;
 
 		FClothingSimulationCloth(
+			FClothingSimulationConfig* InConfig,
+PRAGMA_DISABLE_DEPRECATION_WARNINGS  // TODO: CHAOS_IS_CLOTHINGSIMULATIONMESH_ABSTRACT
+			FClothingSimulationMesh* InMesh,
+PRAGMA_ENABLE_DEPRECATION_WARNINGS
+			TArray<FClothingSimulationCollider*>&& InColliders,
+			uint32 InGroupId);
+
+		UE_DEPRECATED(5.2, "Use config based constructor instead.")
+		FClothingSimulationCloth(
 PRAGMA_DISABLE_DEPRECATION_WARNINGS  // TODO: CHAOS_IS_CLOTHINGSIMULATIONMESH_ABSTRACT
 			FClothingSimulationMesh* InMesh,
 PRAGMA_ENABLE_DEPRECATION_WARNINGS
@@ -61,7 +73,7 @@ PRAGMA_ENABLE_DEPRECATION_WARNINGS
 			FRealSingle InFictitiousAngularScale,
 			const TVec2<FRealSingle>& InDrag,
 			const TVec2<FRealSingle>& InLift,
-			bool bInUseLegacyWind,
+			bool bInUsePointBasedWindModel,
 			const TVec2<FRealSingle>& InPressure,
 			FRealSingle InDampingCoefficient,
 			FRealSingle InLocalDampingCoefficient,
@@ -93,17 +105,17 @@ PRAGMA_ENABLE_DEPRECATION_WARNINGS
 		// ---- Animatable property setters ----
 		void SetMaxDistancesMultiplier(FRealSingle InMaxDistancesMultiplier) { MaxDistancesMultiplier = InMaxDistancesMultiplier; }
 
-		void SetMaterialProperties(const TVec2<FRealSingle>& InEdgeStiffness, const TVec2<FRealSingle>& InBendingStiffness, const TVec2<FRealSingle>& InAreaStiffness) { EdgeStiffness = InEdgeStiffness; BendingStiffness = InBendingStiffness; AreaStiffness = InAreaStiffness; }
-		void SetLongRangeAttachmentProperties(const TVec2<FRealSingle>& InTetherStiffness, const TVec2<FRealSingle>& InTetherScale) { TetherStiffness = InTetherStiffness; TetherScale = InTetherScale;  }
-		void SetCollisionProperties(FRealSingle InCollisionThickness, FRealSingle InFrictionCoefficient, bool bInUseCCD, FRealSingle InSelfCollisionThickness) { CollisionThickness = InCollisionThickness; FrictionCoefficient = InFrictionCoefficient; bUseCCD = bInUseCCD; SelfCollisionThickness = InSelfCollisionThickness; }
-		void SetBackstopProperties(bool bInEnableBackstop) { bEnableBackstop = bInEnableBackstop; }
-		void SetDampingProperties(FRealSingle InDampingCoefficient, FRealSingle InLocalDampingCoefficient = 0.f) { DampingCoefficient = InDampingCoefficient; LocalDampingCoefficient = InLocalDampingCoefficient; }
-		void SetAerodynamicsProperties(const TVec2<FRealSingle>& InDrag, const TVec2<FRealSingle>& InLift, FRealSingle InAirDensity, const FVec3& InWindVelocity) { Drag = InDrag; Lift = InLift; InAirDensity = AirDensity; WindVelocity = InWindVelocity; }
-		void SetPressureProperties(const TVec2<FRealSingle>& InPressure) { Pressure = InPressure; }
-		void SetGravityProperties(FRealSingle InGravityScale, bool bInIsGravityOverridden, const FVec3& InGravityOverride) { GravityScale = InGravityScale; bIsGravityOverridden = bInIsGravityOverridden; GravityOverride = InGravityOverride; }
-		void SetAnimDriveProperties(const TVec2<FRealSingle>& InAnimDriveStiffness, const TVec2<FRealSingle>& InAnimDriveDamping) { AnimDriveStiffness = InAnimDriveStiffness; AnimDriveDamping = InAnimDriveDamping; }
-		void GetAnimDriveProperties(TVec2<FRealSingle>& OutAnimDriveStiffness, TVec2<FRealSingle>& OutAnimDriveDamping) { OutAnimDriveStiffness = AnimDriveStiffness; OutAnimDriveDamping = AnimDriveDamping; }
-		void SetVelocityScaleProperties(const FVec3& InLinearVelocityScale, FRealSingle InAngularVelocityScale, FRealSingle InFictitiousAngularScale) { LinearVelocityScale = InLinearVelocityScale; AngularVelocityScale = InAngularVelocityScale; FictitiousAngularScale = InFictitiousAngularScale;  }
+		void SetMaterialProperties(const TVec2<FRealSingle>& InEdgeStiffness, const TVec2<FRealSingle>& InBendingStiffness, const TVec2<FRealSingle>& InAreaStiffness);
+		void SetLongRangeAttachmentProperties(const TVec2<FRealSingle>& InTetherStiffness, const TVec2<FRealSingle>& InTetherScale);
+		void SetCollisionProperties(FRealSingle InCollisionThickness, FRealSingle InFrictionCoefficient, bool bInUseCCD, FRealSingle InSelfCollisionThickness);
+		void SetBackstopProperties(bool bInEnableBackstop);
+		void SetDampingProperties(FRealSingle InDampingCoefficient, FRealSingle InLocalDampingCoefficient = 0.f);
+		void SetAerodynamicsProperties(const TVec2<FRealSingle>& InDrag, const TVec2<FRealSingle>& InLift, FRealSingle InAirDensity, const FVec3& InWindVelocity);
+		void SetPressureProperties(const TVec2<FRealSingle>& InPressure);
+		void SetGravityProperties(FRealSingle InGravityScale, bool bInUseGravityOverride, const FVec3& InGravityOverride);
+		void SetAnimDriveProperties(const TVec2<FRealSingle>& InAnimDriveStiffness, const TVec2<FRealSingle>& InAnimDriveDamping);
+		void GetAnimDriveProperties(TVec2<FRealSingle>& OutAnimDriveStiffness, TVec2<FRealSingle>& OutAnimDriveDamping);
+		void SetVelocityScaleProperties(const FVec3& InLinearVelocityScale, FRealSingle InAngularVelocityScale, FRealSingle InFictitiousAngularScale);
 
 		void Reset() { bNeedsReset = true; }
 		void Teleport() { bNeedsTeleport = true; }
@@ -115,6 +127,9 @@ PRAGMA_ENABLE_DEPRECATION_WARNINGS
 		void SetMesh(FClothingSimulationMesh* InMesh);
 		PRAGMA_ENABLE_DEPRECATION_WARNINGS
 
+		FClothingSimulationConfig* GetConfig() const { return Config; }
+		void SetConfig(FClothingSimulationConfig* InConfig);
+
 		const TArray<FClothingSimulationCollider*>& GetColliders() const { return Colliders; }
 		void SetColliders(TArray<FClothingSimulationCollider*>&& InColliders);
 		void AddCollider(FClothingSimulationCollider* InCollider);
@@ -173,54 +188,14 @@ PRAGMA_ENABLE_DEPRECATION_WARNINGS
 PRAGMA_DISABLE_DEPRECATION_WARNINGS  // TODO: CHAOS_IS_CLOTHINGSIMULATIONMESH_ABSTRACT
 		FClothingSimulationMesh* Mesh;
 PRAGMA_ENABLE_DEPRECATION_WARNINGS
+		FClothingSimulationConfig* Config;
 		TArray<FClothingSimulationCollider*> Colliders;
 		uint32 GroupId;
-		EMassMode MassMode;
-		FRealSingle MassValue;
-		FRealSingle MinPerParticleMass;
-		TVec2<FRealSingle> EdgeStiffness;
-		TVec2<FRealSingle> EdgeDampingRatio;
-		TVec2<FRealSingle> BendingStiffness;
-		TVec2<FRealSingle> BendingDampingRatio;
-		FRealSingle BucklingRatio;
-		TVec2<FRealSingle> BucklingStiffness;
-		bool bUseBendingElements;
-		TVec2<FRealSingle> AreaStiffness;
-		FRealSingle VolumeStiffness;
-		bool bUseThinShellVolumeConstraints;
-		TVec2<FRealSingle> TetherStiffness;
-		TVec2<FRealSingle> TetherScale;
-		ETetherMode TetherMode;
-		FRealSingle MaxDistancesMultiplier;  // Animatable
-		TVec2<FRealSingle> AnimDriveStiffness;  // Animatable
-		TVec2<FRealSingle> AnimDriveDamping;  // Animatable
-		FRealSingle ShapeTargetStiffness;
-		bool bUseXPBDEdgeConstraints;
-		bool bUseXPBDBendingConstraints;
-		bool bUseXPBDAreaConstraints;
-		FRealSingle GravityScale;
-		bool bIsGravityOverridden;
-		TVec3<FRealSingle> GravityOverride;
-		TVec3<FRealSingle> LinearVelocityScale;  // Linear ratio applied to the reference bone transforms
-		FRealSingle AngularVelocityScale;  // Angular ratio factor applied to the reference bone transforms
-		FRealSingle FictitiousAngularScale;
-		TVec2<FRealSingle> Drag;
-		TVec2<FRealSingle> Lift;
-		FVec3 WindVelocity;
-		FRealSingle AirDensity;
-		bool bUseLegacyWind;
-		TVec2<FRealSingle> Pressure;
-		FRealSingle DampingCoefficient;
-		FRealSingle LocalDampingCoefficient;
-		FRealSingle CollisionThickness;
-		FRealSingle FrictionCoefficient;
-		bool bUseCCD;
-		bool bUseSelfCollisions;
-		FRealSingle SelfCollisionThickness;
-		FRealSingle SelfCollisionFrictionCoefficient;
-		bool bUseSelfIntersections;
-		bool bEnableBackstop;
-		bool bUseLegacyBackstop;
+
+		TSharedPtr<FManagedArrayCollection> PropertyCollection;  // Used for backward compatibility only, otherwise the properties are owned by the Config
+
+		FRealSingle MaxDistancesMultiplier;  // Legacy multiplier
+
 		bool bUseLODIndexOverride;
 		int32 LODIndexOverride;
 		bool bNeedsReset;
diff --git a/Engine/Plugins/Experimental/ChaosCloth/Source/Chaos/Public/ChaosCloth/ChaosClothingSimulationSolver.h b/Engine/Plugins/Experimental/ChaosCloth/Source/Chaos/Public/ChaosCloth/ChaosClothingSimulationSolver.h
--- a/Engine/Plugins/Experimental/ChaosCloth/Source/Chaos/Public/ChaosCloth/ChaosClothingSimulationSolver.h
+++ b/Engine/Plugins/Experimental/ChaosCloth/Source/Chaos/Public/ChaosCloth/ChaosClothingSimulationSolver.h
@@ -19,6 +19,7 @@ namespace Chaos
 	class FClothingSimulationCloth;
 	class FClothingSimulationCollider;
 	class FClothingSimulationMesh;
+	class FClothingSimulationConfig;
 
 	// Solver simulation node
 	class CHAOSCLOTH_API FClothingSimulationSolver final : public FPhysicsSolverEvents
@@ -49,12 +50,13 @@ namespace Chaos
 		void SetWindVelocity(const TVec3<FRealSingle>& InWindVelocity, FRealSingle InLegacyWindAdaption = (FRealSingle)0.);
 		const TVec3<FRealSingle>& GetWindVelocity() const { return WindVelocity; }
 
-		void SetNumIterations(int32 InNumIterations) { NumIterations = InNumIterations; }
-		int32 GetNumIterations() const { return NumIterations; }
-		void SetMaxNumIterations(int32 InMaxNumIterations) { MaxNumIterations = InMaxNumIterations; }
-		int32 GetMaxNumIterations() const { return MaxNumIterations; }
-		void SetNumSubsteps(int32 InNumSubsteps) { NumSubsteps = InNumSubsteps; }
-		int32 GetNumSubsteps() const { return NumSubsteps; }
+		void SetNumIterations(int32 InNumIterations);
+		int32 GetNumIterations() const;
+		void SetMaxNumIterations(int32 InMaxNumIterations);
+		int32 GetMaxNumIterations() const;
+		void SetNumSubsteps(int32 InNumSubsteps);
+		int32 GetNumSubsteps() const;
+
 		void SetEnableSolver(bool InbEnableSolver) { bEnableSolver = InbEnableSolver; }
 		bool GetEnableSolver() const { return bEnableSolver; }
 		// ---- End of the animatable property setters ----
@@ -70,6 +72,16 @@ namespace Chaos
 
 		TConstArrayView<const FClothingSimulationCloth*> GetCloths() const { return Cloths; }
 
+		/** Get the solver configuration. */
+		FClothingSimulationConfig* GetConfig() const { return Config; }
+
+		/**
+		 * Set the solver configuration.
+		 * Can use a cloth config if a single cloth is being simulated.
+		 * When a solver configuration is set, the solver shared properties (e.g. NumIterations, MaxNumIterations, NumSubsteps) are overriden.
+		 */
+		void SetConfig(FClothingSimulationConfig* InConfig) { Config = InConfig; }
+
 		/** Advance the simulation. */
 		void Update(FSolverReal InDeltaTime);
 
@@ -223,6 +235,9 @@ namespace Chaos
 
 		// Object arrays
 		TArray<FClothingSimulationCloth*> Cloths;
+		
+		// Solver config
+		FClothingSimulationConfig* Config = nullptr;
 
 		// Simulation group attributes
 		TArrayCollectionArray<Softs::FSolverRigidTransform3> PreSimulationTransforms;  // Allow a different frame of reference for each cloth groups
@@ -257,6 +272,8 @@ namespace Chaos
 		// Time stepping
 		FSolverReal Time;
 		FSolverReal DeltaTime;
+
+		// Solver config properties, superseded by the config pointer
 		int32 NumIterations;
 		int32 MaxNumIterations;
 		int32 NumSubsteps;
diff --git a/Engine/Plugins/Experimental/ChaosClothAsset/Source/ChaosClothAssetEngine/Private/ChaosClothAsset/ClothComponent.cpp b/Engine/Plugins/Experimental/ChaosClothAsset/Source/ChaosClothAssetEngine/Private/ChaosClothAsset/ClothComponent.cpp
--- a/Engine/Plugins/Experimental/ChaosClothAsset/Source/ChaosClothAssetEngine/Private/ChaosClothAsset/ClothComponent.cpp
+++ b/Engine/Plugins/Experimental/ChaosClothAsset/Source/ChaosClothAssetEngine/Private/ChaosClothAsset/ClothComponent.cpp
@@ -3,8 +3,10 @@
 #include "ChaosClothAsset/ClothComponent.h"
 #include "ChaosClothAsset/ClothAsset.h"
 #include "ChaosClothAsset/ClothAssetPrivate.h"
+#include "ChaosClothAsset/ClothCollection.h"
 #include "ChaosClothAsset/ClothSimulationModel.h"
 #include "ChaosClothAsset/ClothSimulationProxy.h"
+#include "Chaos/PropertyCollectionAdapter.h"
 #include "HAL/IConsoleManager.h"
 #include "Rendering/SkeletalMeshRenderData.h"
 #include "SkeletalRenderPublic.h"
@@ -104,6 +106,10 @@ void UChaosClothComponent::OnRegister()
 				bHasValidBoneTransform = true;
 			}
 
+			// Create cloth property collection, and fill up the collection with the original cloth asset properties
+			PropertyCollection = MakeShared<FManagedArrayCollection>();
+			ResetConfigProperties();
+
 			// Create simulation proxy
 			ClothSimulationProxy = MakeUnique<FClothSimulationProxy>(*this);
 		}
@@ -117,6 +123,26 @@ bool UChaosClothComponent::IsSimulationSuspended() const
 	return bSuspendSimulation || !ClothSimulationProxy.IsValid() || (CVarClothPhysics && !CVarClothPhysics->GetBool());
 }
 
+void UChaosClothComponent::ResetConfigProperties()
+{
+	using namespace UE::Chaos::ClothAsset;
+	using namespace ::Chaos::Softs;
+
+	PropertyCollectionAdapter.Reset();
+
+	if (GetClothAsset())
+	{
+		if (TSharedPtr<const FClothCollection> ClothCollection = GetClothAsset()->GetClothCollection())
+		{
+			// Create a mutable adapter for our component's property collection, and use it to copy the properties from the cloth collection
+			FPropertyCollectionMutableAdapter(PropertyCollection).Copy(TSharedPtr<const FManagedArrayCollection>(ClothCollection));
+
+			// The Adapter used to access the properties shouldn't be mutable, since adding/removing properties would cause issues
+			PropertyCollectionAdapter = MakeUnique<FPropertyCollectionAdapter>(PropertyCollection);
+		}
+	}
+}
+
 void UChaosClothComponent::OnUnregister()
 {
 	Super::OnUnregister();
diff --git a/Engine/Plugins/Experimental/ChaosClothAsset/Source/ChaosClothAssetEngine/Private/ChaosClothAsset/ClothSimulationProxy.cpp b/Engine/Plugins/Experimental/ChaosClothAsset/Source/ChaosClothAssetEngine/Private/ChaosClothAsset/ClothSimulationProxy.cpp
--- a/Engine/Plugins/Experimental/ChaosClothAsset/Source/ChaosClothAssetEngine/Private/ChaosClothAsset/ClothSimulationProxy.cpp
+++ b/Engine/Plugins/Experimental/ChaosClothAsset/Source/ChaosClothAssetEngine/Private/ChaosClothAsset/ClothSimulationProxy.cpp
@@ -8,6 +8,7 @@
 #include "ChaosClothAsset/ClothSimulationModel.h"
 #include "ChaosCloth/ChaosClothingSimulationCloth.h"
 #include "ChaosCloth/ChaosClothingSimulationCollider.h"
+#include "ChaosCloth/ChaosClothingSimulationConfig.h"
 #include "ChaosCloth/ChaosClothingSimulationSolver.h"
 #include "ChaosCloth/ChaosClothVisualization.h"
 #include "PhysicsEngine/PhysicsSettings.h"
@@ -96,7 +97,7 @@ namespace UE::Chaos::ClothAsset
 		constexpr bool bIsInitialization = true;
 		ClothSimulationContext.Fill(ClothComponent, 0.f, MaxDeltaTime, bIsInitialization);
 
-		// Create mesh node
+		// Create mesh simulation thread object
 		const UChaosClothAsset* const ClothAsset = ClothComponent.GetClothAsset();
 		if (!ClothAsset)
 		{
@@ -113,70 +114,33 @@ namespace UE::Chaos::ClothAsset
 #endif
 		const int32 MeshIndex = Meshes.Emplace(MakeUnique<FClothSimulationMesh>(*ClothSimulationModel, ClothSimulationContext, DebugName));
 
-		// Create collider node
+		// Create collider simulation thread object
 PRAGMA_DISABLE_DEPRECATION_WARNINGS  // TODO: CHAOS_IS_CLOTHINGSIMULATIONMESH_ABSTRACT
 		const FReferenceSkeleton* const ReferenceSkeleton = &ClothAsset->GetRefSkeleton();
 		const int32 ColliderIndex = ClothComponent.GetPhysicsAsset() ? Colliders.Emplace(MakeUnique<FClothingSimulationCollider>(ClothComponent.GetPhysicsAsset(), ReferenceSkeleton)) : INDEX_NONE;
 PRAGMA_ENABLE_DEPRECATION_WARNINGS
 
-		//// Set the external collision data to get updated at every frame
-		//Colliders[ColliderIndex]->SetCollisionData(&ExternalCollisionData);
+		//Colliders[ColliderIndex]->SetCollisionData(&ExternalCollisionData);  // TODO: External collision data
+
+		// Create cloth config simulation thread object
+		const int32 ClothConfigIndex = Configs.Emplace(MakeUnique<FClothingSimulationConfig>(ClothComponent.GetPropertyCollection()));
+
+		// Create cloth simulation thread object
+		constexpr uint32 GroupId = 0;
 
-		// Create cloth node
 		const int32 ClothIndex = Cloths.Emplace(MakeUnique<FClothingSimulationCloth>(
+			Configs[ClothConfigIndex].Get(),
 PRAGMA_DISABLE_DEPRECATION_WARNINGS  // TODO: CHAOS_IS_CLOTHINGSIMULATIONMESH_ABSTRACT
 			Meshes[MeshIndex].Get(),
 PRAGMA_ENABLE_DEPRECATION_WARNINGS
 			ColliderIndex != INDEX_NONE ? TArray<FClothingSimulationCollider*>({ Colliders[ColliderIndex].Get() }) : TArray<FClothingSimulationCollider*>(),
-			0,
-			FClothingSimulationCloth::EMassMode::Density,  //	(FClothingSimulationCloth::EMassMode)ClothConfig->MassMode,
-			0.35f,  //	ClothConfig->GetMassValue(),
-			0.0001f,  //	ClothConfig->MinPerParticleMass,
-			TVec2<FRealSingle>(1.f),  //	TVec2<FRealSingle>(ClothConfig->EdgeStiffnessWeighted.Low, ClothConfig->EdgeStiffnessWeighted.High),
-			TVec2<FRealSingle>(1.f),  //	TVec2<FRealSingle>(ClothConfig->BendingStiffnessWeighted.Low, ClothConfig->BendingStiffnessWeighted.High),
-			0.f,  //	ClothConfig->BucklingRatio,
-			TVec2<FRealSingle>(1.f),  //	TVec2<FRealSingle>(ClothConfig->BucklingStiffnessWeighted.Low, ClothConfig->BucklingStiffnessWeighted.High),
-			false,  //	ClothConfig->bUseBendingElements,
-			TVec2<FRealSingle>(1.f),  //	TVec2<FRealSingle>(ClothConfig->AreaStiffnessWeighted.Low, ClothConfig->AreaStiffnessWeighted.High),
-			0.f,  //	ClothConfig->VolumeStiffness,
-			false,  //	ClothConfig->bUseThinShellVolumeConstraints,
-			TVec2<FRealSingle>(1.f),   //	TVec2<FRealSingle>(ClothConfig->TetherStiffness.Low, ClothConfig->TetherStiffness.High),  // Animatable
-			TVec2<FRealSingle>(1.f),   //	TVec2<FRealSingle>(ClothConfig->TetherScale.Low, ClothConfig->TetherScale.High),  // Animatable
-			FClothingSimulationCloth::ETetherMode::Geodesic,  //	ClothConfig->bUseGeodesicDistance ? FClothingSimulationCloth::ETetherMode::Geodesic : FClothingSimulationCloth::ETetherMode::Euclidean,
-			1.f,  //	/*MaxDistancesMultiplier =*/ 1.f,  // Animatable
-			TVec2<FRealSingle>(0.f),  //	TVec2<FRealSingle>(ClothConfig->AnimDriveStiffness.Low, ClothConfig->AnimDriveStiffness.High),  // Animatable
-			TVec2<FRealSingle>(0.f),  //	TVec2<FRealSingle>(ClothConfig->AnimDriveDamping.Low, ClothConfig->AnimDriveDamping.High),  // Animatable
-			0.f,  //	ClothConfig->ShapeTargetStiffness,  // TODO: This is now deprecated
-			false,  //	false, // bUseXPBDEdgeSprings
-			false,  //	false, // bUseXPBDBendingElements
-			false,  //	false, // bUseXPBDAreaSprings
-			.5f,   //	ClothConfig->GravityScale,
-			false,  //	ClothConfig->bUseGravityOverride,
-			TVec3<FRealSingle>(0.f, 0.f, -980.665f),  //	ClothConfig->Gravity,
-			TVec3<FRealSingle>(0.75f),  //	ClothConfig->LinearVelocityScale,
-			0.75f,  //	ClothConfig->AngularVelocityScale,
-			1.f,  //	ClothConfig->FictitiousAngularScale,
-			TVec2<FRealSingle>(0.035f, 1.f),  //	TVec2<FRealSingle>(ClothConfig->Drag.Low, ClothConfig->Drag.High),  // Animatable
-			TVec2<FRealSingle>(0.035f, 1.f),  //	TVec2<FRealSingle>(ClothConfig->Lift.Low, ClothConfig->Lift.High),  // Animatable
-			false,  //	ClothConfig->bUsePointBasedWindModel,
-			TVec2<FRealSingle>(0.f),  //	TVec2<FRealSingle>(ClothConfig->Pressure.Low, ClothConfig->Pressure.High),  // Animatable
-			0.01f,  //	ClothConfig->DampingCoefficient,
-			0.f,  //	ClothConfig->LocalDampingCoefficient,
-			1.f,  //	ClothConfig->CollisionThickness,
-			0.8f,  //	ClothConfig->FrictionCoefficient,
-			false,  //	ClothConfig->bUseCCD,
-			false,  //	ClothConfig->bUseSelfCollisions,
-			2.f,  //	ClothConfig->SelfCollisionThickness,
-			0.f,  //	ClothConfig->SelfCollisionFriction,
-			false,  //	ClothConfig->bUseSelfIntersections,
-			false,  //	ClothConfig->bUseLegacyBackstop,
-			false,  //	/*bUseLODIndexOverride =*/ false,
-			INDEX_NONE));  //	/*LODIndexOverride =*/ INDEX_NONE));
-		
-		Cloths[ClothIndex]->SetAerodynamicsProperties(TVec2<FRealSingle>(0.5f, 1.f), TVec2<FRealSingle>(0.35f, 1.f), 1.225e-6f, TVec3<FRealSingle>(-30.f, -30.f, 200.f));
-
-		// Add cloth to solver
+			GroupId));
+
 		Solver->AddCloth(Cloths[ClothIndex].Get());
+
+		// Create solver config simulation thread object
+		const int32 SolverConfigIndex = Configs.Emplace(MakeUnique<FClothingSimulationConfig>(ClothComponent.GetPropertyCollection()));  // TODO: Use a separate solver config for outfits
+		Solver->SetConfig(Configs[SolverConfigIndex].Get());
 	}
 
 	FClothSimulationProxy::~FClothSimulationProxy()
@@ -193,6 +157,12 @@ PRAGMA_ENABLE_DEPRECATION_WARNINGS
 
 		if (DeltaTime > 0.f && !ClothComponent.IsSimulationSuspended())
 		{
+			// Replace physics thread's configs with the game thread's configs
+			for (const TUniquePtr<::Chaos::FClothingSimulationConfig>& Config : Configs)
+			{
+				Config->Initialize(ClothComponent.GetPropertyCollection());  // TODO: Outfit and multi-cloths/solver config update
+			}
+
 			// Start the the cloth simulation thread
 			ParallelTask = TGraphTask<FClothSimulationProxyParallelTask>::CreateTask(nullptr, ENamedThreads::GameThread).ConstructAndDispatchWhenReady(*this);
 		}
@@ -226,9 +196,6 @@ PRAGMA_ENABLE_DEPRECATION_WARNINGS
 		Solver->SetLocalSpaceRotation((FQuat)ClothSimulationContext.ComponentTransform.GetRotation());
 		Solver->SetWindVelocity(ClothSimulationContext.WindVelocity);
 		Solver->SetGravity(ClothSimulationContext.WorldGravity);
-		Solver->SetNumIterations(8);  // TODO: Solver parameters
-		Solver->SetMaxNumIterations(10);  // TODO: Solver parameters
-		Solver->SetNumSubsteps(1);  // TODO: Solver parameters
 
 		// Check teleport modes
 		for (const TUniquePtr<FClothingSimulationCloth>& Cloth : Cloths)
diff --git a/Engine/Plugins/Experimental/ChaosClothAsset/Source/ChaosClothAssetEngine/Private/ChaosClothAsset/ClothSimulationProxy.h b/Engine/Plugins/Experimental/ChaosClothAsset/Source/ChaosClothAssetEngine/Private/ChaosClothAsset/ClothSimulationProxy.h
--- a/Engine/Plugins/Experimental/ChaosClothAsset/Source/ChaosClothAssetEngine/Private/ChaosClothAsset/ClothSimulationProxy.h
+++ b/Engine/Plugins/Experimental/ChaosClothAsset/Source/ChaosClothAssetEngine/Private/ChaosClothAsset/ClothSimulationProxy.h
@@ -14,6 +14,7 @@ namespace Chaos
 	class FClothingSimulationSolver;
 	class FClothingSimulationMesh;
 	class FClothingSimulationCloth;
+	class FClothingSimulationConfig;
 	class FClothingSimulationCollider;
 	class FClothVisualization;
 }
@@ -84,6 +85,7 @@ PRAGMA_DISABLE_DEPRECATION_WARNINGS  // TODO: CHAOS_IS_CLOTHINGSIMULATIONMESH_AB
 		TArray<TUniquePtr<::Chaos::FClothingSimulationMesh>> Meshes;
 PRAGMA_ENABLE_DEPRECATION_WARNINGS
 		TArray<TUniquePtr<::Chaos::FClothingSimulationCloth>> Cloths;
+		TArray<TUniquePtr<::Chaos::FClothingSimulationConfig>> Configs;
 		TArray<TUniquePtr<::Chaos::FClothingSimulationCollider>> Colliders;
 		TUniquePtr<::Chaos::FClothVisualization> Visualization;
 
diff --git a/Engine/Plugins/Experimental/ChaosClothAsset/Source/ChaosClothAssetEngine/Public/ChaosClothAsset/ClothComponent.h b/Engine/Plugins/Experimental/ChaosClothAsset/Source/ChaosClothAssetEngine/Public/ChaosClothAsset/ClothComponent.h
--- a/Engine/Plugins/Experimental/ChaosClothAsset/Source/ChaosClothAssetEngine/Public/ChaosClothAsset/ClothComponent.h
+++ b/Engine/Plugins/Experimental/ChaosClothAsset/Source/ChaosClothAssetEngine/Public/ChaosClothAsset/ClothComponent.h
@@ -3,16 +3,22 @@
 #pragma once
 
 #include "Components/SkinnedMeshComponent.h"
-
 #include "ClothComponent.generated.h"
 
 class UChaosClothAsset;
 class UChaosClothComponent;
+struct FManagedArrayCollection;
+
+namespace Chaos::Softs
+{
+	class FPropertyCollectionAdapter;
+}
 
 namespace UE::Chaos::ClothAsset
 {
 	class FClothSimulationProxy;
 }
+
 /**
  * Cloth simulation component.
  */
@@ -67,6 +73,17 @@ public:
 	UFUNCTION(BlueprintCallable, Category = "Components|Simulation", Meta = (Keywords = "Chaos Cloth Simulation Suspend"))
 	bool IsSimulationSuspended() const;
 
+	/** Reset all cloth simulation config properties to the values stored in the original cloth asset. */
+	UFUNCTION(BlueprintCallable, Category = "Components|Simulation", Meta = (Keywords = "Chaos Cloth Config Property"))
+	void ResetConfigProperties();
+
+	/**
+	 * Return the property collection holding the runtime properties for this cloth component.
+	 * This might be different from the cloth asset's since the component's properties can be modified in code or in blueprints.
+	 * This could also be different from the cloth simulation object until the cloth simulation thread synchronise the properties.
+	 */
+	TSharedPtr<const FManagedArrayCollection> GetPropertyCollection() const { return TSharedPtr<const FManagedArrayCollection>(PropertyCollection); }
+
 protected:
 	//~ Begin UObject Interface
 	virtual void PostLoad() override;
@@ -122,7 +139,6 @@ private:
 	UPROPERTY()
 	uint8 bSuspendSimulation : 1;
 
-
 	/** Whether to use the leader component pose. */
 	UPROPERTY()
 	uint8 bBindToLeaderComponent : 1;
@@ -139,5 +155,8 @@ private:
 	UPROPERTY()
 	float BlendWeight = 1.f;
 
+	TSharedPtr<FManagedArrayCollection> PropertyCollection;
+	TUniquePtr<::Chaos::Softs::FPropertyCollectionAdapter> PropertyCollectionAdapter;
+
 	TUniquePtr<UE::Chaos::ClothAsset::FClothSimulationProxy> ClothSimulationProxy;
 };
diff --git a/Engine/Source/Runtime/Experimental/Chaos/Public/Chaos/VelocityField.h b/Engine/Source/Runtime/Experimental/Chaos/Public/Chaos/VelocityField.h
--- a/Engine/Source/Runtime/Experimental/Chaos/Public/Chaos/VelocityField.h
+++ b/Engine/Source/Runtime/Experimental/Chaos/Public/Chaos/VelocityField.h
@@ -61,7 +61,7 @@ public:
 		}
 	}
 
-	void SetProperties(const FPropertyCollectionConstAdapter& PropertyCollection)
+	void SetProperties(const FPropertyCollectionConstAdapter& PropertyCollection, FSolverReal WorldScale)
 	{
 		if (IsFluidDensityMutable(PropertyCollection))
 		{
@@ -86,8 +86,8 @@ public:
 		if (IsPressureMutable(PropertyCollection))
 		{
 			const FSolverVec2 Pressure(GetWeightedFloatPressure(PropertyCollection));
-			PressureBase = Pressure[0];
-			PressureRange = Pressure[1] - PressureBase;
+			PressureBase = Pressure[0] / WorldScale;
+			PressureRange = Pressure[1] / WorldScale - PressureBase;
 		}
 	}
 
