Choas Cloth - Fixed a crash when opening a ClothAsset in the new ClothEditor that has no properties set.

#rnx

diff --git a/Engine/Plugins/Experimental/ChaosCloth/Source/Chaos/Private/ChaosCloth/ChaosClothConstraints.cpp b/Engine/Plugins/Experimental/ChaosCloth/Source/Chaos/Private/ChaosCloth/ChaosClothConstraints.cpp
--- a/Engine/Plugins/Experimental/ChaosCloth/Source/Chaos/Private/ChaosCloth/ChaosClothConstraints.cpp
+++ b/Engine/Plugins/Experimental/ChaosCloth/Source/Chaos/Private/ChaosCloth/ChaosClothConstraints.cpp
@@ -289,6 +289,9 @@ void FClothConstraints::CreateBendingConstraints(
 				DampingMultipliers,
 				ConfigProperties,
 				/*bTrimKinematicConstraints =*/ true);
+
+			++NumConstraintInits;  // Uses init to update the property tables
+			++NumConstraintRules;
 		}
 		else if (Softs::FPBDBendingSpringConstraints::IsEnabled(ConfigProperties))
 		{
@@ -303,9 +306,10 @@ void FClothConstraints::CreateBendingConstraints(
 				BendingStiffnessMultipliers,
 				ConfigProperties,
 				/*bTrimKinematicConstraints =*/ true);
+
+			++NumConstraintInits;  // Uses init to update the property tables
+			++NumConstraintRules;
 		}
-		++NumConstraintInits;  // Uses init to update the property tables
-		++NumConstraintRules;
 	}
 }
 
