Chaos Cloth - Fix crash when updating the cloth config when a SkeletalMeshComponent using the same cloth asset is running in the background.

#rnx

diff --git a/Engine/Plugins/Experimental/ChaosCloth/Source/Chaos/Private/ChaosCloth/ChaosClothingSimulation.cpp b/Engine/Plugins/Experimental/ChaosCloth/Source/Chaos/Private/ChaosCloth/ChaosClothingSimulation.cpp
--- a/Engine/Plugins/Experimental/ChaosCloth/Source/Chaos/Private/ChaosCloth/ChaosClothingSimulation.cpp
+++ b/Engine/Plugins/Experimental/ChaosCloth/Source/Chaos/Private/ChaosCloth/ChaosClothingSimulation.cpp
@@ -298,6 +298,7 @@ void FClothingSimulation::Shutdown()
 	Meshes.Reset();
 	Cloths.Reset();
 	Colliders.Reset();
+	Configs.Reset();
 	ClothSharedSimConfig = nullptr;
 }
 
@@ -331,6 +332,8 @@ void FClothingSimulation::CreateActor(USkeletalMeshComponent* InOwnerComponent,
 	{
 		ClothSharedSimConfig = Asset->GetClothConfig<UChaosClothSharedSimConfig>();
 
+		UpdateSimulationFromSharedSimConfig();
+
 		// Must set the local space location prior to adding any mesh/cloth, as otherwise the start poses would be in the wrong local space
 		const FClothingSimulationContext* const Context = static_cast<const FClothingSimulationContext*>(InOwnerComponent->GetClothingSimulationContext());
 		check(Context);
@@ -380,11 +383,6 @@ PRAGMA_ENABLE_DEPRECATION_WARNINGS
 	// Add cloth to solver
 	Solver->AddCloth(Cloths[ClothIndex].Get());
 
-	// Create solver config runtime simulation object
-	const int32 SolverConfigIndex = Configs.Emplace(MakeUnique<FClothingSimulationConfig>());
-	Configs[SolverConfigIndex]->Initialize(nullptr, ClothSharedSimConfig);
-	Solver->SetConfig(Configs[SolverConfigIndex].Get());
-
 	// Update stats
 	UpdateStats(Cloths[ClothIndex].Get());
 
@@ -425,6 +423,23 @@ void FClothingSimulation::UpdateStats(const FClothingSimulationCloth* Cloth)
 	NumDynamicParticles += Cloth->GetNumActiveDynamicParticles();
 }
 
+void FClothingSimulation::UpdateSimulationFromSharedSimConfig()
+{
+	check(Solver);
+	if (ClothSharedSimConfig) // ClothSharedSimConfig will be a null pointer if all cloth instances are disabled in which case we will use default Evolution parameters
+	{
+		FClothingSimulationConfig* SolverConfig = Solver->GetConfig();
+		if (!SolverConfig)
+		{
+			// Create solver config runtime simulation object
+			const int32 SolverConfigIndex = Configs.Emplace(MakeUnique<FClothingSimulationConfig>());
+			SolverConfig = Configs[SolverConfigIndex].Get();
+			Solver->SetConfig(SolverConfig);
+		}
+		SolverConfig->Initialize(nullptr, ClothSharedSimConfig);
+	}
+}
+
 void FClothingSimulation::SetNumIterations(int32 InNumIterations)
 {
 	Solver->SetNumIterations(InNumIterations);
@@ -740,6 +755,8 @@ void FClothingSimulation::GetCollisions(FClothCollisionData& OutCollisions, bool
 
 void FClothingSimulation::RefreshClothConfig(const IClothingSimulationContext* InContext)
 {
+	UpdateSimulationFromSharedSimConfig();
+
 	// Update new space location
 	const FClothingSimulationContext* const Context = static_cast<const FClothingSimulationContext*>(InContext);
 	static const bool bReset = true;
@@ -780,12 +797,6 @@ void FClothingSimulation::RefreshClothConfig(const IClothingSimulationContext* I
 		// Update stats
 		UpdateStats(Cloth.Get());
 	}
-
-	// Update solver config runtime simulation object
-	FClothingSimulationConfig* const SolverConfig = Solver->GetConfig();
-	check(SolverConfig);
-	SolverConfig->Initialize(nullptr, ClothSharedSimConfig);
-
 	UE_LOG(LogChaosCloth, VeryVerbose, TEXT("RefreshClothConfig, all constraints and self-collisions have been updated for all clothing assets and LODs."));
 }
 
diff --git a/Engine/Plugins/Experimental/ChaosCloth/Source/Chaos/Public/ChaosCloth/ChaosClothingSimulation.h b/Engine/Plugins/Experimental/ChaosCloth/Source/Chaos/Public/ChaosCloth/ChaosClothingSimulation.h
--- a/Engine/Plugins/Experimental/ChaosCloth/Source/Chaos/Public/ChaosCloth/ChaosClothingSimulation.h
+++ b/Engine/Plugins/Experimental/ChaosCloth/Source/Chaos/Public/ChaosCloth/ChaosClothingSimulation.h
@@ -119,6 +119,8 @@ namespace Chaos
 		void ResetStats();
 		void UpdateStats(const FClothingSimulationCloth* Cloth);
 
+		void UpdateSimulationFromSharedSimConfig();
+
 	private:
 		// Visualization object
 		FClothVisualization Visualization;
