Chaos Cloth - Automatically adjusts MaxIterationCount or IterationCount when one or the other changes.

#rnx

diff --git a/Engine/Plugins/Experimental/ChaosCloth/Source/Chaos/Private/ChaosCloth/ChaosClothConfig.cpp b/Engine/Plugins/Experimental/ChaosCloth/Source/Chaos/Private/ChaosCloth/ChaosClothConfig.cpp
--- a/Engine/Plugins/Experimental/ChaosCloth/Source/Chaos/Private/ChaosCloth/ChaosClothConfig.cpp
+++ b/Engine/Plugins/Experimental/ChaosCloth/Source/Chaos/Private/ChaosCloth/ChaosClothConfig.cpp
@@ -257,6 +257,15 @@ void UChaosClothSharedSimConfig::PostEditChangeChainProperty(FPropertyChangedCha
 {
 	Super::PostEditChangeChainProperty(ChainEvent);
 
+	if (ChainEvent.GetPropertyName() == GET_MEMBER_NAME_CHECKED(UChaosClothSharedSimConfig, MaxIterationCount))
+	{
+		IterationCount = FMath::Min(IterationCount, MaxIterationCount);
+	}
+	else if (ChainEvent.GetPropertyName() == GET_MEMBER_NAME_CHECKED(UChaosClothSharedSimConfig, IterationCount))
+	{
+		MaxIterationCount = FMath::Max(MaxIterationCount, IterationCount);
+	}
+
 	// Update the simulation if there is any interactor attached to the skeletal mesh component
 	if (ChainEvent.ChangeType != EPropertyChangeType::Interactive)
 	{
