Chaos Cloth - Fixed crash in FClothVisualization when the debug materials are loaded.

diff --git a/Engine/Plugins/Experimental/ChaosCloth/Source/Chaos/Public/ChaosCloth/ChaosClothVisualization.h b/Engine/Plugins/Experimental/ChaosCloth/Source/Chaos/Public/ChaosCloth/ChaosClothVisualization.h
--- a/Engine/Plugins/Experimental/ChaosCloth/Source/Chaos/Public/ChaosCloth/ChaosClothVisualization.h
+++ b/Engine/Plugins/Experimental/ChaosCloth/Source/Chaos/Public/ChaosCloth/ChaosClothVisualization.h
@@ -103,9 +103,9 @@ namespace Chaos
 		const ::Chaos::FClothingSimulationSolver* Solver;
 #if WITH_EDITOR
 		// Visualization material
-		const UMaterial* ClothMaterial;
-		const UMaterial* ClothMaterialVertex;
-		const UMaterial* CollisionMaterial;
+		TObjectPtr<const UMaterial> ClothMaterial = nullptr;
+		TObjectPtr<const UMaterial> ClothMaterialVertex = nullptr;
+		TObjectPtr<const UMaterial> CollisionMaterial = nullptr;
 #endif  // #if WITH_EDITOR
 #endif  // #if CHAOS_DEBUG_DRAW
 	};
