CIS - Fix non unity builds undefined struct 'TObjectPtr<const UMaterial>'.

#rnx

diff --git a/Engine/Plugins/Experimental/ChaosCloth/Source/Chaos/Public/ChaosCloth/ChaosClothVisualization.h b/Engine/Plugins/Experimental/ChaosCloth/Source/Chaos/Public/ChaosCloth/ChaosClothVisualization.h
--- a/Engine/Plugins/Experimental/ChaosCloth/Source/Chaos/Public/ChaosCloth/ChaosClothVisualization.h
+++ b/Engine/Plugins/Experimental/ChaosCloth/Source/Chaos/Public/ChaosCloth/ChaosClothVisualization.h
@@ -6,6 +6,7 @@
 
 #if WITH_EDITOR
 #include "UObject/GCObject.h"
+#include "UObject/ObjectPtr.h"
 #endif  // #if WITH_EDITOR
 
 class FCanvas;
