Chaos Cloth: record reference space transform for correct playback.

Since Chaos Cloth caches are now read in two places (SkeletalMeshComponentCacheAdapter and ClothComponentCacheAdapter), I created a few new structs to share code.
ClothingSimulationCacheData holds cache data that is read for the associated context.
ChaosClothingCacheSchema reads and writes cache data.

We need to write the ReferenceSpaceTransform as part of caching and then to read it back and set it.
Additionally, I'm now only writing the particles that have been simulated (namely other LOD particles will no longer be cached).

ChaosCache has been updated to be able to store named transforms (rather than just per-particle transforms)

Existing cloth caches will continue to work. They simply don't have the transform data.

diff --git a/Engine/Plugins/Experimental/ChaosCaching/Source/ChaosCaching/Private/Chaos/ChaosCache.cpp b/Engine/Plugins/Experimental/ChaosCaching/Source/ChaosCaching/Private/Chaos/ChaosCache.cpp
--- a/Engine/Plugins/Experimental/ChaosCaching/Source/ChaosCaching/Private/Chaos/ChaosCache.cpp
+++ b/Engine/Plugins/Experimental/ChaosCaching/Source/ChaosCaching/Private/Chaos/ChaosCache.cpp
@@ -111,6 +111,38 @@ void UChaosCache::FlushPendingFrames()
 	QUICK_SCOPE_CYCLE_COUNTER(QSTAT_CacheFlushPendingFrames);
 	bool bWroteParticleData = false;
 	FPendingFrameWrite NewData;
+
+	auto WriteDataToTransformTrack = [](FParticleTransformTrack& PTrack, const float Time, const bool bPendingDeactivate, const FTransform& PendingTransform) -> bool
+	{
+		if (PTrack.GetNumKeys() == 0)
+		{
+			// Initial write to this particle
+			PTrack.BeginOffset = Time;
+
+			// Particle will hold at end of recording.
+			PTrack.bDeactivateOnEnd = false;
+		}
+
+		if (bPendingDeactivate)
+		{
+			// Signals that this is the final keyframe and that the particle then deactivates.
+			PTrack.bDeactivateOnEnd = true;
+		}
+		
+		if (ensure(PTrack.GetNumKeys() == 0 || Time > PTrack.KeyTimestamps.Last()))
+		{
+			PTrack.KeyTimestamps.Add(Time);
+
+			// Append the transform (ignoring scale)
+			FRawAnimSequenceTrack& RawTrack = PTrack.RawTransformTrack;
+			RawTrack.ScaleKeys.Add(FVector3f(1.0f));
+			RawTrack.PosKeys.Add(FVector3f(PendingTransform.GetTranslation()));
+			RawTrack.RotKeys.Add(FQuat4f(PendingTransform.GetRotation()));
+			return true;
+		}
+		return false;
+	};
+
 	while(PendingWrites.Dequeue(NewData))
 	{
 		QUICK_SCOPE_CYCLE_COUNTER(QSTAT_CacheFlushSingleFrame);
@@ -157,32 +189,8 @@ void UChaosCache::FlushPendingFrames()
 			FPerParticleCacheData& TargetCacheData = ParticleTracks[TrackIndex];
 			FParticleTransformTrack& PTrack = TargetCacheData.TransformData;
 
-			if(PTrack.GetNumKeys() == 0)
-			{
-				// Initial write to this particle
-				PTrack.BeginOffset = NewData.Time;
-
-				// Particle will hold at end of recording.
-				PTrack.bDeactivateOnEnd = false; 
-			}
-
-			if (ParticleData.bPendingDeactivate)
+			if (WriteDataToTransformTrack(PTrack, NewData.Time, ParticleData.bPendingDeactivate, ParticleData.PendingTransform))
 			{
-				// Signals that this is the final keyframe and that the particle then deactivates.
-				PTrack.bDeactivateOnEnd = true;
-			}
-
-			// Make sure we're actually appending to the track - shouldn't be adding data from the past
-			if(ensure(PTrack.GetNumKeys() == 0 || NewData.Time > PTrack.KeyTimestamps.Last()))
-			{
-				PTrack.KeyTimestamps.Add(NewData.Time);
-
-				// Append the transform (ignoring scale)
-				FRawAnimSequenceTrack& RawTrack = PTrack.RawTransformTrack;
-				RawTrack.ScaleKeys.Add(FVector3f(1.0f));
-				RawTrack.PosKeys.Add(FVector3f(ParticleData.PendingTransform.GetTranslation()));
-				RawTrack.RotKeys.Add(FQuat4f(ParticleData.PendingTransform.GetRotation()));
-
 				for(TPair<FName, float> CurveKeyPair : ParticleData.PendingCurveData)
 				{
 					FRichCurve& TargetCurve = TargetCacheData.CurveData.FindOrAdd(CurveKeyPair.Key);
@@ -203,6 +211,16 @@ void UChaosCache::FlushPendingFrames()
 			CacheTrack->Merge(MoveTemp(PendingTrack.Value));
 		}
 
+		for (const TTuple<FName, FTransform>& PendingNamedTransform : NewData.PendingNamedTransformData)
+		{
+			FNamedTransformTrack* TransformCacheTrack = NamedTransformTracks.Find(PendingNamedTransform.Key);
+			if (!TransformCacheTrack)
+			{
+				TransformCacheTrack = &NamedTransformTracks.Add(PendingNamedTransform.Key, FNamedTransformTrack());
+			}
+			WriteDataToTransformTrack(*TransformCacheTrack, NewData.Time, false, PendingNamedTransform.Value);			
+		}
+
 		++NumRecordedFrames;
 	}
 
@@ -239,6 +257,7 @@ FCacheUserToken UChaosCache::BeginRecord(UPrimitiveComponent* InComponent, FGuid
 			TrackToParticle.Reset();
 			CurveData.Reset();
 			EventTracks.Reset();
+			NamedTransformTracks.Reset();
 
 			PendingWrites.Empty();
 
@@ -362,7 +381,7 @@ FCacheEvaluationResult UChaosCache::Evaluate(const FCacheEvaluationContext& InCo
 
 	Result.EvaluatedTime = InContext.TickRecord.GetTime();
 
-	if(!InContext.bEvaluateTransform && !InContext.bEvaluateCurves && !InContext.bEvaluateEvents && !InContext.bEvaluateChannels)
+	if(!InContext.bEvaluateTransform && !InContext.bEvaluateCurves && !InContext.bEvaluateEvents && !InContext.bEvaluateChannels && !InContext.bEvaluateNamedTransforms)
 	{
 		// no evaluation requested
 		return Result;
@@ -500,6 +519,15 @@ FCacheEvaluationResult UChaosCache::Evaluate(const FCacheEvaluationContext& InCo
 		EvaluateEvents(InContext.TickRecord, Result.Events);
 	}
 
+	if (InContext.bEvaluateNamedTransforms)
+	{
+		Result.NamedTransforms.Reset();
+		for (const TPair<FName, FNamedTransformTrack>& NamedTransform : NamedTransformTracks)
+		{
+			Result.NamedTransforms.Add(NamedTransform.Key, NamedTransform.Value.Evaluate(InContext.TickRecord.GetTime(), nullptr));
+		}
+	}
+
 	// Update the tick record on completion for the next run
 	InContext.TickRecord.LastTime = InContext.TickRecord.GetTime();
 	InContext.TickRecord.CurrentDt = 0.0f;
diff --git a/Engine/Plugins/Experimental/ChaosCaching/Source/ChaosCaching/Public/Chaos/ChaosCache.h b/Engine/Plugins/Experimental/ChaosCaching/Source/ChaosCaching/Public/Chaos/ChaosCache.h
--- a/Engine/Plugins/Experimental/ChaosCaching/Source/ChaosCaching/Public/Chaos/ChaosCache.h
+++ b/Engine/Plugins/Experimental/ChaosCaching/Source/ChaosCaching/Public/Chaos/ChaosCache.h
@@ -161,6 +161,7 @@ struct FCacheEvaluationContext
 		, bEvaluateTransform(false)
 		, bEvaluateCurves(false)
 		, bEvaluateEvents(false)
+		, bEvaluateNamedTransforms(false)
 	{
 	}
 
@@ -170,6 +171,7 @@ struct FCacheEvaluationContext
 	bool                 bEvaluateEvents;
 	TArray<int32>        EvaluationIndices;
 	bool                 bEvaluateChannels;
+	bool                 bEvaluateNamedTransforms;
 };
 
 struct FCacheEvaluationResult
@@ -181,6 +183,7 @@ public:
 	TArray<TMap<FName, float>>             Curves;
 	TMap<FName, TArray<FCacheEventHandle>> Events;
 	TMap<FName, TArray<float>>             Channels;
+	TMap<FName, FTransform>                NamedTransforms;
 };
 
 struct FPendingParticleWrite
@@ -210,6 +213,8 @@ struct FPendingFrameWrite
 	TArray<int32>				  PendingChannelsIndices;
 	TMap<FName, TArray<float>>	  PendingChannelsData;
 
+	TMap<FName, FTransform>       PendingNamedTransformData;
+
 	template<typename T>
 	FCacheEventTrack& FindOrAddEventTrack(FName InName)
 	{
@@ -380,7 +385,7 @@ public:
 	UPROPERTY()
 	TArray<FPerParticleCacheData> ParticleTracks;
 
-	/** Per-particle data, includes transforms, velocities and other per-particle, per-frame data */
+	/** Per-particle curve data, continuous per-frame data */
 	UPROPERTY()
 	TMap<FName,FRichCurves> ChannelsTracks;
 
@@ -388,6 +393,11 @@ public:
 	UPROPERTY()
 	TMap<FName, FRichCurve> CurveData;
 
+	/** Per component/cache transform data.*/
+	typedef FParticleTransformTrack FNamedTransformTrack;
+	UPROPERTY()
+	TMap<FName, FParticleTransformTrack> NamedTransformTracks;
+
 	template<typename T>
 	FCacheEventTrack& FindOrAddEventTrack(FName InName)
 	{
diff --git a/Engine/Plugins/Experimental/ChaosCloth/Source/Chaos/Private/ChaosCloth/ChaosClothingCacheSchema.cpp b/Engine/Plugins/Experimental/ChaosCloth/Source/Chaos/Private/ChaosCloth/ChaosClothingCacheSchema.cpp
new file mode 100644
--- /dev/null
+++ b/Engine/Plugins/Experimental/ChaosCloth/Source/Chaos/Private/ChaosCloth/ChaosClothingCacheSchema.cpp
@@ -0,0 +1,215 @@
+// Copyright Epic Games, Inc. All Rights Reserved.
+
+#include "ChaosCloth/ChaosClothingCacheSchema.h"
+#include "ClothingSimulationCacheData.h"
+#include "Chaos/ChaosCache.h"
+#include "ChaosCloth/ChaosClothingSimulation.h"
+#include "ChaosCloth/ChaosClothingSimulationCloth.h"
+#include "ChaosCloth/ChaosClothingSimulationSolver.h"
+
+namespace Chaos
+{
+
+void FClothingCacheSchema::RecordPostSolve(const FClothingSimulationSolver& ClothSolver, FPendingFrameWrite& OutFrame, FReal InTime)
+{
+	const int32 NumParticles = ClothSolver.GetNumParticles();
+	if (NumParticles > 0)
+	{
+		const Softs::FSolverVec3* ParticleXs = ClothSolver.GetParticleXs(0);
+		const Softs::FSolverVec3* ParticleVs = ClothSolver.GetParticleVs(0);
+
+		TArray<float> PendingVX, PendingVY, PendingVZ, PendingPX, PendingPY, PendingPZ;
+		TArray<int32>& PendingID = OutFrame.PendingChannelsIndices;
+
+		PendingID.Reserve(NumParticles);
+		PendingVX.Reserve(NumParticles);
+		PendingVY.Reserve(NumParticles);
+		PendingVZ.Reserve(NumParticles);
+		PendingPX.Reserve(NumParticles);
+		PendingPY.Reserve(NumParticles);
+		PendingPZ.Reserve(NumParticles);
+
+		const TConstArrayView<const FClothingSimulationCloth*> Cloths = ClothSolver.GetCloths();
+		for (int32 ClothIndex = 0; ClothIndex < Cloths.Num(); ++ClothIndex)
+		{
+			const FClothingSimulationCloth* Cloth = Cloths[ClothIndex];
+			const FTransform& ReferenceSpaceTransform = Cloth->GetReferenceSpaceTransform();
+			const FName ReferenceTransformNameAndIndex(ReferenceTransformsName, Cloth->GetGroupId());
+			OutFrame.PendingNamedTransformData.Add(ReferenceTransformNameAndIndex, ReferenceSpaceTransform);
+
+			const int32 ClothOffset = Cloth->GetOffset(&ClothSolver);
+			const int32 NumClothParticles = Cloth->GetNumParticles(&ClothSolver);
+			check(ClothOffset + NumClothParticles <= NumParticles);
+			for (int32 ParticleIndex = ClothOffset; ParticleIndex < ClothOffset + NumClothParticles; ++ParticleIndex)
+			{
+				const Softs::FSolverVec3& ParticleV = ParticleVs[ParticleIndex];
+				const Softs::FSolverVec3& ParticleX = ParticleXs[ParticleIndex];
+
+				// Adding the vertices relative position to the particle write datas
+				FPendingParticleWrite NewData;
+
+				NewData.ParticleIndex = ParticleIndex;
+				OutFrame.PendingParticleData.Add(MoveTemp(NewData));
+
+				PendingID.Add(ParticleIndex);
+				PendingVX.Add(ParticleV.X);
+				PendingVY.Add(ParticleV.Y);
+				PendingVZ.Add(ParticleV.Z);
+
+				PendingPX.Add(ParticleX.X);
+				PendingPY.Add(ParticleX.Y);
+				PendingPZ.Add(ParticleX.Z);
+			}
+		}
+
+		OutFrame.PendingChannelsData.Add(VelocityXName, PendingVX);
+		OutFrame.PendingChannelsData.Add(VelocityYName, PendingVY);
+		OutFrame.PendingChannelsData.Add(VelocityZName, PendingVZ);
+		OutFrame.PendingChannelsData.Add(PositionXName, PendingPX);
+		OutFrame.PendingChannelsData.Add(PositionYName, PendingPY);
+		OutFrame.PendingChannelsData.Add(PositionZName, PendingPZ);
+	}
+}
+
+void FClothingCacheSchema::PlaybackPreSolve(UChaosCache& InCache, FReal InTime, FPlaybackTickRecord& TickRecord, FClothingSimulationSolver& ClothSolver)
+{
+	FCacheEvaluationContext Context(TickRecord);
+	Context.bEvaluateTransform = false;
+	Context.bEvaluateCurves = false;
+	Context.bEvaluateEvents = false;
+	Context.bEvaluateChannels = true;
+	Context.bEvaluateNamedTransforms = true;
+
+	// The evaluated result are already in world space since we are passing the tickrecord.spacetransform in the context
+	FCacheEvaluationResult EvaluatedResult = InCache.Evaluate(Context, nullptr);
+
+	const TArray<float>* PendingVX = EvaluatedResult.Channels.Find(VelocityXName);
+	const TArray<float>* PendingVY = EvaluatedResult.Channels.Find(VelocityYName);
+	const TArray<float>* PendingVZ = EvaluatedResult.Channels.Find(VelocityZName);
+	const TArray<float>* PendingPX = EvaluatedResult.Channels.Find(PositionXName);
+	const TArray<float>* PendingPY = EvaluatedResult.Channels.Find(PositionYName);
+	const TArray<float>* PendingPZ = EvaluatedResult.Channels.Find(PositionZName);
+
+	if (PendingVX && PendingVY && PendingVZ && PendingPX && PendingPY && PendingPZ)
+	{
+		// Directly set the result of the cache into the solver particles
+		const int32 NumCachedParticles = EvaluatedResult.ParticleIndices.Num();
+		Softs::FSolverVec3* ParticleXs = ClothSolver.GetParticleXs(0);
+		Softs::FSolverVec3* ParticleVs = ClothSolver.GetParticleVs(0);
+		const int32 NumParticles = ClothSolver.GetNumParticles();
+		for (int32 CachedIndex = 0; CachedIndex < NumParticles; ++CachedIndex)
+		{
+			const int32 ParticleIndex = EvaluatedResult.ParticleIndices[CachedIndex];
+			check(ParticleIndex < NumParticles);
+			Softs::FSolverVec3& ParticleV = ParticleVs[ParticleIndex];
+			Softs::FSolverVec3& ParticleX = ParticleXs[ParticleIndex];
+			ParticleV.X = (*PendingVX)[CachedIndex];
+			ParticleV.Y = (*PendingVY)[CachedIndex];
+			ParticleV.Z = (*PendingVZ)[CachedIndex];
+			ParticleX.X = (*PendingPX)[CachedIndex];
+			ParticleX.Y = (*PendingPY)[CachedIndex];
+			ParticleX.Z = (*PendingPZ)[CachedIndex];
+		}
+
+		// Fix up positions based on any discrepency between cached ReferenceSpaceTransform and current ReferenceSpaceTransform
+		TConstArrayView<const FClothingSimulationCloth*> Cloths = ClothSolver.GetCloths();
+		auto GetCloth = [&Cloths](const int32 GroupId)
+		{
+			const FClothingSimulationCloth* const* Found = Cloths.FindByPredicate(
+				[GroupId](const FClothingSimulationCloth* InCloth)
+				{
+					return InCloth->GetGroupId() == GroupId;
+				});
+			return Found ? *Found : nullptr;
+		};
+
+		const FVec3& LocalSpaceLocation = ClothSolver.GetLocalSpaceLocation();
+		constexpr bool bNoCompareNumber = false;
+		for (const TPair<FName, FTransform>& NamedTransform : EvaluatedResult.NamedTransforms)
+		{
+			if (NamedTransform.Key.IsEqual(ReferenceTransformsName, ENameCase::CaseSensitive, bNoCompareNumber))
+			{
+				const int32 ClothGroupId = NamedTransform.Key.GetNumber();
+				const FClothingSimulationCloth* Cloth = GetCloth(ClothGroupId);
+				if (Cloth)
+				{
+					const FRigidTransform3& CurrentTransform = Cloth->GetReferenceSpaceTransform();
+					const FTransform& CachedTransform = NamedTransform.Value;
+					if (!CurrentTransform.Equals(CachedTransform))
+					{
+						const FTransform RelativeTransform = CurrentTransform.GetRelativeTransform(CachedTransform);
+						const int32 ClothOffset = Cloth->GetOffset(&ClothSolver);
+						const int32 NumClothParticles = Cloth->GetNumParticles(&ClothSolver);
+						check(ClothOffset + NumClothParticles <= NumParticles);
+						for (int32 ParticleIndex = ClothOffset; ParticleIndex < ClothOffset + NumClothParticles; ++ParticleIndex)
+						{
+							Softs::FSolverVec3& ParticleX = ParticleXs[ParticleIndex];
+							ParticleX = Softs::FSolverVec3(RelativeTransform.TransformPosition(FVec3(ParticleX + LocalSpaceLocation)) - LocalSpaceLocation);
+							Softs::FSolverVec3& ParticleV = ParticleVs[ParticleIndex];
+							ParticleV = Softs::FSolverVec3(RelativeTransform.TransformVector(FVec3(ParticleV)));
+						}
+					}
+				}
+			}
+		}
+	}
+}
+
+void FClothingCacheSchema::LoadCacheData(UChaosCache* InCache, Chaos::FReal InTime, FClothingSimulationCacheData& CacheData)
+{
+	FPlaybackTickRecord TickRecord;
+	TickRecord.SetLastTime(InTime);
+
+	FCacheEvaluationContext Context(TickRecord);
+	Context.bEvaluateTransform = false;
+	Context.bEvaluateCurves = false;
+	Context.bEvaluateEvents = false;
+	Context.bEvaluateChannels = true;
+	Context.bEvaluateNamedTransforms = true;
+
+	const FCacheEvaluationResult EvaluatedResult = InCache->Evaluate(Context, nullptr);
+
+	const TArray<float>* PendingVX = EvaluatedResult.Channels.Find(VelocityXName);
+	const TArray<float>* PendingVY = EvaluatedResult.Channels.Find(VelocityYName);
+	const TArray<float>* PendingVZ = EvaluatedResult.Channels.Find(VelocityZName);
+	const TArray<float>* PendingPX = EvaluatedResult.Channels.Find(PositionXName);
+	const TArray<float>* PendingPY = EvaluatedResult.Channels.Find(PositionYName);
+	const TArray<float>* PendingPZ = EvaluatedResult.Channels.Find(PositionZName);
+
+	const bool bHasPositions = PendingPX && PendingPY && PendingPZ;
+	const bool bHasVelocities = PendingVX && PendingVY && PendingVZ;
+
+	if ( bHasPositions )
+	{
+		CacheData.CacheIndices = EvaluatedResult.ParticleIndices;
+		const int32 NumParticles = EvaluatedResult.ParticleIndices.Num();
+
+		CacheData.CachedPositions.SetNum(NumParticles);
+		CacheData.CachedVelocities.SetNum(bHasVelocities ? NumParticles : 0);
+
+		for (int32 ParticleIndex = 0; ParticleIndex < NumParticles; ++ParticleIndex)
+		{
+			CacheData.CachedPositions[ParticleIndex].X = (*PendingPX)[ParticleIndex];
+			CacheData.CachedPositions[ParticleIndex].Y = (*PendingPY)[ParticleIndex];
+			CacheData.CachedPositions[ParticleIndex].Z = (*PendingPZ)[ParticleIndex];
+			if (bHasVelocities)
+			{
+				CacheData.CachedVelocities[ParticleIndex].X = (*PendingVX)[ParticleIndex];
+				CacheData.CachedVelocities[ParticleIndex].Y = (*PendingVY)[ParticleIndex];
+				CacheData.CachedVelocities[ParticleIndex].Z = (*PendingVZ)[ParticleIndex];
+			}
+		}
+	}
+
+	CacheData.CachedReferenceSpaceTransforms.Reset();
+	for (const TPair<FName, FTransform>& NamedTransform : EvaluatedResult.NamedTransforms)
+	{
+		constexpr bool bNoCompareNumber = false;
+		if (NamedTransform.Key.IsEqual(ReferenceTransformsName, ENameCase::CaseSensitive, bNoCompareNumber))
+		{
+			CacheData.CachedReferenceSpaceTransforms.Add(NamedTransform.Key.GetNumber(), NamedTransform.Value);
+		}
+	}
+}
+
+} // namespace Chaos
\ No newline at end of file
diff --git a/Engine/Plugins/Experimental/ChaosCloth/Source/Chaos/Private/ChaosCloth/ChaosClothingSimulation.cpp b/Engine/Plugins/Experimental/ChaosCloth/Source/Chaos/Private/ChaosCloth/ChaosClothingSimulation.cpp
--- a/Engine/Plugins/Experimental/ChaosCloth/Source/Chaos/Private/ChaosCloth/ChaosClothingSimulation.cpp
+++ b/Engine/Plugins/Experimental/ChaosCloth/Source/Chaos/Private/ChaosCloth/ChaosClothingSimulation.cpp
@@ -526,14 +526,24 @@ void FClothingSimulation::Simulate(IClothingSimulationContext* InContext)
 		}
 
 		// Step the simulation
-		if(Solver->GetEnableSolver() || (Context->CachedPositions.Num() == 0 && Context->CachedVelocities.Num() == 0))
+		PRAGMA_DISABLE_DEPRECATION_WARNINGS // Supporting deprecated CachedPositions instead of new CacheData
+		if(Solver->GetEnableSolver() || (!Context->CacheData.HasData() && Context->CachedPositions.Num() == 0))
 		{
 			Solver->Update(SmoothedDeltaTime);
 		}
 		else
 		{
-			Solver->UpdateFromCache(Context->CachedPositions, Context->CachedVelocities);
+			if (Context->CacheData.HasData())
+			{
+				Solver->UpdateFromCache(Context->CacheData);
+			}
+			else
+			{
+				check(Context->CachedPositions.Num());
+				Solver->UpdateFromCache(Context->CachedPositions, Context->CachedVelocities);
+			}
 		}
+		PRAGMA_ENABLE_DEPRECATION_WARNINGS
 
 		// Keep the actual used number of iterations for the stats
 		NumIterations = Solver->GetNumUsedIterations();
diff --git a/Engine/Plugins/Experimental/ChaosCloth/Source/Chaos/Private/ChaosCloth/ChaosClothingSimulationCloth.cpp b/Engine/Plugins/Experimental/ChaosCloth/Source/Chaos/Private/ChaosCloth/ChaosClothingSimulationCloth.cpp
--- a/Engine/Plugins/Experimental/ChaosCloth/Source/Chaos/Private/ChaosCloth/ChaosClothingSimulationCloth.cpp
+++ b/Engine/Plugins/Experimental/ChaosCloth/Source/Chaos/Private/ChaosCloth/ChaosClothingSimulationCloth.cpp
@@ -730,6 +730,12 @@ int32 FClothingSimulationCloth::GetOffset(const FClothingSimulationSolver* Solve
 	return LODData.IsValidIndex(LODIndex) ? GetOffset(Solver, LODIndex) : INDEX_NONE;
 }
 
+int32 FClothingSimulationCloth::GetNumParticles(const FClothingSimulationSolver* Solver) const
+{
+	const int32 LODIndex = LODIndices.FindChecked(Solver);
+	return LODData.IsValidIndex(LODIndex) ? GetNumParticles(LODIndex) : 0;
+}
+
 const FTriangleMesh& FClothingSimulationCloth::GetTriangleMesh(const FClothingSimulationSolver* Solver) const
 {
 	const int32 LODIndex = LODIndices.FindChecked(Solver);
@@ -972,6 +978,15 @@ void FClothingSimulationCloth::PostUpdate(FClothingSimulationSolver* Solver)
 	}
 }
 
+void FClothingSimulationCloth::UpdateFromCache(const FClothingSimulationCacheData& CacheData)
+{
+	if (const FTransform* CachedReferenceSpaceTransform = CacheData.CachedReferenceSpaceTransforms.Find(GetGroupId()))
+	{
+		ReferenceSpaceTransform = *CachedReferenceSpaceTransform;
+		ReferenceSpaceTransform.SetScale3D(FVec3(1.f));
+	}
+}
+
 TConstArrayView<Softs::FSolverVec3> FClothingSimulationCloth::GetAnimationPositions(const FClothingSimulationSolver* Solver) const
 {
 	check(Solver);
diff --git a/Engine/Plugins/Experimental/ChaosCloth/Source/Chaos/Private/ChaosCloth/ChaosClothingSimulationSolver.cpp b/Engine/Plugins/Experimental/ChaosCloth/Source/Chaos/Private/ChaosCloth/ChaosClothingSimulationSolver.cpp
--- a/Engine/Plugins/Experimental/ChaosCloth/Source/Chaos/Private/ChaosCloth/ChaosClothingSimulationSolver.cpp
+++ b/Engine/Plugins/Experimental/ChaosCloth/Source/Chaos/Private/ChaosCloth/ChaosClothingSimulationSolver.cpp
@@ -1181,6 +1181,31 @@ void FClothingSimulationSolver::Update(Softs::FSolverReal InDeltaTime)
 	OldLocalSpaceLocation = LocalSpaceLocation;
 }
 
+
+void FClothingSimulationSolver::UpdateFromCache(const FClothingSimulationCacheData& CacheData)
+{
+	Chaos::Softs::FSolverParticles& SolverParticles = Evolution->Particles();
+	const int32 NumParticles = GetNumParticles();
+	const int32 NumCachedParticles = CacheData.CacheIndices.Num();	
+	const bool bHasVelocity = CacheData.CachedVelocities.Num() > 0;
+	for (int32 CacheIndex = 0; CacheIndex < NumCachedParticles; ++CacheIndex)
+	{
+		const int32 ParticleIndex = CacheData.CacheIndices[CacheIndex];
+		if (ensure(ParticleIndex < NumParticles))
+		{
+			SolverParticles.X(ParticleIndex) = CacheData.CachedPositions[CacheIndex];
+			SolverParticles.V(ParticleIndex) = bHasVelocity ? (FSolverVec3)CacheData.CachedVelocities[CacheIndex] : FSolverVec3::ZeroVector;
+		}
+	}
+	PhysicsParallelFor(Cloths.Num(), [this, &CacheData](int32 ClothIndex)
+	{
+		FClothingSimulationCloth* const Cloth = Cloths[ClothIndex];
+		Cloth->UpdateFromCache(CacheData);
+		Cloth->PostUpdate(this);
+	}, /*bForceSingleThreaded =*/ !bClothSolverParallelClothPostUpdate);
+
+}
+
 void FClothingSimulationSolver::UpdateFromCache(const TArray<FVector>& CachedPositions, const TArray<FVector>& CachedVelocities) 
 {
 	Chaos::Softs::FSolverParticles& SolverParticles = Evolution->Particles();
diff --git a/Engine/Plugins/Experimental/ChaosCloth/Source/Chaos/Private/ChaosCloth/SkeletalMeshComponentCacheAdapter.cpp b/Engine/Plugins/Experimental/ChaosCloth/Source/Chaos/Private/ChaosCloth/SkeletalMeshComponentCacheAdapter.cpp
--- a/Engine/Plugins/Experimental/ChaosCloth/Source/Chaos/Private/ChaosCloth/SkeletalMeshComponentCacheAdapter.cpp
+++ b/Engine/Plugins/Experimental/ChaosCloth/Source/Chaos/Private/ChaosCloth/SkeletalMeshComponentCacheAdapter.cpp
@@ -1,10 +1,10 @@
 // Copyright Epic Games, Inc. All Rights Reserved.
 
 #include "ChaosCloth/SkeletalMeshComponentCacheAdapter.h"
+#include "ChaosCloth/ChaosClothingCacheSchema.h"
 #include "Chaos/ParticleHandle.h"
 #include "Chaos/ChaosCache.h"
 #include "Chaos/CacheManagerActor.h"
-#include "Chaos/PBDEvolution.h"
 #include "ChaosCloth/ChaosClothingSimulation.h"
 #include "ChaosCloth/ChaosClothingSimulationSolver.h"
 #include "Components/SkeletalMeshComponent.h"
@@ -38,101 +38,18 @@ namespace Chaos
 
 	void FSkeletalMeshCacheAdapter::Record_PostSolve(UPrimitiveComponent* InComponent, const FTransform& InRootTransform, FPendingFrameWrite& OutFrame, Chaos::FReal InTime) const
 	{
-		if( Chaos::FClothingSimulationSolver* ClothSolver = GetClothSolver(InComponent))
+		if( const FClothingSimulationSolver* ClothSolver = GetClothSolver(InComponent))
 		{
-			const uint32 NumParticles = ClothSolver->GetNumParticles();
-			if(NumParticles > 0)
-			{
-				const Softs::FSolverVec3* ParticleXs = ClothSolver->GetParticleXs(0);
-				const Softs::FSolverVec3* ParticleVs = ClothSolver->GetParticleVs(0);
-				
-				TArray<float> PendingVX, PendingVY, PendingVZ, PendingPX, PendingPY, PendingPZ;
-				TArray<int32>& PendingID = OutFrame.PendingChannelsIndices;
-
-				PendingID.SetNum(NumParticles);
-				PendingVX.SetNum(NumParticles);
-				PendingVY.SetNum(NumParticles);
-				PendingVZ.SetNum(NumParticles);
-				PendingPX.SetNum(NumParticles);
-				PendingPY.SetNum(NumParticles);
-				PendingPZ.SetNum(NumParticles);
-
-				for(uint32 ParticleIndex = 0; ParticleIndex < NumParticles; ++ParticleIndex)
-				{
-					const Softs::FSolverVec3& ParticleV = ParticleVs[ParticleIndex];
-					const Softs::FSolverVec3& ParticleX = ParticleXs[ParticleIndex];
-					
-					// Adding the vertices relative position to the particle write datas
-					FPendingParticleWrite NewData;
-
-					NewData.ParticleIndex = ParticleIndex; 
-					OutFrame.PendingParticleData.Add(MoveTemp(NewData));
-
-					PendingID[ParticleIndex] = ParticleIndex;
-					PendingVX[ParticleIndex] = ParticleV.X;
-					PendingVY[ParticleIndex] = ParticleV.Y;
-					PendingVZ[ParticleIndex] = ParticleV.Z;
-
-					PendingPX[ParticleIndex] = ParticleX.X;
-					PendingPY[ParticleIndex] = ParticleX.Y;
-					PendingPZ[ParticleIndex] = ParticleX.Z;
-				}
-
-				OutFrame.PendingChannelsData.Add(VelocityXName, PendingVX);
-				OutFrame.PendingChannelsData.Add(VelocityYName, PendingVY);
-				OutFrame.PendingChannelsData.Add(VelocityZName, PendingVZ);
-				OutFrame.PendingChannelsData.Add(PositionXName, PendingPX);
-				OutFrame.PendingChannelsData.Add(PositionYName, PendingPY);
-				OutFrame.PendingChannelsData.Add(PositionZName, PendingPZ);
-			}
+			FClothingCacheSchema::RecordPostSolve(*ClothSolver, OutFrame, InTime);			
 		}
 	}
 
 	void FSkeletalMeshCacheAdapter::Playback_PreSolve(UPrimitiveComponent* InComponent, UChaosCache* InCache, Chaos::FReal InTime, FPlaybackTickRecord& TickRecord, TArray<TPBDRigidParticleHandle<Chaos::FReal, 3>*>& OutUpdatedRigids) const
 	{
-		if( Chaos::FClothingSimulationSolver* ClothSolver = GetClothSolver(InComponent))
+		check(InCache);
+		if (FClothingSimulationSolver* ClothSolver = GetClothSolver(InComponent))
 		{
-			FCacheEvaluationContext Context(TickRecord);
-			Context.bEvaluateTransform = false;
-			Context.bEvaluateCurves = true;
-			Context.bEvaluateEvents = false;
-
-			// The evaluated result are already in world space since we are passing the tickrecord.spacetransform in the context
-			FCacheEvaluationResult EvaluatedResult = InCache->Evaluate(Context, nullptr);
-			const int32 NumCurves = EvaluatedResult.Curves.Num();
-
-			const uint32 NumParticles = ClothSolver->GetNumParticles();
-			if(NumParticles > 0 && NumParticles == NumCurves)
-			{
-				Softs::FSolverVec3* ParticleXs = ClothSolver->GetParticleXs(0);
-				Softs::FSolverVec3* ParticleVs = ClothSolver->GetParticleVs(0);
-
-				TArray<float>* PendingVX = EvaluatedResult.Channels.Find(VelocityXName);
-				TArray<float>* PendingVY = EvaluatedResult.Channels.Find(VelocityYName);
-				TArray<float>* PendingVZ = EvaluatedResult.Channels.Find(VelocityZName);
-				TArray<float>* PendingPX = EvaluatedResult.Channels.Find(PositionXName);
-				TArray<float>* PendingPY = EvaluatedResult.Channels.Find(PositionYName);
-				TArray<float>* PendingPZ = EvaluatedResult.Channels.Find(PositionZName);
-					
-				if(PendingVX && PendingVY && PendingVZ && PendingPX && PendingPY && PendingPZ)
-				{
-					for(uint32 ParticleIndex = 0; ParticleIndex < NumParticles; ++ParticleIndex)
-					{
-						Softs::FSolverVec3& ParticleV = ParticleVs[ParticleIndex];
-						Softs::FSolverVec3& ParticleX = ParticleXs[ParticleIndex];
-					
-						// Directly set the result of the cache into the solver particles
-						const int32 TransformIndex = EvaluatedResult.ParticleIndices[ParticleIndex];
-
-						ParticleV.X = (*PendingVX)[TransformIndex];
-						ParticleV.Y = (*PendingVY)[TransformIndex];
-						ParticleV.Z = (*PendingVZ)[TransformIndex];
-						ParticleX.X = (*PendingPX)[TransformIndex];
-						ParticleX.Y = (*PendingPY)[TransformIndex];
-						ParticleX.Z = (*PendingPZ)[TransformIndex];
-					}
-				}
-			}
+			FClothingCacheSchema::PlaybackPreSolve(*InCache, InTime, TickRecord, *ClothSolver);
 		}
 	}
 
@@ -150,7 +67,7 @@ namespace Chaos
 		return MeshComp && MeshComp->GetSkinnedAsset() && InCache->TrackToParticle.Num() > 0;
 	}
 
-	Chaos::FClothingSimulationSolver* FSkeletalMeshCacheAdapter::GetClothSolver(UPrimitiveComponent* InComponent) const
+	FClothingSimulationSolver* FSkeletalMeshCacheAdapter::GetClothSolver(UPrimitiveComponent* InComponent) const
 	{
 		if(InComponent)
 		{
@@ -158,14 +75,14 @@ namespace Chaos
 			{
 				if(MeshComp->GetClothingSimulation())
 				{
-					return StaticCast<Chaos::FClothingSimulation*>(MeshComp->GetClothingSimulation())->Solver.Get();
+					return StaticCast<FClothingSimulation*>(MeshComp->GetClothingSimulation())->Solver.Get();
 				}
 			}
 		}
 		return nullptr;
 	}
 
-	Chaos::FPhysicsSolverEvents* FSkeletalMeshCacheAdapter::BuildEventsSolver(UPrimitiveComponent* InComponent) const
+	FPhysicsSolverEvents* FSkeletalMeshCacheAdapter::BuildEventsSolver(UPrimitiveComponent* InComponent) const
 	{
 		if(USkeletalMeshComponent* MeshComp = CastChecked<USkeletalMeshComponent>(InComponent))
 		{
@@ -178,12 +95,12 @@ namespace Chaos
 		return nullptr;
 	}
 	
-	Chaos::FPhysicsSolver* FSkeletalMeshCacheAdapter::GetComponentSolver(UPrimitiveComponent* InComponent) const
+	FPhysicsSolver* FSkeletalMeshCacheAdapter::GetComponentSolver(UPrimitiveComponent* InComponent) const
 	{
 		return nullptr;
 	}
 	
-	void FSkeletalMeshCacheAdapter::SetRestState(UPrimitiveComponent* InComponent, UChaosCache* InCache, const FTransform& InRootTransform, Chaos::FReal InTime) const
+	void FSkeletalMeshCacheAdapter::SetRestState(UPrimitiveComponent* InComponent, UChaosCache* InCache, const FTransform& InRootTransform, FReal InTime) const
 	{
 		if (!InCache || InCache->GetDuration() == 0.0f)
 		{
@@ -193,51 +110,13 @@ namespace Chaos
 		if(USkeletalMeshComponent* MeshComp = CastChecked<USkeletalMeshComponent>(InComponent))
 		{
 			InitializeForRestState(InComponent);
-		
-			FPlaybackTickRecord TickRecord;
-			TickRecord.SetLastTime(InTime);
-		
-			FCacheEvaluationContext Context(TickRecord);
-			Context.bEvaluateTransform = false;
-			Context.bEvaluateCurves = false;
-			Context.bEvaluateEvents = false;
-			Context.bEvaluateChannels = true;
-
-			FCacheEvaluationResult EvaluatedResult = InCache->Evaluate(Context,nullptr);
 
 			if( MeshComp->GetClothingSimulationContext())
 			{
-				FClothingSimulationContext* SimulationContext =
-						StaticCast<FClothingSimulationContext*>(MeshComp->GetClothingSimulationContext());
-				
-				const int32 NumCurves = EvaluatedResult.ParticleIndices.Num();
-				if(NumCurves == EvaluatedResult.ParticleIndices.Num())
-				{
-					SimulationContext->CachedPositions.SetNum(NumCurves);
-					SimulationContext->CachedVelocities.SetNum(NumCurves);
+				FClothingSimulationCacheData& SimulationContextData =
+						StaticCast<FClothingSimulationContext*>(MeshComp->GetClothingSimulationContext())->CacheData;
 
-					TArray<float>* PendingVX = EvaluatedResult.Channels.Find(VelocityXName);
-					TArray<float>* PendingVY = EvaluatedResult.Channels.Find(VelocityYName);
-					TArray<float>* PendingVZ = EvaluatedResult.Channels.Find(VelocityZName);
-					TArray<float>* PendingPX = EvaluatedResult.Channels.Find(PositionXName);
-					TArray<float>* PendingPY = EvaluatedResult.Channels.Find(PositionYName);
-					TArray<float>* PendingPZ = EvaluatedResult.Channels.Find(PositionZName);
-					
-					if(PendingVX && PendingVY && PendingVZ && PendingPX && PendingPY && PendingPZ)
-					{
-						for(int32 ParticleIndex = 0; ParticleIndex < NumCurves; ++ParticleIndex)
-						{
-							const int32 TransformIndex = EvaluatedResult.ParticleIndices[ParticleIndex];
-							
-							SimulationContext->CachedVelocities[ParticleIndex].X = (*PendingVX)[TransformIndex];
-							SimulationContext->CachedVelocities[ParticleIndex].Y = (*PendingVY)[TransformIndex];
-							SimulationContext->CachedVelocities[ParticleIndex].Z = (*PendingVZ)[TransformIndex];
-							SimulationContext->CachedPositions[ParticleIndex].X = (*PendingPX)[TransformIndex];
-							SimulationContext->CachedPositions[ParticleIndex].Y = (*PendingPY)[TransformIndex];
-							SimulationContext->CachedPositions[ParticleIndex].Z = (*PendingPZ)[TransformIndex];
-						}
-					}
-				}
+				FClothingCacheSchema::LoadCacheData(InCache, InTime, SimulationContextData);
 			}
 		}
 	}
@@ -258,7 +137,7 @@ namespace Chaos
 			MeshComp->RecreateClothingActors();
 		}
 
-		if( Chaos::FClothingSimulationSolver* ClothSolver = GetClothSolver(InComponent))
+		if( FClothingSimulationSolver* ClothSolver = GetClothSolver(InComponent))
 		{
 			ClothSolver->SetEnableSolver(false);
 		}
@@ -266,7 +145,7 @@ namespace Chaos
 
 	bool FSkeletalMeshCacheAdapter::InitializeForRecord(UPrimitiveComponent* InComponent, UChaosCache* InCache)
 	{
-		if( Chaos::FClothingSimulationSolver* ClothSolver = GetClothSolver(InComponent))
+		if( FClothingSimulationSolver* ClothSolver = GetClothSolver(InComponent))
 		{
 			ClothSolver->SetEnableSolver(true);
 		}
@@ -277,7 +156,7 @@ namespace Chaos
 	{
 		EnsureIsInGameThreadContext();
 		
-		if( Chaos::FClothingSimulationSolver* ClothSolver = GetClothSolver(InComponent))
+		if( FClothingSimulationSolver* ClothSolver = GetClothSolver(InComponent))
 		{
 			ClothSolver->SetEnableSolver(false);
 		}
diff --git a/Engine/Plugins/Experimental/ChaosCloth/Source/Chaos/Private/ChaosCloth/SkeletalMeshComponentCacheAdapter.h b/Engine/Plugins/Experimental/ChaosCloth/Source/Chaos/Private/ChaosCloth/SkeletalMeshComponentCacheAdapter.h
--- a/Engine/Plugins/Experimental/ChaosCloth/Source/Chaos/Private/ChaosCloth/SkeletalMeshComponentCacheAdapter.h
+++ b/Engine/Plugins/Experimental/ChaosCloth/Source/Chaos/Private/ChaosCloth/SkeletalMeshComponentCacheAdapter.h
@@ -50,11 +50,5 @@ namespace Chaos
 		*/
 		void InitializeForRestState(UPrimitiveComponent* InComponent) const;
 
-		inline static const FName VelocityXName = TEXT("VelocityX");
-		inline static const FName VelocityYName = TEXT("VelocityY");
-		inline static const FName VelocityZName = TEXT("VelocityZ");
-		inline static const FName PositionXName = TEXT("PositionX");
-		inline static const FName PositionYName = TEXT("PositionY");
-		inline static const FName PositionZName = TEXT("PositionZ");
 	};
 }    // namespace Chaos
diff --git a/Engine/Plugins/Experimental/ChaosCloth/Source/Chaos/Public/ChaosCloth/ChaosClothingCacheSchema.h b/Engine/Plugins/Experimental/ChaosCloth/Source/Chaos/Public/ChaosCloth/ChaosClothingCacheSchema.h
new file mode 100644
--- /dev/null
+++ b/Engine/Plugins/Experimental/ChaosCloth/Source/Chaos/Public/ChaosCloth/ChaosClothingCacheSchema.h
@@ -0,0 +1,37 @@
+// Copyright Epic Games, Inc. All Rights Reserved.
+
+#pragma once
+
+#include "Chaos/Declares.h"
+#include "CoreTypes.h"
+#include "UObject/NameTypes.h"
+
+struct FClothingSimulationCacheData;
+class UChaosCache;
+struct FPendingFrameWrite;
+struct FPlaybackTickRecord;
+
+namespace Chaos
+{
+class FClothingSimulationSolver;
+
+struct FClothingCacheSchema
+{
+	// This should only be called from the ClothSolver as a PostSolve callback (when it is safe to read directly from ClothSolver data)
+	static void CHAOSCLOTH_API RecordPostSolve(const FClothingSimulationSolver& ClothSolver, FPendingFrameWrite& OutFrame, FReal InTime);
+	// This should only be called from the ClothSolver as a PreSolve callback (when it is safe to write directly to ClothSolver data)
+	static void CHAOSCLOTH_API PlaybackPreSolve(UChaosCache& InCache, FReal InTime, FPlaybackTickRecord& TickRecord, FClothingSimulationSolver& ClothSolver);
+	static void CHAOSCLOTH_API LoadCacheData(UChaosCache* InCache, FReal InTime, FClothingSimulationCacheData& CacheData);
+
+private:
+	inline static const FName VelocityXName = TEXT("VelocityX");
+	inline static const FName VelocityYName = TEXT("VelocityY");
+	inline static const FName VelocityZName = TEXT("VelocityZ");
+	inline static const FName PositionXName = TEXT("PositionX");
+	inline static const FName PositionYName = TEXT("PositionY");
+	inline static const FName PositionZName = TEXT("PositionZ");
+
+	inline static const FName ReferenceTransformsName = TEXT("ReferenceTransform");
+};
+
+} // namespace Chaos
\ No newline at end of file
diff --git a/Engine/Plugins/Experimental/ChaosCloth/Source/Chaos/Public/ChaosCloth/ChaosClothingSimulationCloth.h b/Engine/Plugins/Experimental/ChaosCloth/Source/Chaos/Public/ChaosCloth/ChaosClothingSimulationCloth.h
--- a/Engine/Plugins/Experimental/ChaosCloth/Source/Chaos/Public/ChaosCloth/ChaosClothingSimulationCloth.h
+++ b/Engine/Plugins/Experimental/ChaosCloth/Source/Chaos/Public/ChaosCloth/ChaosClothingSimulationCloth.h
@@ -8,6 +8,7 @@
 #include "ChaosCloth/ChaosClothConstraints.h"
 
 struct FManagedArrayCollection;
+struct FClothingSimulationCacheData;
 
 namespace Chaos
 {
@@ -156,6 +157,8 @@ PRAGMA_ENABLE_DEPRECATION_WARNINGS
 		FAABB3 CalculateBoundingBox(const FClothingSimulationSolver* Solver) const;
 		// Return the current LOD offset in the solver's particle array, or INDEX_NONE if no LOD is currently selected.
 		int32 GetOffset(const FClothingSimulationSolver* Solver) const;
+		// Return the current LOD num particles, or 0 if no LOD is currently selected.
+		int32 GetNumParticles(const FClothingSimulationSolver* Solver) const;
 		// Return the current LOD mesh.
 		const FTriangleMesh& GetTriangleMesh(const FClothingSimulationSolver* Solver) const;
 		// Return the current LOD weightmaps.
@@ -175,6 +178,8 @@ PRAGMA_ENABLE_DEPRECATION_WARNINGS
 		void PreUpdate(FClothingSimulationSolver* Solver);
 		void Update(FClothingSimulationSolver* Solver);
 		void PostUpdate(FClothingSimulationSolver* Solver);
+
+		void UpdateFromCache(const FClothingSimulationCacheData& CacheData);
 		// ---- End of the Solver interface ----
 
 	private:
diff --git a/Engine/Plugins/Experimental/ChaosCloth/Source/Chaos/Public/ChaosCloth/ChaosClothingSimulationSolver.h b/Engine/Plugins/Experimental/ChaosCloth/Source/Chaos/Public/ChaosCloth/ChaosClothingSimulationSolver.h
--- a/Engine/Plugins/Experimental/ChaosCloth/Source/Chaos/Public/ChaosCloth/ChaosClothingSimulationSolver.h
+++ b/Engine/Plugins/Experimental/ChaosCloth/Source/Chaos/Public/ChaosCloth/ChaosClothingSimulationSolver.h
@@ -11,6 +11,7 @@
 
 class USkeletalMeshComponent;
 class UClothingAssetCommon;
+struct FClothingSimulationCacheData;
 
 namespace Chaos
 {
@@ -89,6 +90,8 @@ namespace Chaos
 		FSolverReal GetDeltaTime() const { return DeltaTime; }
 
 		/** Set the cached positions onto the particles */
+		void UpdateFromCache(const FClothingSimulationCacheData& CacheData);
+		UE_DEPRECATED(5.3, "Use UpdateFromCache(CacheData) instead")
 		void UpdateFromCache(const TArray<FVector>& CachedPositions, const TArray<FVector>& CachedVelocities);
 
 		// Return the actual of number of iterations used by the Evolution solver after the update (different from the number of iterations, depends on frame rate)
diff --git a/Engine/Plugins/Experimental/ChaosClothAsset/Source/ChaosClothAssetEngine/ChaosClothAssetEngine.Build.cs b/Engine/Plugins/Experimental/ChaosClothAsset/Source/ChaosClothAssetEngine/ChaosClothAssetEngine.Build.cs
--- a/Engine/Plugins/Experimental/ChaosClothAsset/Source/ChaosClothAssetEngine/ChaosClothAssetEngine.Build.cs
+++ b/Engine/Plugins/Experimental/ChaosClothAsset/Source/ChaosClothAssetEngine/ChaosClothAssetEngine.Build.cs
@@ -21,6 +21,7 @@ public class ChaosClothAssetEngine : ModuleRules
 				"ChaosCloth",
 				"ChaosCaching",
 				"ProfileVisualizer",
+				"ClothingSystemRuntimeCommon",
 				"ClothingSystemRuntimeInterface",
 				"DataflowEngine"
 			}
diff --git a/Engine/Plugins/Experimental/ChaosClothAsset/Source/ChaosClothAssetEngine/Private/ChaosClothAsset/ClothComponentCacheAdapter.cpp b/Engine/Plugins/Experimental/ChaosClothAsset/Source/ChaosClothAssetEngine/Private/ChaosClothAsset/ClothComponentCacheAdapter.cpp
--- a/Engine/Plugins/Experimental/ChaosClothAsset/Source/ChaosClothAssetEngine/Private/ChaosClothAsset/ClothComponentCacheAdapter.cpp
+++ b/Engine/Plugins/Experimental/ChaosClothAsset/Source/ChaosClothAssetEngine/Private/ChaosClothAsset/ClothComponentCacheAdapter.cpp
@@ -5,6 +5,7 @@
 #include "Chaos/ChaosCache.h"
 #include "Chaos/CacheManagerActor.h"
 #include "Chaos/PBDEvolution.h"
+#include "ChaosCloth/ChaosClothingCacheSchema.h"
 #include "ChaosCloth/ChaosClothingSimulation.h"
 #include "ChaosCloth/ChaosClothingSimulationSolver.h"
 #include "ChaosClothAsset/ClothComponent.h"
@@ -81,53 +82,12 @@ namespace UE::Chaos::ClothAsset
 
 		if (UChaosClothComponent* ClothComp = GetClothComponent(InComponent))
 		{
-			FPlaybackTickRecord TickRecord;
-			TickRecord.SetLastTime(InTime);
-		
-			FCacheEvaluationContext Context(TickRecord);
-			Context.bEvaluateTransform = false;
-			Context.bEvaluateCurves = false;
-			Context.bEvaluateEvents = false;
-			Context.bEvaluateChannels = true;
-
-			FCacheEvaluationResult EvaluatedResult = InCache->Evaluate(Context, nullptr);
-			const int32 NumCurves = EvaluatedResult.ParticleIndices.Num();
-
-			const TArray<float>* const PendingVX = EvaluatedResult.Channels.Find(VelocityXName);
-			const TArray<float>* const PendingVY = EvaluatedResult.Channels.Find(VelocityYName);
-			const TArray<float>* const PendingVZ = EvaluatedResult.Channels.Find(VelocityZName);
-			const TArray<float>* const PendingPX = EvaluatedResult.Channels.Find(PositionXName);
-			const TArray<float>* const PendingPY = EvaluatedResult.Channels.Find(PositionYName);
-			const TArray<float>* const PendingPZ = EvaluatedResult.Channels.Find(PositionZName);
-
-			TArray<FVector> CachedPositions; 
-			const bool bHasPositions = PendingPX && PendingPY && PendingPZ;
-			if (bHasPositions)
-			{
-				CachedPositions.SetNum(NumCurves);
-				for(int32 ParticleIndex = 0; ParticleIndex < NumCurves; ++ParticleIndex)
-				{
-					const int32 TransformIndex = EvaluatedResult.ParticleIndices[ParticleIndex];
-					CachedPositions[ParticleIndex] = FVector((*PendingPX)[TransformIndex], (*PendingPY)[TransformIndex], (*PendingPZ)[TransformIndex]);
-				}
-			}
-			TArray<FVector>* CachedPositionsPtr = bHasPositions ? &CachedPositions : nullptr;
-
-			TArray<FVector> CachedVelocities; 
-			const bool bHasVelocities = PendingVX && PendingVY && PendingVZ;
-			if (bHasVelocities)
-			{
-				CachedVelocities.SetNum(NumCurves);
-				for(int32 ParticleIndex = 0; ParticleIndex < NumCurves; ++ParticleIndex)
-				{
-					const int32 TransformIndex = EvaluatedResult.ParticleIndices[ParticleIndex];
-					CachedVelocities[ParticleIndex] = FVector((*PendingVX)[TransformIndex], (*PendingVY)[TransformIndex], (*PendingVZ)[TransformIndex]);
-				}
-			}
-			TArray<FVector>* CachedVelocitiesPtr = bHasVelocities ? &CachedVelocities : nullptr;
+			FClothingSimulationCacheData CacheData;
+			::Chaos::FClothingCacheSchema::LoadCacheData(InCache, InTime, CacheData);
 
 			FClothSimulationProxy* Proxy = ClothComp->ClothSimulationProxy.Get();
-			Proxy->Tick_GameThread(0, CachedPositionsPtr, CachedVelocitiesPtr);
+
+			Proxy->Tick_GameThread(0, &CacheData);
 			Proxy->CompleteParallelSimulation_GameThread();
 			ClothComp->MarkRenderDynamicDataDirty();
 			ClothComp->DoDeferredRenderUpdates_Concurrent();
diff --git a/Engine/Plugins/Experimental/ChaosClothAsset/Source/ChaosClothAssetEngine/Private/ChaosClothAsset/ClothComponentCacheAdapter.h b/Engine/Plugins/Experimental/ChaosClothAsset/Source/ChaosClothAssetEngine/Private/ChaosClothAsset/ClothComponentCacheAdapter.h
--- a/Engine/Plugins/Experimental/ChaosClothAsset/Source/ChaosClothAssetEngine/Private/ChaosClothAsset/ClothComponentCacheAdapter.h
+++ b/Engine/Plugins/Experimental/ChaosClothAsset/Source/ChaosClothAssetEngine/Private/ChaosClothAsset/ClothComponentCacheAdapter.h
@@ -44,12 +44,5 @@ namespace UE::Chaos::ClothAsset
 		/** Return the cloth component.
 		 */
 		UChaosClothComponent* GetClothComponent(UPrimitiveComponent* InComponent) const;
-
-		inline static const FName VelocityXName = TEXT("VelocityX");
-		inline static const FName VelocityYName = TEXT("VelocityY");
-		inline static const FName VelocityZName = TEXT("VelocityZ");
-		inline static const FName PositionXName = TEXT("PositionX");
-		inline static const FName PositionYName = TEXT("PositionY");
-		inline static const FName PositionZName = TEXT("PositionZ");
 	};
 }    // namespace Chaos
diff --git a/Engine/Plugins/Experimental/ChaosClothAsset/Source/ChaosClothAssetEngine/Private/ChaosClothAsset/ClothSimulationContext.cpp b/Engine/Plugins/Experimental/ChaosClothAsset/Source/ChaosClothAssetEngine/Private/ChaosClothAsset/ClothSimulationContext.cpp
--- a/Engine/Plugins/Experimental/ChaosClothAsset/Source/ChaosClothAssetEngine/Private/ChaosClothAsset/ClothSimulationContext.cpp
+++ b/Engine/Plugins/Experimental/ChaosClothAsset/Source/ChaosClothAssetEngine/Private/ChaosClothAsset/ClothSimulationContext.cpp
@@ -9,7 +9,7 @@
 
 namespace UE::Chaos::ClothAsset
 {
-	void FClothSimulationContext::Fill(const UChaosClothComponent& ClothComponent, float InDeltaTime, float MaxDeltaTime, bool bIsInitialization, const TArray<FVector>* InCachedPositions, const TArray<FVector>* InCachedVelocities)
+	void FClothSimulationContext::Fill(const UChaosClothComponent& ClothComponent, float InDeltaTime, float MaxDeltaTime, bool bIsInitialization, FClothingSimulationCacheData* InCacheData)
 	{
 		// Set the time
 		DeltaTime = FMath::Min(InDeltaTime, MaxDeltaTime);
@@ -129,22 +129,13 @@ namespace UE::Chaos::ClothAsset
 			WindVelocity = FVector::ZeroVector;
 		}
 
-		if (InCachedPositions != nullptr)
+		if (InCacheData)
 		{
-			CachedPositions = *InCachedPositions;
+			CacheData = MoveTemp(*InCacheData);
 		}
 		else
 		{
-			CachedPositions.Reset();
-		}
-
-		if (InCachedVelocities != nullptr)
-		{
-			CachedVelocities = *InCachedVelocities;
-		}
-		else
-		{
-			CachedVelocities.Reset();
+			CacheData.Reset();
 		}
 	}
 }  // End namespace UE::Chaos::ClothAsset
diff --git a/Engine/Plugins/Experimental/ChaosClothAsset/Source/ChaosClothAssetEngine/Private/ChaosClothAsset/ClothSimulationContext.h b/Engine/Plugins/Experimental/ChaosClothAsset/Source/ChaosClothAssetEngine/Private/ChaosClothAsset/ClothSimulationContext.h
--- a/Engine/Plugins/Experimental/ChaosClothAsset/Source/ChaosClothAssetEngine/Private/ChaosClothAsset/ClothSimulationContext.h
+++ b/Engine/Plugins/Experimental/ChaosClothAsset/Source/ChaosClothAssetEngine/Private/ChaosClothAsset/ClothSimulationContext.h
@@ -4,6 +4,7 @@
 
 #include "Containers/Array.h"
 #include "Math/Transform.h"
+#include "ClothingSimulation.h"
 
 class UChaosClothComponent;
 
@@ -47,13 +48,10 @@ namespace UE::Chaos::ClothAsset
 		/** Wind velocity at the component location. */
 		FVector WindVelocity;
 
-		/** Positions used by ChaosCache */
-		TArray<FVector> CachedPositions;
-
-		/** Velocities used by ChaosCache */
-		TArray<FVector> CachedVelocities;
+		/* Data used by Chaos Cache*/
+		FClothingSimulationCacheData CacheData;
 
 		/** Fill the context from data collected from the specified component. */
-		void Fill(const UChaosClothComponent& ClothComponent, float InDeltaTime, float MaxDeltaTime, bool bIsInitialization = false, const TArray<FVector>* InCachedPositions = nullptr, const TArray<FVector>* InCachedVelocities = nullptr);
+		void Fill(const UChaosClothComponent& ClothComponent, float InDeltaTime, float MaxDeltaTime, bool bIsInitialization = false, FClothingSimulationCacheData* CacheData = nullptr);
 	};
 }
diff --git a/Engine/Plugins/Experimental/ChaosClothAsset/Source/ChaosClothAssetEngine/Private/ChaosClothAsset/ClothSimulationProxy.cpp b/Engine/Plugins/Experimental/ChaosClothAsset/Source/ChaosClothAssetEngine/Private/ChaosClothAsset/ClothSimulationProxy.cpp
--- a/Engine/Plugins/Experimental/ChaosClothAsset/Source/ChaosClothAssetEngine/Private/ChaosClothAsset/ClothSimulationProxy.cpp
+++ b/Engine/Plugins/Experimental/ChaosClothAsset/Source/ChaosClothAssetEngine/Private/ChaosClothAsset/ClothSimulationProxy.cpp
@@ -12,6 +12,7 @@
 #include "ChaosCloth/ChaosClothingSimulationSolver.h"
 #include "ChaosCloth/ChaosClothVisualization.h"
 #include "PhysicsEngine/PhysicsSettings.h"
+#include "ClothingSimulation.h"
 
 #if INTEL_ISPC
 #include "ClothSimulationProxy.ispc.generated.h"
@@ -148,14 +149,14 @@ PRAGMA_ENABLE_DEPRECATION_WARNINGS
 		CompleteParallelSimulation_GameThread();
 	}
 
-	void FClothSimulationProxy::Tick_GameThread(float DeltaTime, const TArray<FVector>* CachedPositions, const TArray<FVector>* CachedVelocities)
+	void FClothSimulationProxy::Tick_GameThread(float DeltaTime, FClothingSimulationCacheData* CacheData)
 	{
 		SCOPE_CYCLE_COUNTER(STAT_ClothSimulationProxy_TickGame);
 
 		// Fill a new context, note the context is also needed when the simulation is suspended
-		ClothSimulationContext.Fill(ClothComponent, DeltaTime, MaxDeltaTime, false, CachedPositions, CachedVelocities);
+		ClothSimulationContext.Fill(ClothComponent, DeltaTime, MaxDeltaTime, false, CacheData);
 
-		const bool bUseCache = ClothSimulationContext.CachedPositions.Num() > 0 || ClothSimulationContext.CachedVelocities.Num() > 0;
+		const bool bUseCache = ClothSimulationContext.CacheData.CacheIndices.Num() > 0;
 		if (bUseCache)
 		{
 			Solver->SetEnableSolver(false);
@@ -186,7 +187,7 @@ PRAGMA_ENABLE_DEPRECATION_WARNINGS
 
 		TRACE_CPUPROFILER_EVENT_SCOPE(FClothSimulationProxy_TickPhysics);
 		SCOPE_CYCLE_COUNTER(STAT_ClothSimulationProxy_TickPhysics);
-		const bool bUseCache = ClothSimulationContext.CachedPositions.Num() > 0 || ClothSimulationContext.CachedVelocities.Num() > 0;
+		const bool bUseCache = ClothSimulationContext.CacheData.HasData();
 
 		if (ClothSimulationContext.DeltaTime == 0.f && !bUseCache)
 		{
@@ -229,7 +230,7 @@ PRAGMA_ENABLE_DEPRECATION_WARNINGS
 		}
 		else
 		{
-			Solver->UpdateFromCache(ClothSimulationContext.CachedPositions, ClothSimulationContext.CachedVelocities);
+			Solver->UpdateFromCache(ClothSimulationContext.CacheData);
 		}
 
 		// Keep the actual used number of iterations for the stats
diff --git a/Engine/Plugins/Experimental/ChaosClothAsset/Source/ChaosClothAssetEngine/Public/ChaosClothAsset/ClothSimulationProxy.h b/Engine/Plugins/Experimental/ChaosClothAsset/Source/ChaosClothAssetEngine/Public/ChaosClothAsset/ClothSimulationProxy.h
--- a/Engine/Plugins/Experimental/ChaosClothAsset/Source/ChaosClothAssetEngine/Public/ChaosClothAsset/ClothSimulationProxy.h
+++ b/Engine/Plugins/Experimental/ChaosClothAsset/Source/ChaosClothAssetEngine/Public/ChaosClothAsset/ClothSimulationProxy.h
@@ -23,6 +23,7 @@ class UChaosClothComponent;
 class UChaosClothAsset;
 struct FChaosClothSimulationModel;
 struct FReferenceSkeleton;
+struct FClothingSimulationCacheData;
 
 namespace UE::Chaos::ClothAsset
 {
@@ -42,7 +43,7 @@ namespace UE::Chaos::ClothAsset
 		FClothSimulationProxy& operator=(const FClothSimulationProxy&) = delete;
 		FClothSimulationProxy& operator=(FClothSimulationProxy&&) = delete;
 
-		void Tick_GameThread(float DeltaTime, const TArray<FVector>* CachedPositions = nullptr, const TArray<FVector>* CachedVelocities = nullptr);
+		void Tick_GameThread(float DeltaTime, FClothingSimulationCacheData* CacheData = nullptr);
 
 		/** Wait for the parallel task to complete if one was running, and update the simulation data. */
 		void CompleteParallelSimulation_GameThread();
diff --git a/Engine/Source/Runtime/ClothingSystemRuntimeCommon/Private/ClothingSimulation.cpp b/Engine/Source/Runtime/ClothingSystemRuntimeCommon/Private/ClothingSimulation.cpp
--- a/Engine/Source/Runtime/ClothingSystemRuntimeCommon/Private/ClothingSimulation.cpp
+++ b/Engine/Source/Runtime/ClothingSystemRuntimeCommon/Private/ClothingSimulation.cpp
@@ -20,6 +20,7 @@ static TAutoConsoleVariable<float> GClothMaxDeltaTimeTeleportMultiplier(
 	TEXT("A multiplier of the MaxPhysicsDelta time at which we will automatically just teleport cloth to its new location\n")
 	TEXT(" default: 1.5"));
 
+PRAGMA_DISABLE_DEPRECATION_WARNINGS // CachedPositions, CachedVelocities
 FClothingSimulationContextCommon::FClothingSimulationContextCommon()
 	: ComponentToWorld(FTransform::Identity)
 	, WorldGravity(FVector::ZeroVector)
@@ -31,9 +32,12 @@ FClothingSimulationContextCommon::FClothingSimulationContextCommon()
 	, MaxDistanceScale(1.f)
 	, PredictedLod(INDEX_NONE)
 {}
+PRAGMA_ENABLE_DEPRECATION_WARNINGS
 
+PRAGMA_DISABLE_DEPRECATION_WARNINGS // CachedPositions, CachedVelocities
 FClothingSimulationContextCommon::~FClothingSimulationContextCommon()
 {}
+PRAGMA_ENABLE_DEPRECATION_WARNINGS
 
 void FClothingSimulationContextCommon::Fill(const USkeletalMeshComponent* InComponent, float InDeltaSeconds, float InMaxPhysicsDelta)
 {
diff --git a/Engine/Source/Runtime/ClothingSystemRuntimeCommon/Public/ClothingSimulation.h b/Engine/Source/Runtime/ClothingSystemRuntimeCommon/Public/ClothingSimulation.h
--- a/Engine/Source/Runtime/ClothingSystemRuntimeCommon/Public/ClothingSimulation.h
+++ b/Engine/Source/Runtime/ClothingSystemRuntimeCommon/Public/ClothingSimulation.h
@@ -2,6 +2,7 @@
 #pragma once
 
 #include "ClothingSimulationInterface.h"
+#include "ClothingSimulationCacheData.h"
 #include "Containers/Array.h"
 #include "Math/Transform.h"
 #include "Math/Vector.h"
@@ -75,10 +76,15 @@ public:
 	// The predicted LOD of the skeletal mesh component running the simulation
 	int32 PredictedLod;
 	
-	// World space cached positions for the kinematics targets
+	// Data read from the cache.
+	FClothingSimulationCacheData CacheData;
+
+	// World space cached positions for the kinematics targets.
+	UE_DEPRECATED(5.3, "Use CacheData.CachedPositions instead")
 	TArray<FVector> CachedPositions;
 
-	// World space cached velocities for the kinematics targets
+	// World space cached velocities for the kinematics targets.
+	UE_DEPRECATED(5.3, "Use CacheData.CachedVelocities instead")
 	TArray<FVector> CachedVelocities;
 };
 
diff --git a/Engine/Source/Runtime/ClothingSystemRuntimeCommon/Public/ClothingSimulationCacheData.h b/Engine/Source/Runtime/ClothingSystemRuntimeCommon/Public/ClothingSimulationCacheData.h
new file mode 100644
--- /dev/null
+++ b/Engine/Source/Runtime/ClothingSystemRuntimeCommon/Public/ClothingSimulationCacheData.h
@@ -0,0 +1,55 @@
+// Copyright Epic Games, Inc. All Rights Reserved.
+#pragma once
+
+#include "Containers/Array.h"
+#include "Containers/Map.h"
+#include "Math/Transform.h"
+#include "Math/Vector.h"
+
+struct FClothingSimulationCacheData
+{
+	// Solver space cached positions for the kinematics targets. (Sparse--use CacheIndices to map to full)
+	TArray<FVector> CachedPositions;
+
+	// Solver space cached velocities for the kinematics targets. (Sparse--use CacheIndices to map to full)
+	TArray<FVector> CachedVelocities;
+
+	// Indices used to map from CachedPositions/Velocities to solver indices
+	TArray<int32> CacheIndices;
+
+	// Cached ReferenceSpaceTransform per cloth. Key is Cloth GroupId
+	TMap<int32, FTransform> CachedReferenceSpaceTransforms;
+
+	FClothingSimulationCacheData() = default;
+	FClothingSimulationCacheData(const FClothingSimulationCacheData& Other) = default;
+	FClothingSimulationCacheData(FClothingSimulationCacheData&& Other)
+		: CachedPositions(MoveTemp(Other.CachedPositions))
+		, CachedVelocities(MoveTemp(Other.CachedVelocities))
+		, CacheIndices(MoveTemp(Other.CacheIndices))
+		, CachedReferenceSpaceTransforms(MoveTemp(Other.CachedReferenceSpaceTransforms))
+	{}
+
+	FClothingSimulationCacheData& operator=(const FClothingSimulationCacheData& Other) = default;
+	FClothingSimulationCacheData& operator=(FClothingSimulationCacheData&& Other)
+	{
+		CachedPositions = MoveTemp(Other.CachedPositions);
+		CachedVelocities = MoveTemp(Other.CachedVelocities);
+		CacheIndices = MoveTemp(Other.CacheIndices);
+		CachedReferenceSpaceTransforms = MoveTemp(Other.CachedReferenceSpaceTransforms);
+		return *this;
+	}
+
+	void Reset()
+	{
+		CachedPositions.Reset();
+		CachedVelocities.Reset();
+		CacheIndices.Reset();
+		CachedReferenceSpaceTransforms.Reset();
+	}
+
+	bool HasData() const
+	{
+		checkSlow(CachedPositions.Num() == CacheIndices.Num());
+		return CacheIndices.Num() > 0;
+	}
+};
\ No newline at end of file
