Chaos Cloth:
- Fixed the dirty property updates when LOD changes by updating all LODs' properties at all time.
- Fixed the VelocityAndPressureField weight maps crash by checking the property indices before trying to retrieve the weight maps.
- Fixed the VelocityAndPressureField property update when geometry (LOD) changes, and correctly initialize the property indices.
- Fixed the VelocityAndPressureField property update when the Aerodynamics are disabled.
- Fixed the VelocityAndPressureField FluidDensity world scale.
- Fixed the cloth config property name AirDensity to FluidDensity so that it gets correctly picked up by the VelocityAndPressureField code.

#rnx

diff --git a/Engine/Plugins/Experimental/ChaosCloth/Source/Chaos/Private/ChaosCloth/ChaosClothingSimulationCloth.cpp b/Engine/Plugins/Experimental/ChaosCloth/Source/Chaos/Private/ChaosCloth/ChaosClothingSimulationCloth.cpp
--- a/Engine/Plugins/Experimental/ChaosCloth/Source/Chaos/Private/ChaosCloth/ChaosClothingSimulationCloth.cpp
+++ b/Engine/Plugins/Experimental/ChaosCloth/Source/Chaos/Private/ChaosCloth/ChaosClothingSimulationCloth.cpp
@@ -29,7 +29,7 @@ namespace ClothingSimulationClothDefault
 	constexpr float Drag = 0.035f;
 	constexpr float Lift = 0.035f;
 	constexpr float Pressure = 0.f;
-	constexpr float AirDensity = 1.225e-6f;
+	constexpr float AirDensity = 1.225f;
 	constexpr float GravityScale = 1.f;
 	constexpr float GravityZOverride = -980.665f;
 	constexpr float VelocityScale = 0.75f;
@@ -353,19 +353,12 @@ PRAGMA_ENABLE_DEPRECATION_WARNINGS
 	int32 InLODIndexOverride,
 	const TVec2<FRealSingle>& InEdgeDampingRatio,
 	const TVec2<FRealSingle>& InBendingDampingRatio)
-	: Mesh(nullptr)
-	, Colliders()
-	, GroupId(InGroupId)
+	: GroupId(InGroupId)
 	, PropertyCollection(MakeShared<FManagedArrayCollection>())
 	, bUseLODIndexOverride(bInUseLODIndexOverride)
 	, LODIndexOverride(InLODIndexOverride)
-	, bNeedsReset(false)
-	, bNeedsTeleport(false)
-	, NumActiveKinematicParticles(0)
-	, NumActiveDynamicParticles(0)
 {
 	// Turn parameters into a config
-	PropertyCollection = MakeShared<FManagedArrayCollection>();
 	Softs::FCollectionPropertyMutableFacade Properties(PropertyCollection);
 	Properties.DefineSchema();
 
@@ -474,7 +467,7 @@ PRAGMA_ENABLE_DEPRECATION_WARNINGS
 		Properties.SetWeightedValue(LiftIndex, InLift[0], InLift[1]);
 		Properties.SetStringValue(LiftIndex, TEXT("Lift"));
 
-		Properties.AddValue(TEXT("AirDensity"), ClothingSimulationClothDefault::AirDensity, bEnable, bAnimatable);
+		Properties.AddValue(TEXT("FluidDensity"), ClothingSimulationClothDefault::AirDensity, bEnable, bAnimatable);
 	}
 
 	// Pressure
@@ -896,6 +889,13 @@ PRAGMA_DISABLE_DEPRECATION_WARNINGS  // TODO: CHAOS_IS_CLOTHINGSIMULATIONMESH_AB
 	Mesh->Update(Solver, PrevLODIndex, LODIndex, PrevOffset, Offset);
 PRAGMA_ENABLE_DEPRECATION_WARNINGS
 
+	// Update all LODs dirty properties, since it is easier done than re-updating all properties when switching LODs
+	for (TUniquePtr<FLODData>& LODDatum : LODData)
+	{
+		LODDatum->Update(Solver, this);
+	}
+
+	// Retrieve config
 	check(Config);
 	const Softs::FCollectionPropertyFacade& ConfigProperties = Config->GetProperties();
 
@@ -927,18 +927,8 @@ PRAGMA_ENABLE_DEPRECATION_WARNINGS
 				Solver->GetParticleXs(Offset),
 				Solver->GetParticleVs(Offset));
 
-				// Update the wind velocity field for the new LOD mesh
-				FString DragString = TEXT("Drag");            // The default name is the property name
-				FString LiftString = TEXT("Lift");            //
-				FString PressureString = TEXT("Pressure");    //
-				DragString = ConfigProperties.GetStringValue(DragString, DragString);
-				LiftString = ConfigProperties.GetStringValue(LiftString, LiftString);
-				PressureString = ConfigProperties.GetStringValue(PressureString, PressureString);
-				const TConstArrayView<FRealSingle> DragMultipliers = LODData[LODIndex]->WeightMaps.FindRef(DragString);
-				const TConstArrayView<FRealSingle> LiftMultipliers = LODData[LODIndex]->WeightMaps.FindRef(LiftString);
-				const TConstArrayView<FRealSingle> PressureMultipliers = LODData[LODIndex]->WeightMaps.FindRef(PressureString);
-
-				Solver->SetWindAndPressureGeometry(GroupId, GetTriangleMesh(Solver), DragMultipliers, LiftMultipliers, PressureMultipliers);
+			// Update the wind velocity field for the new LOD mesh
+			Solver->SetWindAndPressureGeometry(GroupId, GetTriangleMesh(Solver), ConfigProperties, LODData[LODIndex]->WeightMaps);
 		}
 		else
 		{
@@ -947,11 +937,9 @@ PRAGMA_ENABLE_DEPRECATION_WARNINGS
 		}
 	}
 
+	// Update Cloth group parameters
 	if (LODIndex != INDEX_NONE)
 	{
-		// Update Cloth group parameters  TODO: Cloth groups should exist as their own node object so that they can be used by several cloth objects
-		LODData[LODIndex]->Update(Solver, this);
-
 		// TODO: Move all groupID updates out of the cloth update to allow to use of the same GroupId with different cloths
 
 		// Set the reference input velocity and deal with teleport & reset; external forces depends on these values, so they must be initialized before then
@@ -1002,46 +990,9 @@ PRAGMA_ENABLE_DEPRECATION_WARNINGS
 		const bool bUsePointBasedWindModel = ConfigProperties.GetValue<bool>(TEXT("UsePointBasedWindModel"));
 		Solver->AddExternalForces(GroupId, bUsePointBasedWindModel);
 
-		int32 PressurePropertyKeyIndex;
-		const FVec2f Pressure(ConfigProperties.GetWeightedFloatValue(TEXT("Pressure"), ClothingSimulationClothDefault::Pressure, &PressurePropertyKeyIndex));
-		if (bUsePointBasedWindModel && ClothingSimulationClothConsoleVariables::CVarLegacyDisablesAccurateWind.GetValueOnAnyThread())
-		{
-			Solver->SetWindAndPressureProperties(GroupId, TVec2<FRealSingle>(0.), TVec2<FRealSingle>(0.), (FRealSingle)0., Pressure);  // Disable the wind velocity field
-			
-			if (ConfigProperties.IsStringDirty(PressurePropertyKeyIndex))
-			{
-				const FString& PressureString = ConfigProperties.GetStringValue(PressurePropertyKeyIndex);
-				Solver->SetWindAndPressureMultipliers(
-					GroupId,
-					TConstArrayView<FRealSingle>(),
-					TConstArrayView<FRealSingle>(),
-					LODData[LODIndex]->WeightMaps.FindRef(PressureString));
-			}
-		}
-		else
-		{
-			int32 DragPropertyKeyIndex;
-			int32 LiftPropertyKeyIndex;
-			const FVec2f Drag(ConfigProperties.GetWeightedFloatValue(TEXT("Drag"), ClothingSimulationClothDefault::Drag, &DragPropertyKeyIndex));
-			const FVec2f Lift(ConfigProperties.GetWeightedFloatValue(TEXT("Lift"), ClothingSimulationClothDefault::Lift, &LiftPropertyKeyIndex));
-			const FRealSingle AirDensity = ConfigProperties.GetValue<float>(TEXT("AirDensity"), ClothingSimulationClothDefault::AirDensity);
-
-			Solver->SetWindAndPressureProperties(GroupId, Drag, Lift, AirDensity, Pressure);
-
-			if (ConfigProperties.IsStringDirty(DragPropertyKeyIndex) ||
-				ConfigProperties.IsStringDirty(LiftPropertyKeyIndex) ||
-				ConfigProperties.IsStringDirty(PressurePropertyKeyIndex))
-			{
-				const FString& DragString = ConfigProperties.GetStringValue(DragPropertyKeyIndex);
-				const FString& LiftString = ConfigProperties.GetStringValue(LiftPropertyKeyIndex);
-				const FString& PressureString = ConfigProperties.GetStringValue(PressurePropertyKeyIndex);
-				Solver->SetWindAndPressureMultipliers(
-					GroupId,
-					LODData[LODIndex]->WeightMaps.FindRef(DragString),
-					LODData[LODIndex]->WeightMaps.FindRef(LiftString),
-					LODData[LODIndex]->WeightMaps.FindRef(PressureString));
-			}
-		}
+		const bool bPointBasedWindDisablesAccurateWind = ClothingSimulationClothConsoleVariables::CVarLegacyDisablesAccurateWind.GetValueOnAnyThread();
+		const bool bEnableAerodynamics = !(bUsePointBasedWindModel && bPointBasedWindDisablesAccurateWind);
+		Solver->SetWindAndPressureProperties(GroupId, ConfigProperties, LODData[LODIndex]->WeightMaps, bEnableAerodynamics);
 
 		const FVec3f WindVelocity = ConfigProperties.GetValue<FVector3f>(TEXT("WindVelocity"));
 		Solver->SetWindVelocity(GroupId, WindVelocity + Solver->GetWindVelocity());
@@ -1077,9 +1028,6 @@ void FClothingSimulationCloth::PostUpdate(FClothingSimulationSolver* Solver)
 		// Update normals
 		LODData[LODIndex]->UpdateNormals(Solver);
 	}
-
-	// Clear the cloth config dirty flags
-	Config->GetProperties().ClearDirtyFlags();
 }
 
 void FClothingSimulationCloth::UpdateFromCache(const FClothingSimulationCacheData& CacheData)
@@ -1173,9 +1121,10 @@ void FClothingSimulationCloth::SetDampingProperties(FRealSingle InDampingCoeffic
 
 void FClothingSimulationCloth::SetAerodynamicsProperties(const TVec2<FRealSingle>& InDrag, const TVec2<FRealSingle>& InLift, FRealSingle InAirDensity, const FVec3& InWindVelocity)
 {
+	constexpr float WordScale = 100.f;  // Unreal's world unit is the cm
 	Config->GetProperties().SetWeightedFloatValue(TEXT("Drag"), FVector2f(InDrag));
 	Config->GetProperties().SetWeightedFloatValue(TEXT("Lift"), FVector2f(InLift));
-	Config->GetProperties().SetValue(TEXT("AirDensity"), (float)InAirDensity);
+	Config->GetProperties().SetValue(TEXT("FluidDensity"), (float)InAirDensity * WordScale);  // AirDensity is here in kg/cm^3 for legacy reason but muust be in kg/m^3 in the config UI
 	Config->GetProperties().SetValue(TEXT("WindVelocity"), FVector3f(InWindVelocity));
 }
 
diff --git a/Engine/Plugins/Experimental/ChaosCloth/Source/Chaos/Private/ChaosCloth/ChaosClothingSimulationConfig.cpp b/Engine/Plugins/Experimental/ChaosCloth/Source/Chaos/Private/ChaosCloth/ChaosClothingSimulationConfig.cpp
--- a/Engine/Plugins/Experimental/ChaosCloth/Source/Chaos/Private/ChaosCloth/ChaosClothingSimulationConfig.cpp
+++ b/Engine/Plugins/Experimental/ChaosCloth/Source/Chaos/Private/ChaosCloth/ChaosClothingSimulationConfig.cpp
@@ -160,6 +160,9 @@ namespace Chaos
 				const int32 LiftIndex = Properties->AddProperty(TEXT("Lift"), bEnable, bAnimatable);
 				Properties->SetWeightedValue(LiftIndex, ClothConfig->Lift.Low, ClothConfig->Lift.High);
 				Properties->SetStringValue(LiftIndex, TEXT("Lift"));
+
+				constexpr float AirDensity = 1.225f;
+				Properties->AddValue(TEXT("FluidDensity"), AirDensity, bEnable, bAnimatable);
 			}
 
 			// Pressure
diff --git a/Engine/Plugins/Experimental/ChaosCloth/Source/Chaos/Private/ChaosCloth/ChaosClothingSimulationSolver.cpp b/Engine/Plugins/Experimental/ChaosCloth/Source/Chaos/Private/ChaosCloth/ChaosClothingSimulationSolver.cpp
--- a/Engine/Plugins/Experimental/ChaosCloth/Source/Chaos/Private/ChaosCloth/ChaosClothingSimulationSolver.cpp
+++ b/Engine/Plugins/Experimental/ChaosCloth/Source/Chaos/Private/ChaosCloth/ChaosClothingSimulationSolver.cpp
@@ -782,37 +782,54 @@ int32 FClothingSimulationSolver::GetNumSubsteps() const
 
 void FClothingSimulationSolver::SetWindVelocity(uint32 GroupId, const TVec3<FRealSingle>& InWindVelocity)
 {
-	Softs::FVelocityAndPressureField& VelocityField = Evolution->GetVelocityAndPressureField(GroupId);
-	VelocityField.SetVelocity(Softs::FSolverVec3(InWindVelocity));
+	Softs::FVelocityAndPressureField& VelocityAndPressureField = Evolution->GetVelocityAndPressureField(GroupId);
+	VelocityAndPressureField.SetVelocity(Softs::FSolverVec3(InWindVelocity));
 }
 
 void FClothingSimulationSolver::SetWindAndPressureGeometry(
 	uint32 GroupId,
 	const FTriangleMesh& TriangleMesh,
-	const TConstArrayView<FRealSingle>& DragMultipliers,
-	const TConstArrayView<FRealSingle>& LiftMultipliers,
-	const TConstArrayView<FRealSingle>& PressureMultipliers)
+	const Softs::FCollectionPropertyConstFacade& PropertyCollection,
+	const TMap<FString, TConstArrayView<FRealSingle>>& WeightMaps)
 {
-	Softs::FVelocityAndPressureField& VelocityField = Evolution->GetVelocityAndPressureField(GroupId);
-	VelocityField.SetGeometry(&TriangleMesh, DragMultipliers, LiftMultipliers, PressureMultipliers);
+	Softs::FVelocityAndPressureField& VelocityAndPressureField = Evolution->GetVelocityAndPressureField(GroupId);
+	VelocityAndPressureField.SetGeometry(&TriangleMesh, PropertyCollection, WeightMaps, ClothingSimulationSolverConstant::WorldScale);
 }
 
-void FClothingSimulationSolver::SetWindAndPressureMultipliers(
+void FClothingSimulationSolver::SetWindAndPressureGeometry(
 	uint32 GroupId,
+	const FTriangleMesh& TriangleMesh,
 	const TConstArrayView<FRealSingle>& DragMultipliers,
 	const TConstArrayView<FRealSingle>& LiftMultipliers,
 	const TConstArrayView<FRealSingle>& PressureMultipliers)
 {
-	Softs::FVelocityAndPressureField& VelocityField = Evolution->GetVelocityAndPressureField(GroupId);
-	VelocityField.SetMultipliers(DragMultipliers, LiftMultipliers, PressureMultipliers);
+	Softs::FVelocityAndPressureField& VelocityAndPressureField = Evolution->GetVelocityAndPressureField(GroupId);
+	VelocityAndPressureField.SetGeometry(&TriangleMesh, DragMultipliers, LiftMultipliers, PressureMultipliers);
 }
 
-void FClothingSimulationSolver::SetWindAndPressureProperties(uint32 GroupId, const TVec2<FRealSingle>& Drag, const TVec2<FRealSingle>& Lift, FRealSingle AirDensity, const TVec2<FRealSingle>& Pressure)
+void FClothingSimulationSolver::SetWindAndPressureProperties(
+	uint32 GroupId,
+	const Softs::FCollectionPropertyConstFacade& PropertyCollection,
+	const TMap<FString, TConstArrayView<FRealSingle>>& WeightMaps,
+	bool bEnableAerodynamics)
 {
-	Softs::FVelocityAndPressureField& VelocityField = Evolution->GetVelocityAndPressureField(GroupId);
+	Softs::FVelocityAndPressureField& VelocityAndPressureField = Evolution->GetVelocityAndPressureField(GroupId);
+	VelocityAndPressureField.SetProperties(PropertyCollection, WeightMaps, ClothingSimulationSolverConstant::WorldScale, bEnableAerodynamics);
+}
 
-	// UI Pressure is in kg/m s^2. Need to convert to kg/cm s^2 for solver.
-	VelocityField.SetProperties(Drag, Lift, AirDensity, Pressure / ClothingSimulationSolverConstant::WorldScale);
+void FClothingSimulationSolver::SetWindAndPressureProperties(
+	uint32 GroupId,
+	const TVec2<FRealSingle>& Drag,
+	const TVec2<FRealSingle>& Lift,
+	FRealSingle FluidDensity,
+	const TVec2<FRealSingle>& Pressure)
+{
+	Softs::FVelocityAndPressureField& VelocityAndPressureField = Evolution->GetVelocityAndPressureField(GroupId);
+	VelocityAndPressureField.SetProperties(
+		Drag,
+		Lift,
+		FluidDensity / FMath::Cube(ClothingSimulationSolverConstant::WorldScale),  // Fluid density is given in kg/m^3. Need to convert to kg/cm^3 for solver.
+		Pressure / ClothingSimulationSolverConstant::WorldScale);  // UI Pressure is in kg/m s^2. Need to convert to kg/cm s^2 for solver.
 }
 
 const Softs::FVelocityAndPressureField& FClothingSimulationSolver::GetWindVelocityAndPressureField(uint32 GroupId) const
@@ -1189,9 +1206,6 @@ void FClothingSimulationSolver::Update(Softs::FSolverReal InDeltaTime)
 			FClothingSimulationCloth* const Cloth = Cloths[ClothIndex];
 			Cloth->PostUpdate(this);
 		}, /*bForceSingleThreaded =*/ !bClothSolverParallelClothPostUpdate);
-
-		// Clear the solver config dirty flags
-		Config->GetProperties().ClearDirtyFlags();
 	}
 
 	// Save old space location for next update
diff --git a/Engine/Plugins/Experimental/ChaosCloth/Source/Chaos/Public/ChaosCloth/ChaosClothingSimulationCloth.h b/Engine/Plugins/Experimental/ChaosCloth/Source/Chaos/Public/ChaosCloth/ChaosClothingSimulationCloth.h
--- a/Engine/Plugins/Experimental/ChaosCloth/Source/Chaos/Public/ChaosCloth/ChaosClothingSimulationCloth.h
+++ b/Engine/Plugins/Experimental/ChaosCloth/Source/Chaos/Public/ChaosCloth/ChaosClothingSimulationCloth.h
@@ -111,7 +111,7 @@ PRAGMA_ENABLE_DEPRECATION_WARNINGS
 		void SetCollisionProperties(FRealSingle InCollisionThickness, FRealSingle InFrictionCoefficient, bool bInUseCCD, FRealSingle InSelfCollisionThickness);
 		void SetBackstopProperties(bool bInEnableBackstop);
 		void SetDampingProperties(FRealSingle InDampingCoefficient, FRealSingle InLocalDampingCoefficient = 0.f);
-		void SetAerodynamicsProperties(const TVec2<FRealSingle>& InDrag, const TVec2<FRealSingle>& InLift, FRealSingle InAirDensity, const FVec3& InWindVelocity);
+		void SetAerodynamicsProperties(const TVec2<FRealSingle>& InDrag, const TVec2<FRealSingle>& InLift, FRealSingle InAirDensity, const FVec3& InWindVelocity);  // AirDensity is here in kg/cm^3 for legacy reason (kg/m^3 in UI)
 		void SetPressureProperties(const TVec2<FRealSingle>& InPressure);
 		void SetGravityProperties(FRealSingle InGravityScale, bool bInUseGravityOverride, const FVec3& InGravityOverride);
 		void SetAnimDriveProperties(const TVec2<FRealSingle>& InAnimDriveStiffness, const TVec2<FRealSingle>& InAnimDriveDamping);
@@ -195,20 +195,20 @@ PRAGMA_ENABLE_DEPRECATION_WARNINGS
 
 		// Cloth parameters
 PRAGMA_DISABLE_DEPRECATION_WARNINGS  // TODO: CHAOS_IS_CLOTHINGSIMULATIONMESH_ABSTRACT
-		FClothingSimulationMesh* Mesh;
+		FClothingSimulationMesh* Mesh = nullptr;
 PRAGMA_ENABLE_DEPRECATION_WARNINGS
-		FClothingSimulationConfig* Config;
+		FClothingSimulationConfig* Config = nullptr;
 		TArray<FClothingSimulationCollider*> Colliders;
-		uint32 GroupId;
+		uint32 GroupId = 0;
 
 		TSharedPtr<FManagedArrayCollection> PropertyCollection;  // Used for backward compatibility only, otherwise the properties are owned by the Config
 
-		FRealSingle MaxDistancesMultiplier;  // Legacy multiplier
+		FRealSingle MaxDistancesMultiplier = 1.f;  // Legacy multiplier
 
-		bool bUseLODIndexOverride;
-		int32 LODIndexOverride;
-		bool bNeedsReset;
-		bool bNeedsTeleport;
+		bool bUseLODIndexOverride = false;
+		int32 LODIndexOverride = 0;
+		bool bNeedsReset = false;
+		bool bNeedsTeleport = false;
 
 		// Reference space transform
 		FRigidTransform3 ReferenceSpaceTransform;  // TODO: Add override in the style of LODIndexOverride
@@ -218,7 +218,7 @@ PRAGMA_ENABLE_DEPRECATION_WARNINGS
 		TMap<FClothingSimulationSolver*, int32> LODIndices;
 
 		// Stats
-		int32 NumActiveKinematicParticles;
-		int32 NumActiveDynamicParticles;
+		int32 NumActiveKinematicParticles = 0;
+		int32 NumActiveDynamicParticles = 0;
 	};
 } // namespace Chaos
diff --git a/Engine/Plugins/Experimental/ChaosCloth/Source/Chaos/Public/ChaosCloth/ChaosClothingSimulationSolver.h b/Engine/Plugins/Experimental/ChaosCloth/Source/Chaos/Public/ChaosCloth/ChaosClothingSimulationSolver.h
--- a/Engine/Plugins/Experimental/ChaosCloth/Source/Chaos/Public/ChaosCloth/ChaosClothingSimulationSolver.h
+++ b/Engine/Plugins/Experimental/ChaosCloth/Source/Chaos/Public/ChaosCloth/ChaosClothingSimulationSolver.h
@@ -22,6 +22,11 @@ namespace Chaos
 	class FClothingSimulationMesh;
 	class FClothingSimulationConfig;
 
+	namespace Softs
+	{
+		class FCollectionPropertyConstFacade;
+	}
+
 	// Solver simulation node
 	class CHAOSCLOTH_API FClothingSimulationSolver final : public FPhysicsSolverEvents
 	{
@@ -138,34 +143,35 @@ namespace Chaos
 		void SetWindVelocity(uint32 GroupId, const TVec3<FRealSingle>& InWindVelocity);
 
 		// Set the geometry affected by the wind and pressure
-		void SetWindAndPressureGeometry(uint32 GroupId, const FTriangleMesh& TriangleMesh, const TConstArrayView<FRealSingle>& DragMultipliers, const TConstArrayView<FRealSingle>& LiftMultipliers, const TConstArrayView<FRealSingle>& PressureMultipliers);
-		UE_DEPRECATED(5.1, "Chaos::Softs::FVelocityField has been renamed FVelocityAndPressureField to match its new behavior.")
-		void SetWindGeometry(uint32 GroupId, const FTriangleMesh& TriangleMesh, const TConstArrayView<FRealSingle>& DragMultipliers, const TConstArrayView<FRealSingle>& LiftMultipliers)
-		{
-			SetWindAndPressureGeometry(GroupId, TriangleMesh, DragMultipliers, LiftMultipliers, TConstArrayView<FRealSingle>());
-		}
-
-		// Set new weight maps for the wind and pressure multipliers
-		void SetWindAndPressureMultipliers(
+		void SetWindAndPressureGeometry(
+			uint32 GroupId,
+			const FTriangleMesh& TriangleMesh,
+			const Softs::FCollectionPropertyConstFacade& PropertyCollection,
+			const TMap<FString, TConstArrayView<FRealSingle>>& WeightMaps);
+
+		void SetWindAndPressureGeometry(
 			uint32 GroupId,
+			const FTriangleMesh& TriangleMesh,
 			const TConstArrayView<FRealSingle>& DragMultipliers,
 			const TConstArrayView<FRealSingle>& LiftMultipliers,
 			const TConstArrayView<FRealSingle>& PressureMultipliers);
 
 		// Set the wind and pressure properties.
-		void SetWindAndPressureProperties(uint32 GroupId, const TVec2<FRealSingle>& Drag, const TVec2<FRealSingle>& Lift, FRealSingle AirDensity = 1.225e-6f, const TVec2<FRealSingle>& Pressure = TVec2<FRealSingle>::ZeroVector);
-		UE_DEPRECATED(5.1, "Chaos::Softs::FVelocityField has been renamed FVelocityAndPressureField to match its new behavior.")
-		void SetWindProperties(uint32 GroupId, const TVec2<FRealSingle>& Drag, const TVec2<FRealSingle>& Lift, FRealSingle AirDensity = 1.225e-6f)
-		{
-			SetWindAndPressureProperties(GroupId, Drag, Lift, AirDensity);
-		}
+		void SetWindAndPressureProperties(
+			uint32 GroupId,
+			const Softs::FCollectionPropertyConstFacade& PropertyCollection,
+			const TMap<FString, TConstArrayView<FRealSingle>>& WeightMaps,
+			bool bEnableAerodynamics);
+
+		void SetWindAndPressureProperties(
+			uint32 GroupId,
+			const TVec2<FRealSingle>& Drag,
+			const TVec2<FRealSingle>& Lift,
+			FRealSingle FluidDensity = 1.225f,
+			const TVec2<FRealSingle>& Pressure = TVec2<FRealSingle>::ZeroVector);
 
 		// Return the wind velocity and pressure field associated with a given group id.
 		const Softs::FVelocityAndPressureField& GetWindVelocityAndPressureField(uint32 GroupId) const;
-		UE_DEPRECATED(5.1, "Chaos::Softs::FVelocityField has been renamed FVelocityAndPressureField to match its new behavior.")
-		PRAGMA_DISABLE_DEPRECATION_WARNINGS
-		const Softs::FVelocityField& GetWindVelocityField(uint32 GroupId) { return GetWindVelocityAndPressureField(GroupId); }
-		PRAGMA_ENABLE_DEPRECATION_WARNINGS
 
 		// Add external forces to the particles
 		void AddExternalForces(uint32 GroupId, bool bUseLegacyWind);
diff --git a/Engine/Plugins/Experimental/ChaosClothAsset/Source/ChaosClothAssetEngine/Private/ChaosClothAsset/ClothComponent.cpp b/Engine/Plugins/Experimental/ChaosClothAsset/Source/ChaosClothAssetEngine/Private/ChaosClothAsset/ClothComponent.cpp
--- a/Engine/Plugins/Experimental/ChaosClothAsset/Source/ChaosClothAssetEngine/Private/ChaosClothAsset/ClothComponent.cpp
+++ b/Engine/Plugins/Experimental/ChaosClothAsset/Source/ChaosClothAssetEngine/Private/ChaosClothAsset/ClothComponent.cpp
@@ -343,7 +343,12 @@ void UChaosClothComponent::StartNewParallelSimulation(float DeltaTime)
 	if (ClothSimulationProxy.IsValid())
 	{
 		CSV_SCOPED_TIMING_STAT(Animation, Cloth);
-		ClothSimulationProxy->Tick_GameThread(DeltaTime);
+		const bool bIsSimulating = ClothSimulationProxy->Tick_GameThread(DeltaTime);
+
+		if (bIsSimulating)
+		{
+			CollectionPropertyFacade->ClearDirtyFlags();
+		}
 	}
 }
 
diff --git a/Engine/Plugins/Experimental/ChaosClothAsset/Source/ChaosClothAssetEngine/Private/ChaosClothAsset/ClothSimulationProxy.cpp b/Engine/Plugins/Experimental/ChaosClothAsset/Source/ChaosClothAssetEngine/Private/ChaosClothAsset/ClothSimulationProxy.cpp
--- a/Engine/Plugins/Experimental/ChaosClothAsset/Source/ChaosClothAssetEngine/Private/ChaosClothAsset/ClothSimulationProxy.cpp
+++ b/Engine/Plugins/Experimental/ChaosClothAsset/Source/ChaosClothAssetEngine/Private/ChaosClothAsset/ClothSimulationProxy.cpp
@@ -162,7 +162,7 @@ PRAGMA_ENABLE_DEPRECATION_WARNINGS
 		CompleteParallelSimulation_GameThread();
 	}
 
-	void FClothSimulationProxy::Tick_GameThread(float DeltaTime, FClothingSimulationCacheData* CacheData)
+	bool FClothSimulationProxy::Tick_GameThread(float DeltaTime, FClothingSimulationCacheData* CacheData)
 	{
 		SCOPE_CYCLE_COUNTER(STAT_ClothSimulationProxy_TickGame);
 
@@ -182,12 +182,14 @@ PRAGMA_ENABLE_DEPRECATION_WARNINGS
 
 			// Start the the cloth simulation thread
 			ParallelTask = TGraphTask<FClothSimulationProxyParallelTask>::CreateTask(nullptr, ENamedThreads::GameThread).ConstructAndDispatchWhenReady(*this);
+
+			return true;  // Simulating
 		}
-		else
-		{
-			// It still needs to write back to the GT cache as the context has changed
-			WriteSimulationData();
-		}
+
+		// It still needs to write back to the GT cache as the context has changed
+		WriteSimulationData();
+
+		return false;  // Not simulating
 	}
 
 	void FClothSimulationProxy::Tick()
diff --git a/Engine/Plugins/Experimental/ChaosClothAsset/Source/ChaosClothAssetEngine/Public/ChaosClothAsset/ClothSimulationProxy.h b/Engine/Plugins/Experimental/ChaosClothAsset/Source/ChaosClothAssetEngine/Public/ChaosClothAsset/ClothSimulationProxy.h
--- a/Engine/Plugins/Experimental/ChaosClothAsset/Source/ChaosClothAssetEngine/Public/ChaosClothAsset/ClothSimulationProxy.h
+++ b/Engine/Plugins/Experimental/ChaosClothAsset/Source/ChaosClothAssetEngine/Public/ChaosClothAsset/ClothSimulationProxy.h
@@ -44,7 +44,8 @@ namespace UE::Chaos::ClothAsset
 		FClothSimulationProxy& operator=(const FClothSimulationProxy&) = delete;
 		FClothSimulationProxy& operator=(FClothSimulationProxy&&) = delete;
 
-		void Tick_GameThread(float DeltaTime, FClothingSimulationCacheData* CacheData = nullptr);
+		/** Start the simulation if it isn't paused or suspended and return true, or simply update the existing simulation data and return false otherwise. */
+		bool Tick_GameThread(float DeltaTime, FClothingSimulationCacheData* CacheData = nullptr);
 
 		/** Wait for the parallel task to complete if one was running, and update the simulation data. */
 		void CompleteParallelSimulation_GameThread();
diff --git a/Engine/Source/Runtime/Experimental/Chaos/Private/Chaos/VelocityField.cpp b/Engine/Source/Runtime/Experimental/Chaos/Private/Chaos/VelocityField.cpp
--- a/Engine/Source/Runtime/Experimental/Chaos/Private/Chaos/VelocityField.cpp
+++ b/Engine/Source/Runtime/Experimental/Chaos/Private/Chaos/VelocityField.cpp
@@ -19,22 +19,21 @@ namespace Chaos::Softs {
 void FVelocityAndPressureField::SetProperties(
 	const FCollectionPropertyConstFacade& PropertyCollection,
 	const TMap<FString, TConstArrayView<FRealSingle>>& Weightmaps,
-	FSolverReal WorldScale)
+	FSolverReal WorldScale,
+	bool bEnableAerodynamics)
 {
-	if (IsFluidDensityMutable(PropertyCollection))
-	{
-		constexpr FSolverReal OneQuarter = (FSolverReal)0.25;
-		QuarterRho = (FSolverReal)GetFluidDensity(PropertyCollection) * OneQuarter;
-	}
-
-	bool bCheckForWeightMapUpdate = false;
+	bool bSetMultipliers = false;
 
 	if (IsDragMutable(PropertyCollection))
 	{
 		const FSolverVec2 Drag(GetWeightedFloatDrag(PropertyCollection));
 		DragBase = FMath::Clamp(Drag[0], MinCoefficient, MaxCoefficient);
 		DragRange = FMath::Clamp(Drag[1], MinCoefficient, MaxCoefficient) - DragBase;
-		bCheckForWeightMapUpdate = true;
+
+		if (IsDragStringDirty(PropertyCollection))
+		{
+			bSetMultipliers = true;
+		}
 	}
 
 	if (IsLiftMutable(PropertyCollection))
@@ -42,7 +41,16 @@ void FVelocityAndPressureField::SetProperties(
 		const FSolverVec2 Lift(GetWeightedFloatLift(PropertyCollection));
 		LiftBase = FMath::Clamp(Lift[0], MinCoefficient, MaxCoefficient);
 		LiftRange = FMath::Clamp(Lift[1], MinCoefficient, MaxCoefficient) - LiftBase;
-		bCheckForWeightMapUpdate = true;
+
+		if (IsLiftStringDirty(PropertyCollection))
+		{
+			bSetMultipliers = true;
+		}
+	}
+
+	if (IsFluidDensityMutable(PropertyCollection))
+	{
+		Rho = (FSolverReal)FMath::Max(GetFluidDensity(PropertyCollection), 0.f) / FMath::Cube(WorldScale);
 	}
 
 	if (IsPressureMutable(PropertyCollection))
@@ -50,33 +58,77 @@ void FVelocityAndPressureField::SetProperties(
 		const FSolverVec2 Pressure(GetWeightedFloatPressure(PropertyCollection));
 		PressureBase = Pressure[0] / WorldScale;
 		PressureRange = Pressure[1] / WorldScale - PressureBase;
-		bCheckForWeightMapUpdate = true;
+
+		if (IsPressureStringDirty(PropertyCollection))
+		{
+			bSetMultipliers = true;
+		}
 	}
 
-	if (bCheckForWeightMapUpdate && (
-		IsDragStringDirty(PropertyCollection) ||
-		IsLiftStringDirty(PropertyCollection) ||
-		IsPressureStringDirty(PropertyCollection)))
+	if (bSetMultipliers)
 	{
-		const TConstArrayView<FRealSingle> DragMultipliers = Weightmaps.FindRef(GetDragString(PropertyCollection));
-		const TConstArrayView<FRealSingle> LiftMultipliers = Weightmaps.FindRef(GetLiftString(PropertyCollection));
-		const TConstArrayView<FRealSingle> PressureMultipliers = Weightmaps.FindRef(GetPressureString(PropertyCollection));
-
-		SetMultipliers(DragMultipliers, LiftMultipliers, PressureMultipliers);
+		SetMultipliers(PropertyCollection, Weightmaps);
 	}
+
+	// Update QuarterRho
+	constexpr FSolverReal OneQuarter = (FSolverReal)0.25;
+	QuarterRho = bEnableAerodynamics ? Rho * OneQuarter : (FSolverReal)0.;
 }
 
-void FVelocityAndPressureField::SetGeometry(const FTriangleMesh* TriangleMesh, const TConstArrayView<FRealSingle>& DragMultipliers, const TConstArrayView<FRealSingle>& LiftMultipliers, const TConstArrayView<FRealSingle>& PressureMultipliers)
+void FVelocityAndPressureField::SetProperties(
+	const FSolverVec2& Drag,
+	const FSolverVec2& Lift,
+	const FSolverReal FluidDensity,
+	const FSolverVec2& Pressure,
+	FSolverReal WorldScale)
+{
+	DragBase = FMath::Clamp(Drag[0], MinCoefficient, MaxCoefficient);
+	DragRange = FMath::Clamp(Drag[1], MinCoefficient, MaxCoefficient) - DragBase;
+	LiftBase = FMath::Clamp(Lift[0], MinCoefficient, MaxCoefficient);
+	LiftRange = FMath::Clamp(Lift[1], MinCoefficient, MaxCoefficient) - LiftBase;
+	Rho = FMath::Max(FluidDensity / FMath::Cube(WorldScale), (FSolverReal)0.);
+	PressureBase = Pressure[0] / WorldScale;
+	PressureRange = Pressure[1] / WorldScale - PressureBase;
+
+	constexpr FSolverReal OneQuarter = (FSolverReal)0.25;
+	QuarterRho = Rho * OneQuarter;
+}
+
+void FVelocityAndPressureField::SetGeometry(
+	const FTriangleMesh* TriangleMesh,
+	const FCollectionPropertyConstFacade& PropertyCollection,
+	const TMap<FString, TConstArrayView<FRealSingle>>& Weightmaps,
+	FSolverReal WorldScale)
+{
+	SetGeometry(TriangleMesh);
+	SetProperties(
+		FSolverVec2(GetWeightedFloatDrag(PropertyCollection, 0.f)),  // If these properties don't exist, set their values to 0, not to DefaultCoefficients!
+		FSolverVec2(GetWeightedFloatLift(PropertyCollection, 0.f)),
+		(FSolverReal)GetFluidDensity(PropertyCollection, 0.f),
+		FSolverVec2(GetWeightedFloatPressure(PropertyCollection, 0.f)),  // These getters also initialize the property indices, so keep before SetMultipliers
+		WorldScale);
+	SetMultipliers(PropertyCollection, Weightmaps);
+}
+
+void FVelocityAndPressureField::SetGeometry(
+	const FTriangleMesh* TriangleMesh,
+	const TConstArrayView<FRealSingle>& DragMultipliers,
+	const TConstArrayView<FRealSingle>& LiftMultipliers,
+	const TConstArrayView<FRealSingle>& PressureMultipliers)
+{
+	SetGeometry(TriangleMesh);
+	SetMultipliers(DragMultipliers, LiftMultipliers, PressureMultipliers);
+}
+
+void FVelocityAndPressureField::SetGeometry(const FTriangleMesh* TriangleMesh)
 {
 	if (TriangleMesh)
 	{
 		PointToTriangleMap = TriangleMesh->GetPointToTriangleMap();
 		Elements = TriangleMesh->GetElements();
-				
 		const TVec2<int32> Range = TriangleMesh->GetVertexRange();
 		Offset = Range[0];
 		NumParticles = 1 + Range[1] - Offset;
-
 		Forces.SetNumUninitialized(Elements.Num());
 	}
 	else
@@ -86,8 +138,21 @@ void FVelocityAndPressureField::SetGeometry(const FTriangleMesh* TriangleMesh, c
 		Offset = 0;
 		NumParticles = 0;
 		Forces.Reset();
-		SetProperties(FSolverVec2(0.), FSolverVec2(0.), (FSolverReal)0., FSolverVec2(0.));
 	}
+}
+
+void FVelocityAndPressureField::SetMultipliers(
+	const FCollectionPropertyConstFacade& PropertyCollection,
+	const TMap<FString, TConstArrayView<FRealSingle>>& Weightmaps)
+{
+	const TConstArrayView<FRealSingle> DragMultipliers = (DragIndex != INDEX_NONE) ?
+		Weightmaps.FindRef(GetDragString(PropertyCollection)) : TConstArrayView<FRealSingle>();
+
+	const TConstArrayView<FRealSingle> LiftMultipliers = (LiftIndex != INDEX_NONE) ?
+		Weightmaps.FindRef(GetLiftString(PropertyCollection)) : TConstArrayView<FRealSingle>();
+
+	const TConstArrayView<FRealSingle> PressureMultipliers = (PressureIndex != INDEX_NONE) ?
+		Weightmaps.FindRef(GetPressureString(PropertyCollection)) : TConstArrayView<FRealSingle>();
 
 	SetMultipliers(DragMultipliers, LiftMultipliers, PressureMultipliers);
 }
diff --git a/Engine/Source/Runtime/Experimental/Chaos/Public/Chaos/VelocityField.h b/Engine/Source/Runtime/Experimental/Chaos/Public/Chaos/VelocityField.h
--- a/Engine/Source/Runtime/Experimental/Chaos/Public/Chaos/VelocityField.h
+++ b/Engine/Source/Runtime/Experimental/Chaos/Public/Chaos/VelocityField.h
@@ -15,7 +15,7 @@ class CHAOS_API FVelocityAndPressureField final
 public:
 	static constexpr FSolverReal DefaultDragCoefficient = (FSolverReal)0.5;
 	static constexpr FSolverReal DefaultLiftCoefficient = (FSolverReal)0.1;
-	static constexpr FSolverReal DefaultFluidDensity = (FSolverReal)1.225e-6;
+	static constexpr FSolverReal DefaultFluidDensity = (FSolverReal)1.225;
 	static constexpr FSolverReal MinCoefficient = (FSolverReal)0.;   // Applies to both drag and lift
 	static constexpr FSolverReal MaxCoefficient = (FSolverReal)10.;  //
 
@@ -32,10 +32,10 @@ public:
 		, NumParticles(0)
 	{
 		SetProperties(
-			FSolverVec2(GetWeightedFloatDrag(PropertyCollection, DefaultDragCoefficient)),
-			FSolverVec2(GetWeightedFloatLift(PropertyCollection, DefaultLiftCoefficient)),
-			FSolverReal(GetFluidDensity(PropertyCollection, DefaultFluidDensity)),
-			FSolverVec2(GetWeightedFloatPressure(PropertyCollection, 0.f)));
+			FSolverVec2(GetWeightedFloatDrag(PropertyCollection, (FSolverReal)0.)),  // If these properties don't exist, set their values to 0, not to DefaultCoefficients!
+			FSolverVec2(GetWeightedFloatLift(PropertyCollection, (FSolverReal)0.)),
+			FSolverReal(GetFluidDensity(PropertyCollection, (FSolverReal)0.)),
+			FSolverVec2(GetWeightedFloatPressure(PropertyCollection, (FSolverReal)0.)));
 	}
 
 	// Construct an uninitialized field. Mesh, properties, and velocity will have to be set for this field to be valid.
@@ -43,7 +43,7 @@ public:
 		: Offset(INDEX_NONE)
 		, NumParticles(0)
 	{
-		SetProperties(FSolverVec2(0.), FSolverVec2(0.), (FSolverReal)0.);
+		SetProperties(FSolverVec2(0.), FSolverVec2(0.), (FSolverReal)0., FSolverVec2(0.));
 	}
 
 	~FVelocityAndPressureField() {}
@@ -64,38 +64,42 @@ public:
 	void SetProperties(
 		const FCollectionPropertyConstFacade& PropertyCollection,
 		const TMap<FString, TConstArrayView<FRealSingle>>& Weightmaps,
-		FSolverReal WorldScale);
+		FSolverReal WorldScale,
+		bool bEnableAerodynamics);
 
-	UE_DEPRECATED(5.3, "Use SetProperties(const FCollectionPropertyConstFacade&, const TMap<FString, TConstArrayView<FRealSingle>>&, FSolverReal) instead.")
+	UE_DEPRECATED(5.3, "Use SetProperties(const FCollectionPropertyConstFacade&, const TMap<FString, TConstArrayView<FRealSingle>>&, FSolverReal, bool) instead.")
 	void SetProperties(const FCollectionPropertyConstFacade& PropertyCollection, FSolverReal WorldScale)
 	{
-		SetProperties(PropertyCollection, TMap<FString, TConstArrayView<FRealSingle>>(), WorldScale);
+		constexpr bool bEnableAerodynamics = true;
+		SetProperties(PropertyCollection, TMap<FString, TConstArrayView<FRealSingle>>(), WorldScale, bEnableAerodynamics);
 	}
 
-	void SetProperties(const FSolverVec2& Drag, const FSolverVec2& Lift, const FSolverReal FluidDensity, const FSolverVec2& Pressure = FSolverVec2::ZeroVector)
-	{
-		constexpr FSolverReal OneQuarter = (FSolverReal)0.25;
-		QuarterRho = FluidDensity * OneQuarter;
-
-		DragBase = FMath::Clamp(Drag[0], MinCoefficient, MaxCoefficient);
-		DragRange = FMath::Clamp(Drag[1], MinCoefficient, MaxCoefficient) - DragBase;
-		LiftBase = FMath::Clamp(Lift[0], MinCoefficient, MaxCoefficient);
-		LiftRange = FMath::Clamp(Lift[1], MinCoefficient, MaxCoefficient) - LiftBase;
-		PressureBase = Pressure[0];
-		PressureRange = Pressure[1] - PressureBase;
-	}
+	void SetProperties(
+		const FSolverVec2& Drag,
+		const FSolverVec2& Lift,
+		const FSolverReal FluidDensity,
+		const FSolverVec2& Pressure = FSolverVec2::ZeroVector,
+		FSolverReal WorldScale = 1.f);
 
 	bool IsActive() const 
 	{ 
-		// Note: range can be a negative value (although not when Lift or Drag base is zero)
-		return (DragBase > (FSolverReal)0. || DragRange != (FSolverReal)0.) || (LiftBase > (FSolverReal)0. || LiftRange != (FSolverReal)0.) || PressureBase != (FSolverReal)0. || PressureRange != (FSolverReal)0.; 
+		return PressureBase != (FSolverReal)0. || PressureRange != (FSolverReal)0. ||
+			(AreAerodynamicsEnabled() && (
+				DragBase > (FSolverReal)0. || DragRange != (FSolverReal)0. ||  // Note: range can be a negative value (although not when Lift or Drag base is zero)
+				LiftBase > (FSolverReal)0. || LiftRange != (FSolverReal)0.));
 	}
 
-	void SetGeometry(const FTriangleMesh* TriangleMesh, const TConstArrayView<FRealSingle>& DragMultipliers, const TConstArrayView<FRealSingle>& LiftMultipliers, const TConstArrayView<FRealSingle>& PressureMultipliers = TConstArrayView<FRealSingle>());
-	void SetMultipliers(
+	void SetGeometry(
+		const FTriangleMesh* TriangleMesh,
+		const FCollectionPropertyConstFacade& PropertyCollection,
+		const TMap<FString, TConstArrayView<FRealSingle>>& Weightmaps,
+		FSolverReal WorldScale);
+
+	void SetGeometry(
+		const FTriangleMesh* TriangleMesh,
 		const TConstArrayView<FRealSingle>& DragMultipliers,
 		const TConstArrayView<FRealSingle>& LiftMultipliers,
-		const TConstArrayView<FRealSingle>& PressureMultipliers);
+		const TConstArrayView<FRealSingle>& PressureMultipliers = TConstArrayView<FRealSingle>());
 
 	void SetVelocity(const FSolverVec3& InVelocity) { Velocity = InVelocity; }
 
@@ -103,6 +107,17 @@ public:
 	TConstArrayView<FSolverVec3> GetForces() const { return TConstArrayView<FSolverVec3>(Forces); }
 
 private:
+	bool AreAerodynamicsEnabled() const { return QuarterRho > (FSolverReal)0.; }
+
+	void SetGeometry(const FTriangleMesh* TriangleMesh);
+
+	void SetMultipliers(const FCollectionPropertyConstFacade& PropertyCollection,const TMap<FString, TConstArrayView<FRealSingle>>& Weightmaps);
+
+	void SetMultipliers(
+		const TConstArrayView<FRealSingle>& DragMultipliers,
+		const TConstArrayView<FRealSingle>& LiftMultipliers,
+		const TConstArrayView<FRealSingle>& PressureMultipliers);
+
 	void UpdateField(const FSolverParticles& InParticles, int32 ElementIndex, const FSolverVec3& InVelocity, const FSolverReal Cd, const FSolverReal Cl, const FSolverReal Cp)
 	{
 		const TVec3<int32>& Element = Elements[ElementIndex];
@@ -141,6 +156,7 @@ private:
 	FSolverReal LiftRange;
 	FSolverReal PressureBase; 
 	FSolverReal PressureRange;
+	FSolverReal Rho;
 	FSolverReal QuarterRho;
 	int32 Offset;
 	int32 NumParticles;
