Chaos Cloth Visualization of new XPBD constraints. Lots of hard-coded stuff and cvar controls. Will want to factor some of this out to be set by new visualization controls.

diff --git a/Engine/Plugins/Experimental/ChaosCloth/Source/Chaos/Private/ChaosCloth/ChaosClothVisualization.cpp b/Engine/Plugins/Experimental/ChaosCloth/Source/Chaos/Private/ChaosCloth/ChaosClothVisualization.cpp
--- a/Engine/Plugins/Experimental/ChaosCloth/Source/Chaos/Private/ChaosCloth/ChaosClothVisualization.cpp
+++ b/Engine/Plugins/Experimental/ChaosCloth/Source/Chaos/Private/ChaosCloth/ChaosClothVisualization.cpp
@@ -24,6 +24,8 @@
 #include "Chaos/PBDTriangleMeshCollisions.h"
 #include "Chaos/XPBDBendingConstraints.h"
 #include "Chaos/XPBDSpringConstraints.h"
+#include "Chaos/XPBDAnisotropicBendingConstraints.h"
+#include "Chaos/XPBDStretchBiasElementConstraints.h"
 #include "Chaos/WeightedLatticeImplicitObject.h"
 #include "DynamicMeshBuilder.h"
 #include "Engine/EngineTypes.h"
@@ -37,17 +39,49 @@
 #include "UObject/ICookInfo.h"
 #endif  // #if WITH_EDITOR
 
+namespace Chaos
+{
+
+namespace Private
+{
 static int DrawSkinnedLattice = 0;
 static FAutoConsoleVariableRef CVarClothVizDrawSkinnedLattice(TEXT("p.ChaosClothVisualization.DrawSkinnedLattice"), DrawSkinnedLattice, TEXT("Draw skinned lattice, 0 = none, 1 = filled, 2 = empty, 3 = both"));
 
-namespace Chaos
+// TODO: move these options to be somewhere the new cloth editor visualization can use.
+enum class EBendingDrawMode : int
+{
+	BuckleStatus = 0,
+	ParallelGraphColor = 1,
+	Anisotropy = 2
+};
+static int32 BendingDrawMode = (int32)EBendingDrawMode::BuckleStatus;
+static FAutoConsoleVariableRef CVarClothVizBendDrawMode(TEXT("p.ChaosClothVisualization.BendingDrawMode"), BendingDrawMode, TEXT("Bending draw mode, 0 = BuckleStatus, 1 = Parallel graph color, 2 = Anisotropy"));
+
+enum class EStretchBiasDrawMode : int
 {
+	ParallelGraphColor = 0,
+	WarpStretch = 1,
+	WeftStretch = 2,
+	BiasStretch = 3
+};
+static int32 StretchBiasDrawMode = (int32)EStretchBiasDrawMode::ParallelGraphColor;
+static FAutoConsoleVariableRef CVarClothVizStretchBiasDrawMode(TEXT("p.ChaosClothVisualization.StretchBiasDrawMode"), StretchBiasDrawMode, TEXT("Stretch draw mode, 0 = Parallel graph color, 1 = Warp Stretch, 2 = Weft Stretch, 3 = BiasStretch"));
+static float StretchBiasDrawRangeMin = -1.f;
+static float StretchBiasDrawRangeMax = 1.f;
+static FAutoConsoleVariableRef CVarClothVizStretchBiasDrawRangeMin(TEXT("p.ChaosClothVisualization.StretchBiasDrawRangeMin"), StretchBiasDrawRangeMin, TEXT("Min stretch in draw color range. Negative = compressed, 0 = undeformed, positive = stretched. (When drawing warp/weft stretch)"));
+static FAutoConsoleVariableRef CVarClothVizStretchBiasDrawRangeMax(TEXT("p.ChaosClothVisualization.StretchBiasDrawRangeMax"), StretchBiasDrawRangeMax, TEXT("Max stretch in draw color range. Negative = compressed, 0 = undeformed, positive = stretched. (When drawing warp/weft stretch)"));
+static bool bStretchBiasDrawOutOfRange = true;
+static FAutoConsoleVariableRef CVarClothVizStretchBiasDrawOutOfRange(TEXT("p.ChaosClothVisualization.StretchBiasDrawOutOfRange"), bStretchBiasDrawOutOfRange, TEXT("Draw out of range elements (When drawing warp/weft stretch)"));
+
+}// namespace Private
+
 	FClothVisualization::FClothVisualization(const ::Chaos::FClothingSimulationSolver* InSolver)
 		: Solver(InSolver)
 	{
 #if WITH_EDITOR
 		FCookLoadScope CookLoadScope(ECookLoadType::EditorOnly);
 		ClothMaterial = LoadObject<UMaterial>(nullptr, TEXT("/Engine/EditorMaterials/Cloth/CameraLitDoubleSided.CameraLitDoubleSided"), nullptr, LOAD_None, nullptr);  // LOAD_EditorOnly
+		ClothMaterialColor = LoadObject<UMaterial>(nullptr, TEXT("/Engine/EditorMaterials/Cloth/CameraLitVertexColor.CameraLitVertexColor"), nullptr, LOAD_None, nullptr);  // LOAD_EditorOnly
 		ClothMaterialVertex = LoadObject<UMaterial>(nullptr, TEXT("/Engine/EditorMaterials/WidgetVertexColorMaterial"), nullptr, LOAD_None, nullptr);  // LOAD_EditorOnly
 		CollisionMaterial = LoadObject<UMaterial>(nullptr, TEXT("/Engine/EditorMaterials/PhAT_UnselectedMaterial"), nullptr, LOAD_None, nullptr);
 #endif  // #if WITH_EDITOR
@@ -64,6 +98,7 @@ namespace Chaos
 	void FClothVisualization::AddReferencedObjects(FReferenceCollector& Collector)
 	{
 		Collector.AddReferencedObject(ClothMaterial);
+		Collector.AddReferencedObject(ClothMaterialColor);
 		Collector.AddReferencedObject(ClothMaterialVertex);
 		Collector.AddReferencedObject(CollisionMaterial);
 	}
@@ -492,7 +527,7 @@ namespace Chaos
 
 			MeshBuilder.Draw(PDI, FTransform(Rotation, Position).ToMatrixWithScale(), MaterialRenderProxy, SDPG_World, false, false);
 
-			if (DrawSkinnedLattice)
+			if (Private::DrawSkinnedLattice)
 			{
 				const Chaos::TUniformGrid<double, 3>& LatticeGrid = SkinnedLevelSet.GetGrid();
 				const Chaos::TArrayND<Chaos::FVec3, 3>& DeformedPoints = SkinnedLevelSet.GetDeformedPoints();
@@ -510,7 +545,7 @@ namespace Chaos
 						{
 							const bool bIsEmpty = EmptyCells(I, J, K);
 							const int32 EmptyDrawMask = bIsEmpty ? 2 : 1;
-							if (EmptyDrawMask & DrawSkinnedLattice)
+							if (EmptyDrawMask & Private::DrawSkinnedLattice)
 							{
 								const FVector P000 = LocalToWorld.TransformPosition(FVector(DeformedPoints(I, J, K)));
 								const FVector P001 = LocalToWorld.TransformPosition(FVector(DeformedPoints(I, J, K + 1)));
@@ -1239,11 +1274,12 @@ namespace Chaos
 		}
 	}
 
-	static void DrawSpringConstraintColors(FPrimitiveDrawInterface* PDI, const TConstArrayView<::Chaos::Softs::FSolverVec3>& Positions, const ::Chaos::FVec3& LocalSpaceLocation, const ::Chaos::Softs::FPBDSpringConstraints* const SpringConstraints)
+	template<typename SpringConstraintType>
+	static void DrawSpringConstraintColors(FPrimitiveDrawInterface* PDI, const TConstArrayView<::Chaos::Softs::FSolverVec3>& Positions, const ::Chaos::FVec3& LocalSpaceLocation, const SpringConstraintType* const SpringConstraints)
 	{
 		check(SpringConstraints);
 
-		const TArray<TVec2<int32>>& Constraints = SpringConstraints->GetConstraints();
+		const auto& Constraints = SpringConstraints->GetConstraints(); // auto because constraints can be TVec2<int> or TVec4<int>, but we just care about first two for drawing here either way
 		const TArray<int32>& ConstraintsPerColorStartIndex = SpringConstraints->GetConstraintsPerColorStartIndex();
 		if (ConstraintsPerColorStartIndex.Num() > 1)
 		{
@@ -1263,7 +1299,7 @@ namespace Chaos
 				for (int32 ConstraintIndex = ColorStart; ConstraintIndex < ColorEnd; ++ConstraintIndex)
 				{
 					// Draw line
-					const TVec2<int32>& Constraint = Constraints[ConstraintIndex];
+					const auto& Constraint = Constraints[ConstraintIndex];
 					const FVec3 Pos0 = FVec3(Positions[Constraint[0]]) + LocalSpaceLocation;
 					const FVec3 Pos1 = FVec3(Positions[Constraint[1]]) + LocalSpaceLocation;
 					DrawLine(PDI, Pos0, Pos1, DrawColor);
@@ -1272,7 +1308,7 @@ namespace Chaos
 		}
 		else
 		{
-			for (const TVec2<int32>& Constraint : Constraints)
+			for (const auto& Constraint : Constraints)
 			{
 				// Draw line
 				const FVec3 Pos0 = FVec3(Positions[Constraint[0]]) + LocalSpaceLocation;
@@ -1283,6 +1319,193 @@ namespace Chaos
 		}
 	}
 
+	static void DrawStretchBiasConstraints_ParallelGraphColor(FPrimitiveDrawInterface* PDI, const TConstArrayView<::Chaos::Softs::FSolverVec3>& Positions, const ::Chaos::FVec3& LocalSpaceLocation, const FMaterialRenderProxy* MaterialRenderProxy, const Softs::FXPBDStretchBiasElementConstraints* const SpringConstraints)
+	{
+#if WITH_EDITOR
+		if (PDI && MaterialRenderProxy)
+		{
+			const TArray<TVec3<int32>>& Constraints = SpringConstraints->GetConstraints();
+			const TArray<int32>& ConstraintsPerColorStartIndex = SpringConstraints->GetConstraintsPerColorStartIndex();
+
+			const int32 ConstraintColorNum = ConstraintsPerColorStartIndex.Num() - 1;
+			const uint8 HueOffset = 196 / ConstraintColorNum;
+
+			auto ConstraintColor =
+				[HueOffset](int32 ColorIndex)->FColor
+			{
+				return FLinearColor::MakeFromHSV8(ColorIndex * HueOffset, 255, 255).ToFColor(true);
+			};
+
+			FDynamicMeshBuilder MeshBuilder(PDI->View->GetFeatureLevel());
+			int32 VertexIndex = 0;
+			for (int32 ConstraintColorIndex = 0; ConstraintColorIndex < ConstraintColorNum; ++ConstraintColorIndex)
+			{
+				const int32 ColorStart = ConstraintsPerColorStartIndex[ConstraintColorIndex];
+				const int32 ColorEnd = ConstraintsPerColorStartIndex[ConstraintColorIndex + 1];
+				const FColor DrawColor = ConstraintColor(ConstraintColorIndex);
+				for (int32 ConstraintIndex = ColorStart; ConstraintIndex < ColorEnd; ++ConstraintIndex, VertexIndex += 3)
+				{
+					const TVec3<int32>& Constraint = Constraints[ConstraintIndex];
+					const FVector3f Pos0 = FVector3f(Positions[Constraint[0]]);
+					const FVector3f Pos1 = FVector3f(Positions[Constraint[1]]);
+					const FVector3f Pos2 = FVector3f(Positions[Constraint[2]]);
+					const FVector3f Normal = FVector3f::CrossProduct(Pos2 - Pos0, Pos1 - Pos0).GetSafeNormal();
+					const FVector3f Tangent = ((Pos1 + Pos2) * 0.5f - Pos0).GetSafeNormal();
+
+					MeshBuilder.AddVertex(FDynamicMeshVertex(Pos0, Tangent, Normal, FVector2f(0.f, 0.f), DrawColor));
+					MeshBuilder.AddVertex(FDynamicMeshVertex(Pos1, Tangent, Normal, FVector2f(0.f, 1.f), DrawColor));
+					MeshBuilder.AddVertex(FDynamicMeshVertex(Pos2, Tangent, Normal, FVector2f(1.f, 1.f), DrawColor));
+					MeshBuilder.AddTriangle(VertexIndex, VertexIndex + 1, VertexIndex + 2);
+				}
+			}
+
+			FMatrix LocalSimSpaceToWorld(FMatrix::Identity);
+			LocalSimSpaceToWorld.SetOrigin(LocalSpaceLocation);
+			MeshBuilder.Draw(PDI, LocalSimSpaceToWorld, MaterialRenderProxy, SDPG_World, false, false);
+		}
+#endif
+	}
+
+	static void DrawStretchBiasConstraints_WarpWeftStretch(FPrimitiveDrawInterface* PDI, const TConstArrayView<::Chaos::Softs::FSolverVec3>& Positions, const ::Chaos::FVec3& LocalSpaceLocation, const FMaterialRenderProxy* MaterialRenderProxy, const Softs::FXPBDStretchBiasElementConstraints* const SpringConstraints)
+	{
+#if WITH_EDITOR
+		if (PDI && MaterialRenderProxy)
+		{
+			const TArray<TVec3<int32>>& Constraints = SpringConstraints->GetConstraints();
+			const float StretchRangeMinClamped = Private::StretchBiasDrawRangeMin <= Private::StretchBiasDrawRangeMax ? Private::StretchBiasDrawRangeMin : 0.f;
+			const float StretchRangeMaxClamped = Private::StretchBiasDrawRangeMin <= Private::StretchBiasDrawRangeMax ? Private::StretchBiasDrawRangeMax : 0.f;
+			const float StretchRange = StretchRangeMaxClamped - StretchRangeMinClamped;
+			const float StretchRangeInv = StretchRange > UE_KINDA_SMALL_NUMBER ? 1.f / StretchRange : 0.f;
+			const bool bIsWeftStretch = (Private::EStretchBiasDrawMode)Private::StretchBiasDrawMode == Private::EStretchBiasDrawMode::WeftStretch;
+			auto ConstraintColor =
+				[bIsWeftStretch, StretchRangeMinClamped, StretchRangeMaxClamped, StretchRangeInv, &Constraints, &Positions, &SpringConstraints](int32 ConstraintIndex, bool& bOutOfRange)->FColor
+			{
+				bOutOfRange = false;
+
+				// TODO: make these configurable
+				constexpr float OutOfRangeMinHue = 240.f; // blue
+				constexpr float MinHue = 180.f; // cyan
+				constexpr float MaxHue = 60.f; // yellow
+				constexpr float OutOfRangeMaxHue = 0.f; // red
+				constexpr float StretchedValue = 1.f;
+				constexpr float CompressedValue = 0.5;
+
+				const TVec3<int32>& Constraint = Constraints[ConstraintIndex];
+				const Softs::FSolverVec3& P0 = Positions[Constraint[0]];
+				const Softs::FSolverVec3& P1 = Positions[Constraint[1]];
+				const Softs::FSolverVec3& P2 = Positions[Constraint[2]];
+				Softs::FSolverVec3 dX_dU, dX_dV;
+				SpringConstraints->CalculateUVStretch(ConstraintIndex, P0, P1, P2, dX_dU, dX_dV);
+				const Softs::FSolverReal Stretch = bIsWeftStretch ? dX_dV.Length() : dX_dU.Length();
+
+				const int32 WarpWeftIndex = bIsWeftStretch ? 1 : 0;
+
+				const Softs::FSolverVec2 StretchScale = SpringConstraints->GetWarpWeftScale(ConstraintIndex);
+				const Softs::FSolverReal RestStretch = SpringConstraints->GetRestStretchLengths()[ConstraintIndex][WarpWeftIndex] * StretchScale[WarpWeftIndex];
+				const Softs::FSolverReal RestStretchInv = (Softs::FSolverReal)1. / FMath::Max(RestStretch, UE_KINDA_SMALL_NUMBER);
+
+				const float StretchRatio = (Stretch - RestStretch) * RestStretchInv;
+				if (StretchRatio < StretchRangeMinClamped)
+				{
+					bOutOfRange = true;
+					return FLinearColor(OutOfRangeMinHue, 1.f, StretchRatio < 0 ? CompressedValue : StretchedValue).HSVToLinearRGB().ToFColor(true);
+				}
+				if (StretchRatio > StretchRangeMaxClamped)
+				{
+					bOutOfRange = true;
+					return FLinearColor(OutOfRangeMaxHue, 1.f, StretchRatio < 0 ? CompressedValue : StretchedValue).HSVToLinearRGB().ToFColor(true);
+				}
+
+				// Convert from [StretchRangeMinClamped, StretchRangeMaxClamped] --> [MinHue, MaxHue]
+				const float Hue = (MinHue + (MaxHue - MinHue) * (StretchRatio - StretchRangeMinClamped) * StretchRangeInv);
+				return FLinearColor(Hue, 1.f, StretchRatio < 0 ? CompressedValue : StretchedValue).HSVToLinearRGB().ToFColor(true);
+			};
+
+			FDynamicMeshBuilder MeshBuilder(PDI->View->GetFeatureLevel());
+			int32 VertexIndex = 0;
+			for (int32 ConstraintIndex = 0; ConstraintIndex < Constraints.Num(); ++ConstraintIndex)
+			{
+
+				bool bIsOutOfRange;
+				const FColor DrawColor = ConstraintColor(ConstraintIndex, bIsOutOfRange);
+				if (bIsOutOfRange && !Private::bStretchBiasDrawOutOfRange)
+				{
+					continue;
+				}
+
+				const TVec3<int32>& Constraint = Constraints[ConstraintIndex];
+				const FVector3f Pos0 = FVector3f(Positions[Constraint[0]]);
+				const FVector3f Pos1 = FVector3f(Positions[Constraint[1]]);
+				const FVector3f Pos2 = FVector3f(Positions[Constraint[2]]);
+				const FVector3f Normal = FVector3f::CrossProduct(Pos2 - Pos0, Pos1 - Pos0).GetSafeNormal();
+				const FVector3f Tangent = ((Pos1 + Pos2) * 0.5f - Pos0).GetSafeNormal();
+
+				MeshBuilder.AddVertex(FDynamicMeshVertex(Pos0, Tangent, Normal, FVector2f(0.f, 0.f), DrawColor));
+				MeshBuilder.AddVertex(FDynamicMeshVertex(Pos1, Tangent, Normal, FVector2f(0.f, 1.f), DrawColor));
+				MeshBuilder.AddVertex(FDynamicMeshVertex(Pos2, Tangent, Normal, FVector2f(1.f, 1.f), DrawColor));
+				MeshBuilder.AddTriangle(VertexIndex, VertexIndex + 1, VertexIndex + 2);
+				VertexIndex += 3;
+			}
+
+			FMatrix LocalSimSpaceToWorld(FMatrix::Identity);
+			LocalSimSpaceToWorld.SetOrigin(LocalSpaceLocation);
+			MeshBuilder.Draw(PDI, LocalSimSpaceToWorld, MaterialRenderProxy, SDPG_World, false, false);
+		}
+#endif
+	}
+
+	static void DrawStretchBiasConstraints_BiasStretch(FPrimitiveDrawInterface* PDI, const TConstArrayView<::Chaos::Softs::FSolverVec3>& Positions, const ::Chaos::FVec3& LocalSpaceLocation, const FMaterialRenderProxy* MaterialRenderProxy, const Softs::FXPBDStretchBiasElementConstraints* const SpringConstraints)
+	{
+
+#if WITH_EDITOR
+		if (PDI && MaterialRenderProxy)
+		{
+			const TArray<TVec3<int32>>& Constraints = SpringConstraints->GetConstraints();
+			auto ConstraintColor = [&Constraints, &Positions, &SpringConstraints](int32 ConstraintIndex)->FColor
+			{
+				const TVec3<int32>& Constraint = Constraints[ConstraintIndex];
+				const Softs::FSolverVec3& P0 = Positions[Constraint[0]];
+				const Softs::FSolverVec3& P1 = Positions[Constraint[1]];
+				const Softs::FSolverVec3& P2 = Positions[Constraint[2]];
+				Softs::FSolverVec3 dX_dU, dX_dV;
+				SpringConstraints->CalculateUVStretch(ConstraintIndex, P0, P1, P2, dX_dU, dX_dV);
+				
+				const Softs::FSolverVec3 dX_dU_normalized = dX_dU.GetSafeNormal();
+				const Softs::FSolverVec3 dX_dV_normalized = dX_dV.GetSafeNormal();
+				const Softs::FSolverReal Shear = FMath::Abs(Softs::FSolverVec3::DotProduct(dX_dU_normalized, dX_dV_normalized));
+
+				constexpr float UndeformedHue = 240.f; // blue
+				constexpr float MaxDeformedHue = 360.f; // red
+
+				return FLinearColor(UndeformedHue + Shear * (MaxDeformedHue - UndeformedHue), 1.f, 1.f).HSVToLinearRGB().ToFColor(true);
+			};
+
+			FDynamicMeshBuilder MeshBuilder(PDI->View->GetFeatureLevel());
+			int32 VertexIndex = 0;
+			for (int32 ConstraintIndex = 0; ConstraintIndex < Constraints.Num(); ++ConstraintIndex, VertexIndex += 3)
+			{
+				const FColor DrawColor = ConstraintColor(ConstraintIndex);
+
+				const TVec3<int32>& Constraint = Constraints[ConstraintIndex];
+				const FVector3f Pos0 = FVector3f(Positions[Constraint[0]]);
+				const FVector3f Pos1 = FVector3f(Positions[Constraint[1]]);
+				const FVector3f Pos2 = FVector3f(Positions[Constraint[2]]);
+				const FVector3f Normal = FVector3f::CrossProduct(Pos2 - Pos0, Pos1 - Pos0).GetSafeNormal();
+				const FVector3f Tangent = ((Pos1 + Pos2) * 0.5f - Pos0).GetSafeNormal();
+
+				MeshBuilder.AddVertex(FDynamicMeshVertex(Pos0, Tangent, Normal, FVector2f(0.f, 0.f), DrawColor));
+				MeshBuilder.AddVertex(FDynamicMeshVertex(Pos1, Tangent, Normal, FVector2f(0.f, 1.f), DrawColor));
+				MeshBuilder.AddVertex(FDynamicMeshVertex(Pos2, Tangent, Normal, FVector2f(1.f, 1.f), DrawColor));
+				MeshBuilder.AddTriangle(VertexIndex, VertexIndex + 1, VertexIndex + 2);
+			}
+
+			FMatrix LocalSimSpaceToWorld(FMatrix::Identity);
+			LocalSimSpaceToWorld.SetOrigin(LocalSpaceLocation);
+			MeshBuilder.Draw(PDI, LocalSimSpaceToWorld, MaterialRenderProxy, SDPG_World, false, false);
+		}
+#endif
+	}
+
 	void FClothVisualization::DrawEdgeConstraint(FPrimitiveDrawInterface* PDI) const
 	{
 		if (!Solver)
@@ -1309,6 +1532,70 @@ namespace Chaos
 			{
 				DrawSpringConstraintColors(PDI, Positions, LocalSpaceLocation, EdgeConstraints);
 			}
+
+			if (const Softs::FXPBDEdgeSpringConstraints* const EdgeConstraints = ClothConstraints.GetXEdgeSpringConstraints().Get())
+			{
+				DrawSpringConstraintColors(PDI, Positions, LocalSpaceLocation, EdgeConstraints);
+			}
+
+			if (const Softs::FXPBDStretchBiasElementConstraints* const StretchConstraints = ClothConstraints.GetXStretchBiasConstraints().Get())
+			{
+				const FMaterialRenderProxy* MaterialRenderProxy =
+#if WITH_EDITOR
+					ClothMaterialColor->GetRenderProxy();
+#else
+					nullptr;
+#endif
+				switch ((Private::EStretchBiasDrawMode)Private::StretchBiasDrawMode)
+				{
+				case Private::EStretchBiasDrawMode::WarpStretch: // fallthrough
+				case Private::EStretchBiasDrawMode::WeftStretch:
+					DrawStretchBiasConstraints_WarpWeftStretch(PDI, Positions, LocalSpaceLocation, MaterialRenderProxy, StretchConstraints);
+					break;
+				case Private::EStretchBiasDrawMode::BiasStretch:
+					DrawStretchBiasConstraints_BiasStretch(PDI, Positions, LocalSpaceLocation, MaterialRenderProxy, StretchConstraints);
+					break;
+				case Private::EStretchBiasDrawMode::ParallelGraphColor: // fallthrough
+				default:
+					DrawStretchBiasConstraints_ParallelGraphColor(PDI, Positions, LocalSpaceLocation, MaterialRenderProxy, StretchConstraints);
+				}
+
+			}
+		}
+	}
+
+	static void DrawBendingElementBuckleStatus(FPrimitiveDrawInterface* PDI, const TConstArrayView<::Chaos::Softs::FSolverVec3>& Positions, const ::Chaos::FVec3& LocalSpaceLocation, const ::Chaos::Softs::FPBDBendingConstraintsBase* const BendingConstraints)
+	{
+		const TArray<TVec4<int32>>& Constraints = BendingConstraints->GetConstraints();
+		const TArray<bool>& IsBuckled = BendingConstraints->GetIsBuckled();
+
+		// Color constraint edge with red or blue: Red = Buckled, Blue = Not Buckled. 
+		for (int32 ConstraintIndex = 0; ConstraintIndex < Constraints.Num(); ++ConstraintIndex)
+		{
+			const Softs::FSolverVec3& P1 = Positions[Constraints[ConstraintIndex][0]];
+			const Softs::FSolverVec3& P2 = Positions[Constraints[ConstraintIndex][1]];
+
+			const bool bIsBuckled = IsBuckled[ConstraintIndex];
+
+			const FVec3 Pos0 = FVec3(P1) + LocalSpaceLocation;
+			const FVec3 Pos1 = FVec3(P2) + LocalSpaceLocation;
+			DrawLine(PDI, Pos0, Pos1, bIsBuckled ? FLinearColor::Red : FLinearColor::Blue);
+		}
+	}
+
+	static void DrawBendingElementAnisotropy(FPrimitiveDrawInterface* PDI, const TConstArrayView<::Chaos::Softs::FSolverVec3>& Positions, const ::Chaos::FVec3& LocalSpaceLocation, const ::Chaos::Softs::FXPBDAnisotropicBendingConstraints* const BendingConstraints)
+	{
+		const TArray<TVec4<int32>>& Constraints = BendingConstraints->GetConstraints();
+		const TArray<Softs::FSolverVec3>& WarpWeftBiasBaseMultipliers = BendingConstraints->GetWarpWeftBiasBaseMultipliers();
+		for (int32 ConstraintIndex = 0; ConstraintIndex < Constraints.Num(); ++ConstraintIndex)
+		{
+			const Softs::FSolverVec3& P1 = Positions[Constraints[ConstraintIndex][0]];
+			const Softs::FSolverVec3& P2 = Positions[Constraints[ConstraintIndex][1]];
+			const Softs::FSolverVec3& Multiplier = WarpWeftBiasBaseMultipliers[ConstraintIndex];
+
+			const FVec3 Pos0 = FVec3(P1) + LocalSpaceLocation;
+			const FVec3 Pos1 = FVec3(P2) + LocalSpaceLocation;
+			DrawLine(PDI, Pos0, Pos1, FLinearColor(Multiplier[0], Multiplier[1], Multiplier[2]));
 		}
 	}
 
@@ -1339,41 +1626,53 @@ namespace Chaos
 				DrawSpringConstraintColors(PDI, Positions, LocalSpaceLocation, BendingConstraints);
 			}
 
-			if (const Softs::FPBDBendingConstraints* const BendingConstraints = ClothConstraints.GetBendingElementConstraints().Get())
+			if (const Softs::FXPBDBendingSpringConstraints* const BendingConstraints = ClothConstraints.GetXBendingSpringConstraints().Get())
 			{
-				const TArray<TVec4<int32>>& Constraints = BendingConstraints->GetConstraints();
-				const TArray<bool>& IsBuckled = BendingConstraints->GetIsBuckled();
+				DrawSpringConstraintColors(PDI, Positions, LocalSpaceLocation, BendingConstraints);
+			}
 
-				// Color constraint edge with red or blue: Red = Buckled, Blue = Not Buckled. 
-				for (int32 ConstraintIndex = 0; ConstraintIndex < Constraints.Num(); ++ConstraintIndex)
+			if (const Softs::FPBDBendingConstraints* const BendingConstraints = ClothConstraints.GetBendingElementConstraints().Get())
+			{
+				switch ((Private::EBendingDrawMode)Private::BendingDrawMode)
 				{
-					const Softs::FSolverVec3& P1 = Positions[Constraints[ConstraintIndex][0]];
-					const Softs::FSolverVec3& P2 = Positions[Constraints[ConstraintIndex][1]];
-				
-					const bool bIsBuckled = IsBuckled[ConstraintIndex];
-
-					const FVec3 Pos0 = FVec3(P1) + LocalSpaceLocation;
-					const FVec3 Pos1 = FVec3(P2) + LocalSpaceLocation;
-					DrawLine(PDI, Pos0, Pos1, bIsBuckled ? FLinearColor::Red : FLinearColor::Blue);
+				case Private::EBendingDrawMode::ParallelGraphColor:
+					DrawSpringConstraintColors(PDI, Positions, LocalSpaceLocation, BendingConstraints);
+					break;
+				case Private::EBendingDrawMode::BuckleStatus:
+				default:
+					DrawBendingElementBuckleStatus(PDI, Positions, LocalSpaceLocation, BendingConstraints);
+					break;
 				}
 			}
 
 			if (const Softs::FXPBDBendingConstraints* const BendingConstraints = ClothConstraints.GetXBendingElementConstraints().Get())
 			{
-				const TArray<TVec4<int32>>& Constraints = BendingConstraints->GetConstraints();
-				const TArray<bool>& IsBuckled = BendingConstraints->GetIsBuckled();
-
-				// Color constraint edge with red or blue: Red = Buckled, Blue = Not Buckled.
-				for (int32 ConstraintIndex = 0; ConstraintIndex < Constraints.Num(); ++ConstraintIndex)
+				switch ((Private::EBendingDrawMode)Private::BendingDrawMode)
 				{
-					const Softs::FSolverVec3& P1 = Positions[Constraints[ConstraintIndex][0]];
-					const Softs::FSolverVec3& P2 = Positions[Constraints[ConstraintIndex][1]];
-
-					const bool bIsBuckled = IsBuckled[ConstraintIndex];
+				case Private::EBendingDrawMode::ParallelGraphColor:
+					DrawSpringConstraintColors(PDI, Positions, LocalSpaceLocation, BendingConstraints);
+					break;
+				case Private::EBendingDrawMode::BuckleStatus:
+				default:
+					DrawBendingElementBuckleStatus(PDI, Positions, LocalSpaceLocation, BendingConstraints);
+					break;
+				}
+			}
 
-					const FVec3 Pos0 = FVec3(P1) + LocalSpaceLocation;
-					const FVec3 Pos1 = FVec3(P2) + LocalSpaceLocation;
-					DrawLine(PDI, Pos0, Pos1, bIsBuckled ? FLinearColor::Red : FLinearColor::Blue);
+			if (const Softs::FXPBDAnisotropicBendingConstraints* const BendingConstraints = ClothConstraints.GetXAnisoBendingElementConstraints().Get())
+			{
+				switch ((Private::EBendingDrawMode)Private::BendingDrawMode)
+				{
+				case Private::EBendingDrawMode::ParallelGraphColor:
+					DrawSpringConstraintColors(PDI, Positions, LocalSpaceLocation, BendingConstraints);
+					break;
+				case Private::EBendingDrawMode::Anisotropy:
+					DrawBendingElementAnisotropy(PDI, Positions, LocalSpaceLocation, BendingConstraints);
+					break;
+				case Private::EBendingDrawMode::BuckleStatus:
+				default:
+					DrawBendingElementBuckleStatus(PDI, Positions, LocalSpaceLocation, BendingConstraints);
+					break;
 				}
 			}
 		}
diff --git a/Engine/Plugins/Experimental/ChaosCloth/Source/Chaos/Public/ChaosCloth/ChaosClothVisualization.h b/Engine/Plugins/Experimental/ChaosCloth/Source/Chaos/Public/ChaosCloth/ChaosClothVisualization.h
--- a/Engine/Plugins/Experimental/ChaosCloth/Source/Chaos/Public/ChaosCloth/ChaosClothVisualization.h
+++ b/Engine/Plugins/Experimental/ChaosCloth/Source/Chaos/Public/ChaosCloth/ChaosClothVisualization.h
@@ -105,6 +105,7 @@ namespace Chaos
 #if WITH_EDITOR
 		// Visualization material
 		TObjectPtr<const UMaterial> ClothMaterial = nullptr;
+		TObjectPtr<const UMaterial> ClothMaterialColor = nullptr;
 		TObjectPtr<const UMaterial> ClothMaterialVertex = nullptr;
 		TObjectPtr<const UMaterial> CollisionMaterial = nullptr;
 #endif  // #if WITH_EDITOR
diff --git a/Engine/Source/Runtime/Experimental/Chaos/Private/Chaos/XPBDStretchBiasElementConstraints.cpp b/Engine/Source/Runtime/Experimental/Chaos/Private/Chaos/XPBDStretchBiasElementConstraints.cpp
--- a/Engine/Source/Runtime/Experimental/Chaos/Private/Chaos/XPBDStretchBiasElementConstraints.cpp
+++ b/Engine/Source/Runtime/Experimental/Chaos/Private/Chaos/XPBDStretchBiasElementConstraints.cpp
@@ -237,6 +237,8 @@ void FXPBDStretchBiasElementConstraints::InitConstraintsAndRestData(const FSolve
 
 void FXPBDStretchBiasElementConstraints::InitColor(const FSolverParticles& InParticles)
 {
+	ConstraintsPerColorStartIndex.Reset();
+
 	// In dev builds we always color so we can tune the system without restarting. See Apply()
 #if UE_BUILD_SHIPPING || UE_BUILD_TEST
 	if (Constraints.Num() > Chaos_XPBDStretchBias_ParallelConstraintCount)
@@ -274,6 +276,7 @@ void FXPBDStretchBiasElementConstraints::InitColor(const FSolverParticles& InPar
 				++ReorderedIndex;
 			}
 		}
+		check(ReorderedIndex == Constraints.Num());
 		ConstraintsPerColorStartIndex.Add(ReorderedIndex);
 
 		Constraints = MoveTemp(ReorderedConstraints);
@@ -285,6 +288,14 @@ void FXPBDStretchBiasElementConstraints::InitColor(const FSolverParticles& InPar
 		StiffnessBias.ReorderIndices(OrigToReorderedIndices);
 		DampingRatio.ReorderIndices(OrigToReorderedIndices);
 	}
+
+	// Ensure we have valid color indices
+	if (ConstraintsPerColorStartIndex.Num() < 2)
+	{
+		ConstraintsPerColorStartIndex.Reset(2);
+		ConstraintsPerColorStartIndex.Add(0);
+		ConstraintsPerColorStartIndex.Add(Constraints.Num());
+	}
 }
 
 void FXPBDStretchBiasElementConstraints::SetProperties(
@@ -429,7 +440,7 @@ void FXPBDStretchBiasElementConstraints::Apply(FSolverParticles& Particles, cons
 	const FSolverReal WarpScaleNoMap = (FSolverReal)WarpScale;
 	const FSolverReal WeftScaleNoMap = (FSolverReal)WeftScale;
 
-	if (ConstraintsPerColorStartIndex.Num() > 0 && Constraints.Num() > Chaos_XPBDStretchBias_ParallelConstraintCount)
+	if (ConstraintsPerColorStartIndex.Num() > 1 && Constraints.Num() > Chaos_XPBDStretchBias_ParallelConstraintCount)
 	{
 		const int32 ConstraintColorNum = ConstraintsPerColorStartIndex.Num() - 1;
 		if (!StiffnessWarpHasWeightMap && !StiffnessWeftHasWeightMap && !StiffnessBiasHasWeightMap && !DampingHasWeightMap && !WarpScaleHasWeightMap && !WeftScaleHasWeightMap)
@@ -635,6 +646,16 @@ void FXPBDStretchBiasElementConstraints::Apply(FSolverParticles& Particles, cons
 	}
 }
 
+void FXPBDStretchBiasElementConstraints::CalculateUVStretch(const int32 ConstraintIndex, const FSolverVec3& P0, const FSolverVec3& P1, const FSolverVec3& P2, FSolverVec3& dX_dU, FSolverVec3& dX_dV) const
+{
+	const FSolverMatrix22& DeltaUVInv = DeltaUVInverse[ConstraintIndex];
+
+	const FSolverVec3 P01 = P1 - P0;
+	const FSolverVec3 P02 = P2 - P0;
+	dX_dU = P01 * DeltaUVInv.M[0] + P02 * DeltaUVInv.M[1];
+	dX_dV = P01 * DeltaUVInv.M[2] + P02 * DeltaUVInv.M[3];
+}
+
 static FSolverReal CalcDLambda(const FSolverReal StiffnessValue, const FSolverReal Dt, const FSolverReal Lambda, const FSolverReal Damping,
 	const FSolverReal InvM1, const FSolverReal InvM2, const FSolverReal InvM3,
 	const FSolverVec3& V1TimesDt, const FSolverVec3& V2TimesDt, const FSolverVec3& V3TimesDt,
@@ -664,12 +685,8 @@ void FXPBDStretchBiasElementConstraints::ApplyHelper(FSolverParticles& Particles
 
 	const FSolverReal CombinedInvMass = Particles.InvM(i0) + Particles.InvM(i1) + Particles.InvM(i2);
 
-	const FSolverMatrix22& DeltaUVInv = DeltaUVInverse[ConstraintIndex];
-
-	const FSolverVec3 P01 = Particles.P(i1) - Particles.P(i0);
-	const FSolverVec3 P02 = Particles.P(i2) - Particles.P(i0);
-	const FSolverVec3 dX_dU = P01 * DeltaUVInv.M[0] + P02 * DeltaUVInv.M[1];
-	const FSolverVec3 dX_dV = P01 * DeltaUVInv.M[2] + P02 * DeltaUVInv.M[3];
+	FSolverVec3 dX_dU, dX_dV;
+	CalculateUVStretch(ConstraintIndex, Particles.P(i0), Particles.P(i1), Particles.P(i2), dX_dU, dX_dV);
 
 	const FSolverReal dX_dU_length = dX_dU.Length();
 	const FSolverReal dX_dV_length = dX_dV.Length();
@@ -692,6 +709,9 @@ void FXPBDStretchBiasElementConstraints::ApplyHelper(FSolverParticles& Particles
 	const FSolverReal Cv = (dX_dV_length - RestStretchLengths[ConstraintIndex][1]*WeftScaleValue); // stretch in weft direction
 	const FSolverReal Cs = FSolverVec3::DotProduct(dX_dU_normalized, dX_dV_normalized);
 
+
+	const FSolverMatrix22& DeltaUVInv = DeltaUVInverse[ConstraintIndex];
+
 	const FSolverVec3 dCu_dX1 = dX_dU_normalized * DeltaUVInv.M[0];
 	const FSolverVec3 dCu_dX2 = dX_dU_normalized * DeltaUVInv.M[1];
 	const FSolverVec3 dCu_dX0 = -dCu_dX1 - dCu_dX2;
diff --git a/Engine/Source/Runtime/Experimental/Chaos/Public/Chaos/PBDBendingConstraints.h b/Engine/Source/Runtime/Experimental/Chaos/Public/Chaos/PBDBendingConstraints.h
--- a/Engine/Source/Runtime/Experimental/Chaos/Public/Chaos/PBDBendingConstraints.h
+++ b/Engine/Source/Runtime/Experimental/Chaos/Public/Chaos/PBDBendingConstraints.h
@@ -113,6 +113,8 @@ PRAGMA_ENABLE_DEPRECATION_WARNINGS
 
 	void Apply(FSolverParticles& InParticles, const FSolverReal Dt) const;
 
+	const TArray<int32>& GetConstraintsPerColorStartIndex() const { return ConstraintsPerColorStartIndex; }
+
 private:
 	void InitColor(const FSolverParticles& InParticles);
 	void ApplyHelper(FSolverParticles& InParticles, const FSolverReal Dt, const int32 ConstraintIndex, const FSolverReal ExpStiffnessValue, const FSolverReal ExpBucklingValue) const;
diff --git a/Engine/Source/Runtime/Experimental/Chaos/Public/Chaos/XPBDAnisotropicBendingConstraints.h b/Engine/Source/Runtime/Experimental/Chaos/Public/Chaos/XPBDAnisotropicBendingConstraints.h
--- a/Engine/Source/Runtime/Experimental/Chaos/Public/Chaos/XPBDAnisotropicBendingConstraints.h
+++ b/Engine/Source/Runtime/Experimental/Chaos/Public/Chaos/XPBDAnisotropicBendingConstraints.h
@@ -93,6 +93,9 @@ public:
 
 	void Apply(FSolverParticles& Particles, const FSolverReal Dt) const;
 
+	const TArray<int32>& GetConstraintsPerColorStartIndex() const { return ConstraintsPerColorStartIndex; }
+	const TArray<FSolverVec3>& GetWarpWeftBiasBaseMultipliers() const { return WarpWeftBiasBaseMultipliers; }
+
 private:
 	void InitColor(const FSolverParticles& InParticles);
 	void ApplyHelper(FSolverParticles& Particles, const FSolverReal Dt, const int32 ConstraintIndex, const FSolverVec3& ExpStiffnessValues, 
diff --git a/Engine/Source/Runtime/Experimental/Chaos/Public/Chaos/XPBDBendingConstraints.h b/Engine/Source/Runtime/Experimental/Chaos/Public/Chaos/XPBDBendingConstraints.h
--- a/Engine/Source/Runtime/Experimental/Chaos/Public/Chaos/XPBDBendingConstraints.h
+++ b/Engine/Source/Runtime/Experimental/Chaos/Public/Chaos/XPBDBendingConstraints.h
@@ -194,6 +194,8 @@ public:
 
 	void Apply(FSolverParticles& Particles, const FSolverReal Dt) const;
 
+	const TArray<int32>& GetConstraintsPerColorStartIndex() const { return ConstraintsPerColorStartIndex; }
+
 private:
 	void InitColor(const FSolverParticles& InParticles);
 	void ApplyHelper(FSolverParticles& Particles, const FSolverReal Dt, const int32 ConstraintIndex, const FSolverReal ExpStiffnessValue, const FSolverReal ExpBucklingValue, const FSolverReal DampingRatioValue) const;
diff --git a/Engine/Source/Runtime/Experimental/Chaos/Public/Chaos/XPBDSpringConstraints.h b/Engine/Source/Runtime/Experimental/Chaos/Public/Chaos/XPBDSpringConstraints.h
--- a/Engine/Source/Runtime/Experimental/Chaos/Public/Chaos/XPBDSpringConstraints.h
+++ b/Engine/Source/Runtime/Experimental/Chaos/Public/Chaos/XPBDSpringConstraints.h
@@ -101,6 +101,8 @@ public:
 
 	void Apply(FSolverParticles& Particles, const FSolverReal Dt) const;
 
+	const TArray<int32>& GetConstraintsPerColorStartIndex() const { return ConstraintsPerColorStartIndex; }
+
 private:
 	void InitColor(const FSolverParticles& InParticles);
 	void ApplyHelper(FSolverParticles& Particles, const FSolverReal Dt, const int32 ConstraintIndex, const FSolverReal StiffnessValue, const FSolverReal DampingRatioValue) const;
diff --git a/Engine/Source/Runtime/Experimental/Chaos/Public/Chaos/XPBDStretchBiasElementConstraints.h b/Engine/Source/Runtime/Experimental/Chaos/Public/Chaos/XPBDStretchBiasElementConstraints.h
--- a/Engine/Source/Runtime/Experimental/Chaos/Public/Chaos/XPBDStretchBiasElementConstraints.h
+++ b/Engine/Source/Runtime/Experimental/Chaos/Public/Chaos/XPBDStretchBiasElementConstraints.h
@@ -97,6 +97,16 @@ public:
 	}
 
 	void Apply(FSolverParticles& Particles, const FSolverReal Dt) const;
+	void CalculateUVStretch(const int32 ConstraintIndex, const FSolverVec3& P0, const FSolverVec3& P1, const FSolverVec3& P2, FSolverVec3& dX_dU, FSolverVec3& dX_dV) const;
+	
+	const TArray<TVec3<int32>>& GetConstraints() const { return Constraints; }
+	const TArray<FSolverVec2> GetRestStretchLengths() const { return RestStretchLengths; }
+	FSolverVec2 GetWarpWeftScale(const int32 ConstraintIndex) const 
+	{ 
+		return FSolverVec2(WarpScale.HasWeightMap() ? WarpScale[ConstraintIndex] : (FSolverReal)WarpScale,
+			WeftScale.HasWeightMap() ? WeftScale[ConstraintIndex] : (FSolverReal)WeftScale);
+	}
+	const TArray<int32>& GetConstraintsPerColorStartIndex() const { return ConstraintsPerColorStartIndex; }
 
 private:
 	void InitConstraintsAndRestData(const FSolverParticles& InParticles, const FTriangleMesh& TriangleMesh,
