Chaos Cloth - Fixed a regression on the Buckling property setup.

#rnx

diff --git a/Engine/Plugins/Experimental/ChaosCloth/Source/Chaos/Private/ChaosCloth/ChaosClothingSimulationCloth.cpp b/Engine/Plugins/Experimental/ChaosCloth/Source/Chaos/Private/ChaosCloth/ChaosClothingSimulationCloth.cpp
--- a/Engine/Plugins/Experimental/ChaosCloth/Source/Chaos/Private/ChaosCloth/ChaosClothingSimulationCloth.cpp
+++ b/Engine/Plugins/Experimental/ChaosCloth/Source/Chaos/Private/ChaosCloth/ChaosClothingSimulationCloth.cpp
@@ -394,7 +394,7 @@ PRAGMA_ENABLE_DEPRECATION_WARNINGS
 
 			Properties.AddValue(TEXT("BucklingRatio"), InBucklingRatio);
 
-			if (InBucklingStiffness[0] > 0.f && InBucklingStiffness[1] > 0.f)
+			if (InBucklingStiffness[0] > 0.f || InBucklingStiffness[1] > 0.f)
 			{
 				const int32 BucklingStiffnessIndex = Properties.AddProperty(TEXT("BucklingStiffness"), bEnable, bAnimatable);
 				Properties.SetWeightedValue(BucklingStiffnessIndex, InBucklingStiffness[0], InBucklingStiffness[1]);
diff --git a/Engine/Plugins/Experimental/ChaosCloth/Source/Chaos/Private/ChaosCloth/ChaosClothingSimulationConfig.cpp b/Engine/Plugins/Experimental/ChaosCloth/Source/Chaos/Private/ChaosCloth/ChaosClothingSimulationConfig.cpp
--- a/Engine/Plugins/Experimental/ChaosCloth/Source/Chaos/Private/ChaosCloth/ChaosClothingSimulationConfig.cpp
+++ b/Engine/Plugins/Experimental/ChaosCloth/Source/Chaos/Private/ChaosCloth/ChaosClothingSimulationConfig.cpp
@@ -83,7 +83,7 @@ namespace Chaos
 
 					Properties->AddValue(TEXT("BucklingRatio"), ClothConfig->BucklingRatio);
 
-					if (ClothConfig->BucklingStiffnessWeighted.Low > 0.f && ClothConfig->BucklingStiffnessWeighted.High > 0.f)
+					if (ClothConfig->BucklingStiffnessWeighted.Low > 0.f || ClothConfig->BucklingStiffnessWeighted.High > 0.f)
 					{
 						const int32 BucklingStiffnessIndex = Properties->AddProperty(TEXT("BucklingStiffness"), bEnable, bAnimatable);
 						Properties->SetWeightedValue(BucklingStiffnessIndex, ClothConfig->BucklingStiffnessWeighted.Low, ClothConfig->BucklingStiffnessWeighted.High);
