ChaosCloth: reset colliders when resetting cloth

diff --git a/Engine/Plugins/Experimental/ChaosCloth/Source/Chaos/Private/ChaosCloth/ChaosClothingSimulationCloth.cpp b/Engine/Plugins/Experimental/ChaosCloth/Source/Chaos/Private/ChaosCloth/ChaosClothingSimulationCloth.cpp
--- a/Engine/Plugins/Experimental/ChaosCloth/Source/Chaos/Private/ChaosCloth/ChaosClothingSimulationCloth.cpp
+++ b/Engine/Plugins/Experimental/ChaosCloth/Source/Chaos/Private/ChaosCloth/ChaosClothingSimulationCloth.cpp
@@ -958,6 +958,10 @@ PRAGMA_ENABLE_DEPRECATION_WARNINGS
 
 			// Reset to start pose
 			LODData[LODIndex]->ResetStartPose(Solver);
+			for (FClothingSimulationCollider* Collider : Colliders)
+			{
+				Collider->ResetStartPose(Solver, this);
+			}
 			UE_LOG(LogChaosCloth, VeryVerbose, TEXT("Cloth in group Id %d Needs reset."), GroupId);
 		}
 		else if (bNeedsTeleport)
