[Chaos]
* Second round of moving dllexport from types to methods/staticvar

diff --git a/Engine/Plugins/Experimental/ChaosCaching/Source/ChaosCaching/Public/Chaos/Adapters/CacheAdapter.h b/Engine/Plugins/Experimental/ChaosCaching/Source/ChaosCaching/Public/Chaos/Adapters/CacheAdapter.h
--- a/Engine/Plugins/Experimental/ChaosCaching/Source/ChaosCaching/Public/Chaos/Adapters/CacheAdapter.h
+++ b/Engine/Plugins/Experimental/ChaosCaching/Source/ChaosCaching/Public/Chaos/Adapters/CacheAdapter.h
@@ -34,7 +34,7 @@ namespace Chaos
 	 * This base interface should not be used to implement those adapters however but the derived classes
 	 * FCacheRecordAdapter and FCachePlaybackAdapter (declared below) should instead be used.
 	 */
-	class CHAOSCACHING_API FComponentCacheAdapter : public IModularFeature
+	class FComponentCacheAdapter : public IModularFeature
 	{
 	public:
 		enum class SupportType : uint8
@@ -45,9 +45,9 @@ namespace Chaos
 		};
 
 		/** Registration name for modular features module */
-		static const FName FeatureName;
-		static const uint8 EngineAdapterPriorityBegin;
-		static const uint8 UserAdapterPriorityBegin;
+		static CHAOSCACHING_API const FName FeatureName;
+		static CHAOSCACHING_API const uint8 EngineAdapterPriorityBegin;
+		static CHAOSCACHING_API const uint8 UserAdapterPriorityBegin;
 
 		FComponentCacheAdapter()          = default;
 		virtual ~FComponentCacheAdapter() = default;
diff --git a/Engine/Plugins/Experimental/ChaosCaching/Source/ChaosCaching/Public/Chaos/CacheCollection.h b/Engine/Plugins/Experimental/ChaosCaching/Source/ChaosCaching/Public/Chaos/CacheCollection.h
--- a/Engine/Plugins/Experimental/ChaosCaching/Source/ChaosCaching/Public/Chaos/CacheCollection.h
+++ b/Engine/Plugins/Experimental/ChaosCaching/Source/ChaosCaching/Public/Chaos/CacheCollection.h
@@ -6,26 +6,26 @@
 
 class UChaosCache;
 
-UCLASS(Experimental, BlueprintType)
-class CHAOSCACHING_API UChaosCacheCollection : public UObject
+UCLASS(Experimental, BlueprintType, MinimalAPI)
+class UChaosCacheCollection : public UObject
 {
 	GENERATED_BODY()
 public:
 
 	//~ Begin UObject Interface
-	virtual void GetAssetRegistryTags(TArray<FAssetRegistryTag>& OutTags) const override;
+	CHAOSCACHING_API virtual void GetAssetRegistryTags(TArray<FAssetRegistryTag>& OutTags) const override;
 	//~ End UObject Interface
 	
-	UChaosCache* FindCache(const FName& CacheName) const;
-	UChaosCache* FindOrAddCache(const FName& CacheName);
+	CHAOSCACHING_API UChaosCache* FindCache(const FName& CacheName) const;
+	CHAOSCACHING_API UChaosCache* FindOrAddCache(const FName& CacheName);
 
-	void FlushAllCacheWrites();
+	CHAOSCACHING_API void FlushAllCacheWrites();
 
 	const TArray<UChaosCache*>& GetCaches() const {return Caches;}
 
 	/** Return the max duration of all the caches stored in the collection */
-	float GetMaxDuration() const;
+	CHAOSCACHING_API float GetMaxDuration() const;
 
 	UPROPERTY(EditAnywhere, Instanced, Category="Caching", meta=(EditFixedOrder))
 	TArray<TObjectPtr<UChaosCache>> Caches;
-};
\ No newline at end of file
+};
diff --git a/Engine/Plugins/Experimental/ChaosCaching/Source/ChaosCaching/Public/Chaos/CacheManagerActor.h b/Engine/Plugins/Experimental/ChaosCaching/Source/ChaosCaching/Public/Chaos/CacheManagerActor.h
--- a/Engine/Plugins/Experimental/ChaosCaching/Source/ChaosCaching/Public/Chaos/CacheManagerActor.h
+++ b/Engine/Plugins/Experimental/ChaosCaching/Source/ChaosCaching/Public/Chaos/CacheManagerActor.h
@@ -40,7 +40,7 @@ namespace Chaos
 }
 
 USTRUCT(BlueprintType)
-struct CHAOSCACHING_API FObservedComponent
+struct FObservedComponent
 {
 	GENERATED_BODY()
 
@@ -116,17 +116,17 @@ struct CHAOSCACHING_API FObservedComponent
 	bool bHasNotifyBreaks;
 
 	/** Post serialize function to transfer datas from the deprecated TObjectPtr -> TSoftObjectPtr */
-	void PostSerialize(const FArchive& Ar);
+	CHAOSCACHING_API void PostSerialize(const FArchive& Ar);
 
 	/** Prepare runtime tick data for a new run */
-	void ResetRuntimeData(const EStartMode ManagerStartMode);
+	CHAOSCACHING_API void ResetRuntimeData(const EStartMode ManagerStartMode);
 
 	/** Check if the Observed component is enabled for a specific cache mode */
-	bool IsEnabled(ECacheMode CacheMode) const;
+	CHAOSCACHING_API bool IsEnabled(ECacheMode CacheMode) const;
 
 	/** Gets the component from the internal component ref */
-	UPrimitiveComponent* GetComponent(AActor* OwningActor);
-	UPrimitiveComponent* GetComponent(AActor* OwningActor) const;
+	CHAOSCACHING_API UPrimitiveComponent* GetComponent(AActor* OwningActor);
+	CHAOSCACHING_API UPrimitiveComponent* GetComponent(AActor* OwningActor) const;
 	UE_DEPRECATED(5.3, "Use GetComponent(OwningActor) instead.")
 	UPrimitiveComponent* GetComponent() { return GetComponent(nullptr); }
 	UE_DEPRECATED(5.3, "Use GetComponent(OwningActor) instead.")
@@ -162,13 +162,13 @@ struct FPerSolverData
 	TArray<Chaos::TPBDRigidParticleHandle<Chaos::FReal, 3>*> PendingKinematicUpdates;
 };
 
-UCLASS(Experimental)
-class CHAOSCACHING_API AChaosCacheManager : public AActor
+UCLASS(Experimental, MinimalAPI)
+class AChaosCacheManager : public AActor
 {
 	GENERATED_BODY()
 
 public:
-	AChaosCacheManager(const FObjectInitializer& ObjectInitializer = FObjectInitializer::Get());
+	CHAOSCACHING_API AChaosCacheManager(const FObjectInitializer& ObjectInitializer = FObjectInitializer::Get());
 
 	/**
 	 * The Cache Collection asset to use for this observer. This can be used for playback and record simultaneously
@@ -197,30 +197,30 @@ public:
 	float StartTime;
 
 	/** AActor interface */
-	void TickActor(float DeltaTime, enum ELevelTick TickType, FActorTickFunction& ThisTickFunction) override;
+	CHAOSCACHING_API void TickActor(float DeltaTime, enum ELevelTick TickType, FActorTickFunction& ThisTickFunction) override;
 	/** end AActor interface */
 
 	/** UObject interface */
-	virtual void Serialize(FArchive& Ar) override;
+	CHAOSCACHING_API virtual void Serialize(FArchive& Ar) override;
 #if WITH_EDITOR
 	friend class IChaosCachingEditorPlugin;
 #endif
 
 #if WITH_EDITOR
-	virtual void PostEditChangeProperty(FPropertyChangedEvent& PropertyChangedEvent) override;
+	CHAOSCACHING_API virtual void PostEditChangeProperty(FPropertyChangedEvent& PropertyChangedEvent) override;
 #endif
 	/** end UObject interface */
 
 	/** Expose StartTime property to Sequencer. GetStartTime will be called on keys. */
 	UFUNCTION(CallInEditor)
-	void SetStartTime(float InStartTime);
+	CHAOSCACHING_API void SetStartTime(float InStartTime);
 
 	/** 
 	 * Resets all components back to the world space transform they had when the cache for them was originally recorded
 	 * if one is available
 	 */
 	UFUNCTION(BlueprintCallable, Category = "Caching")
-	void ResetAllComponentTransforms();
+	CHAOSCACHING_API void ResetAllComponentTransforms();
 
 	/**
 	 * Resets the component at the specified index in the observed list back to the world space transform it had when the 
@@ -228,24 +228,24 @@ public:
 	 * @param InIndex Index of the component to reset
 	 */
 	UFUNCTION(BlueprintCallable, Category = "Caching")
-	void ResetSingleTransform(int32 InIndex);
+	CHAOSCACHING_API void ResetSingleTransform(int32 InIndex);
 
 #if WITH_EDITOR
 	/**
 	 * Set the component at the specified index in the observed array to be the selected component in the outliner.
 	 * This will also make that component's owner the selected actor in the outliner.
 	 */
-	void SelectComponent(int32 InIndex);
+	CHAOSCACHING_API void SelectComponent(int32 InIndex);
 #endif
 
 	/** Returns true if this cache manager is allowed to record caches. */
 	bool CanRecord() const { return bCanRecord; }
 
 	/** Initialize the cache adapters before playing/recording the cache. */
-	void BeginEvaluate();
+	CHAOSCACHING_API void BeginEvaluate();
 	
 	/** Clean the cache adapters after playing/recording the cache.  */
-	void EndEvaluate();
+	CHAOSCACHING_API void EndEvaluate();
 
 	/** Accessor to the manager observed components (read only) */
 	const TArray<FObservedComponent>& GetObservedComponents() const {return ObservedComponents;}
@@ -256,31 +256,31 @@ public:
 protected:
 
 	/** AActor interface */
-	void BeginPlay() override;
-	void EndPlay(const EEndPlayReason::Type EndPlayReason) override;
+	CHAOSCACHING_API void BeginPlay() override;
+	CHAOSCACHING_API void EndPlay(const EEndPlayReason::Type EndPlayReason) override;
 	/** End AActor interface */
 	
 	/**	Handles physics thread pre-solve (push kinematic data for components under playback) */
-    void HandlePreSolve(Chaos::FReal InDt, Chaos::FPhysicsSolverEvents* InSolver);
+    CHAOSCACHING_API void HandlePreSolve(Chaos::FReal InDt, Chaos::FPhysicsSolverEvents* InSolver);
    
     /** Handles physics thread pre-buffer (mark dirty kinematic particles) */
-    void HandlePreBuffer(Chaos::FReal InDt, Chaos::FPhysicsSolverEvents* InSolver);
+    CHAOSCACHING_API void HandlePreBuffer(Chaos::FReal InDt, Chaos::FPhysicsSolverEvents* InSolver);
    
     /** Handles physics thread post-solve (record data for components under record) */
-    void HandlePostSolve(Chaos::FReal InDt, Chaos::FPhysicsSolverEvents* InSolver);
+    CHAOSCACHING_API void HandlePostSolve(Chaos::FReal InDt, Chaos::FPhysicsSolverEvents* InSolver);
    
     /** Handles solver teardown due to solver destruction / stream-out */
-    void HandleTeardown(Chaos::FPhysicsSolverEvents* InSolver);
+    CHAOSCACHING_API void HandleTeardown(Chaos::FPhysicsSolverEvents* InSolver);
 	
 	/** Evaluates and sets state for all observed components at the specified time. */
-	void OnStartFrameChanged(Chaos::FReal InT);
+	CHAOSCACHING_API void OnStartFrameChanged(Chaos::FReal InT);
 
 	/**
 	* change the cache collection for this player 
 	* if the cache is playing or recording this will have no effect
 	*/
 	UFUNCTION(BlueprintCallable, Category = "Caching")
-	void SetCacheCollection(UChaosCacheCollection* InCacheCollection);
+	CHAOSCACHING_API void SetCacheCollection(UChaosCacheCollection* InCacheCollection);
 
 	/**
 	 * Triggers a component to play or record.
@@ -289,7 +289,7 @@ protected:
 	 * @param InComponent Component to trigger
 	 */
 	UFUNCTION(BlueprintCallable, Category = "Caching")
-	void TriggerComponent(UPrimitiveComponent* InComponent);
+	CHAOSCACHING_API void TriggerComponent(UPrimitiveComponent* InComponent);
 
 	/**
 	 * Triggers a component to play or record.
@@ -298,24 +298,24 @@ protected:
 	 * @param InCacheName Cache name to trigger
 	 */
 	UFUNCTION(BlueprintCallable, Category = "Caching")
-	void TriggerComponentByCache(FName InCacheName);
+	CHAOSCACHING_API void TriggerComponentByCache(FName InCacheName);
 
 	/** Triggers the recording or playback of all observed components */
 	UFUNCTION(BlueprintCallable, Category = "Caching")
-	void TriggerAll();
+	CHAOSCACHING_API void TriggerAll();
 
 	/** Enable playback for a specific component using its cache name */
 	UFUNCTION(BlueprintCallable, Category = "Caching")
-	void EnablePlaybackByCache(FName InCacheName, bool bEnable);
+	CHAOSCACHING_API void EnablePlaybackByCache(FName InCacheName, bool bEnable);
 
 	/** Enable playback for a specific component using its index in the list of observed component */
 	UFUNCTION(BlueprintCallable, Category = "Caching")
-	void EnablePlayback(int32 Index, bool bEnable);
+	CHAOSCACHING_API void EnablePlayback(int32 Index, bool bEnable);
 
-	FObservedComponent* FindObservedComponent(UPrimitiveComponent* InComponent);
-	FObservedComponent& AddNewObservedComponent(UPrimitiveComponent* InComponent);
-	FObservedComponent& FindOrAddObservedComponent(UPrimitiveComponent* InComponent);
-	void ClearObservedComponents();
+	CHAOSCACHING_API FObservedComponent* FindObservedComponent(UPrimitiveComponent* InComponent);
+	CHAOSCACHING_API FObservedComponent& AddNewObservedComponent(UPrimitiveComponent* InComponent);
+	CHAOSCACHING_API FObservedComponent& FindOrAddObservedComponent(UPrimitiveComponent* InComponent);
+	CHAOSCACHING_API void ClearObservedComponents();
 
 	// Determines if the actor is allowed to record a cache.
 	bool bCanRecord;
@@ -325,7 +325,7 @@ protected:
 
 private:
 #if WITH_EDITOR
-	void SetObservedComponentProperties(const ECacheMode& NewCacheMode);
+	CHAOSCACHING_API void SetObservedComponentProperties(const ECacheMode& NewCacheMode);
 #endif
 	
 	friend class UActorFactoryCacheManager; // Allows the actor factory to set up the observed list. See UActorFactoryCacheManager::PostSpawnActor
@@ -339,12 +339,12 @@ private:
 	 * @param InDt Delta for the tick
 	 * @param InCallable Callable to fire if the observed component is active
 	 */
-	void TickObservedComponents(const TArray<int32>& InIndices, Chaos::FReal InDt, FTickObservedFunction InCallable);
+	CHAOSCACHING_API void TickObservedComponents(const TArray<int32>& InIndices, Chaos::FReal InDt, FTickObservedFunction InCallable);
 
 	/**
 	* Stall on any in flight solver tasks that might call callbacks. Necessary before editing callbacks.
 	*/
-	void WaitForObservedComponentSolverTasks();
+	CHAOSCACHING_API void WaitForObservedComponentSolverTasks();
 
 	/** List of observed objects and their caches */
 	UPROPERTY(EditAnywhere, Category = "Caching")
@@ -361,13 +361,13 @@ private:
 	TArray<TTuple<FCacheUserToken, UChaosCache*>> OpenPlaybackCaches;
 };
 
-UCLASS(Experimental)
-class CHAOSCACHING_API AChaosCachePlayer : public AChaosCacheManager
+UCLASS(Experimental, MinimalAPI)
+class AChaosCachePlayer : public AChaosCacheManager
 {
 	GENERATED_BODY()
 
 public:
-	AChaosCachePlayer(const FObjectInitializer& ObjectInitializer = FObjectInitializer::Get());
+	CHAOSCACHING_API AChaosCachePlayer(const FObjectInitializer& ObjectInitializer = FObjectInitializer::Get());
 
 };
 
diff --git a/Engine/Plugins/Experimental/ChaosCaching/Source/ChaosCaching/Public/Chaos/ChaosCache.h b/Engine/Plugins/Experimental/ChaosCaching/Source/ChaosCaching/Public/Chaos/ChaosCache.h
--- a/Engine/Plugins/Experimental/ChaosCaching/Source/ChaosCaching/Public/Chaos/ChaosCache.h
+++ b/Engine/Plugins/Experimental/ChaosCaching/Source/ChaosCaching/Public/Chaos/ChaosCache.h
@@ -212,7 +212,7 @@ struct FCompressedRichCurves
 	TArray<FCompressedRichCurve> CompressedRichCurves;
 };
 
-struct CHAOSCACHING_API FPendingFrameWrite
+struct FPendingFrameWrite
 {
 	float                         Time;
 	TArray<FPendingParticleWrite> PendingParticleData;
@@ -287,17 +287,17 @@ private:
 	FCacheUserToken& operator=(FCacheUserToken&&) = delete;
 };
 
-UCLASS(Experimental)
-class CHAOSCACHING_API UChaosCache : public UObject
+UCLASS(Experimental, MinimalAPI)
+class UChaosCache : public UObject
 {
 	GENERATED_BODY()
 public:
 
-	UChaosCache();
+	CHAOSCACHING_API UChaosCache();
 
 	//~ UObject interface
-	virtual void Serialize(FArchive& Ar) override;
-	virtual void PostLoad() override;
+	CHAOSCACHING_API virtual void Serialize(FArchive& Ar) override;
+	CHAOSCACHING_API virtual void PostLoad() override;
 	//~ UObject interface END
 
 	/**
@@ -305,13 +305,13 @@ public:
 	 * Because of this we can't directly write into the cache, but instead into a pending frame queue that needs to be
 	 * flushed on the main thread to write the pending data into the final storage.
 	 */
-	void FlushPendingFrames();
+	CHAOSCACHING_API void FlushPendingFrames();
 
 	/**
 	 * Reset and initialize a cache to make it ready to record the specified component
 	 * @param InComponent Component to prepare the cache for
 	 */
-	FCacheUserToken BeginRecord(const UPrimitiveComponent* InComponent, FGuid InAdapterId, const FTransform& SpaceTransform);
+	CHAOSCACHING_API FCacheUserToken BeginRecord(const UPrimitiveComponent* InComponent, FGuid InAdapterId, const FTransform& SpaceTransform);
 
 	/**
 	 * End the recording session for the cache. At this point the cache is deemed to now contain
@@ -319,13 +319,13 @@ public:
 	 * optimized which may involve key elimination and compression into a final format for runtime
 	 * @param InOutToken The token that was given by BeginRecord
 	 */
-	void EndRecord(FCacheUserToken& InOutToken);
+	CHAOSCACHING_API void EndRecord(FCacheUserToken& InOutToken);
 
 	/**
 	 * Initialise the cache for playback, may not take any actual action on the cache but
 	 * will provide the caller with a valid cache user token if it is safe to continue with playback
 	 */
-	FCacheUserToken BeginPlayback();
+	CHAOSCACHING_API FCacheUserToken BeginPlayback();
 
 	/**
 	 * End a playback session for the cache. There can be multiple playback sessions open for a
@@ -333,18 +333,18 @@ public:
 	 * token will decrease the session count.
 	 * @param InOutToken The token that was given by BeginRecord
 	 */
-	void EndPlayback(FCacheUserToken& InOutToken);
+	CHAOSCACHING_API void EndPlayback(FCacheUserToken& InOutToken);
 
 	/**
 	 * Adds a new frame to process to a threadsafe queue for later processing in FlushPendingFrames
 	 * @param InFrame New frame to accept, moved into the internal threadsafe queue
 	 */
-	void AddFrame_Concurrent(FPendingFrameWrite&& InFrame);
+	CHAOSCACHING_API void AddFrame_Concurrent(FPendingFrameWrite&& InFrame);
 
 	/**
 	 * Gets the recorded duration of the cache
 	 */
-	float GetDuration() const;
+	CHAOSCACHING_API float GetDuration() const;
 
 	/**
 	 * Evaluate the cache with the specified parameters, returning the evaluated results
@@ -352,19 +352,19 @@ public:
 	 * @param MassToLocalTransforms MassToLocal trasnform if available ( geometry collection are using them )
 	 * @see FCacheEvaluationContext
 	 */
-	FCacheEvaluationResult Evaluate(const FCacheEvaluationContext& InContext, const TArray<FTransform>* MassToLocalTransforms);
+	CHAOSCACHING_API FCacheEvaluationResult Evaluate(const FCacheEvaluationContext& InContext, const TArray<FTransform>* MassToLocalTransforms);
 
 	/**
 	 * Initializes the spawnable template from a currently existing component so it can be spawned by the editor
 	 * when a cache is dragged into the scene.
 	 * @param InComponent Component to build the spawnable template from
 	 */
-	void BuildSpawnableFromComponent(const UPrimitiveComponent* InComponent, const FTransform& SpaceTransform);
+	CHAOSCACHING_API void BuildSpawnableFromComponent(const UPrimitiveComponent* InComponent, const FTransform& SpaceTransform);
 
 	/**
 	 * Read access to the spawnable template stored in the cache
 	 */
-	const FCacheSpawnableTemplate& GetSpawnableTemplate() const;
+	CHAOSCACHING_API const FCacheSpawnableTemplate& GetSpawnableTemplate() const;
 
 	/**
 	 * Evaluates a single particle from the tracks array
@@ -374,13 +374,13 @@ public:
 	 * @param OutOptTransform Transform to fill, skipped if null
 	 * @param OutOptCurves Curves to fill, skipped if null
 	 */
-	void EvaluateSingle(int32 InIndex, FPlaybackTickRecord& InTickRecord, const FTransform* MassToLocal, FTransform* OutOptTransform, TMap<FName, float>* OutOptCurves);
+	CHAOSCACHING_API void EvaluateSingle(int32 InIndex, FPlaybackTickRecord& InTickRecord, const FTransform* MassToLocal, FTransform* OutOptTransform, TMap<FName, float>* OutOptCurves);
 
-	void EvaluateTransform(const FPerParticleCacheData& InData, float InTime, const FTransform* MassToLocal, FTransform& OutTransform);
-	void EvaluateCurves(const FPerParticleCacheData& InData, float InTime, TMap<FName, float>& OutCurves);
-	void EvaluateEvents(FPlaybackTickRecord& InTickRecord, TMap<FName, TArray<FCacheEventHandle>>& OutEvents);
+	CHAOSCACHING_API void EvaluateTransform(const FPerParticleCacheData& InData, float InTime, const FTransform* MassToLocal, FTransform& OutTransform);
+	CHAOSCACHING_API void EvaluateCurves(const FPerParticleCacheData& InData, float InTime, TMap<FName, float>& OutCurves);
+	CHAOSCACHING_API void EvaluateEvents(FPlaybackTickRecord& InTickRecord, TMap<FName, TArray<FCacheEventHandle>>& OutEvents);
 
-	void CompressChannelsData(float ErrorThreshold, float SampleRate);
+	CHAOSCACHING_API void CompressChannelsData(float ErrorThreshold, float SampleRate);
 
 	UPROPERTY(VisibleAnywhere, Category = "Caching")
 	float RecordedDuration;
@@ -437,7 +437,7 @@ public:
 	}
 
 private:
-	void FlushPendingFrames_ChannelOnlyReservePass(TQueue<FPendingFrameWrite, EQueueMode::Spsc>& LocalPendingWrites,
+	CHAOSCACHING_API void FlushPendingFrames_ChannelOnlyReservePass(TQueue<FPendingFrameWrite, EQueueMode::Spsc>& LocalPendingWrites,
 		bool& bCanSimpleCopyChannelData);
 
 	// @return Whether or not any particle data was written
diff --git a/Engine/Plugins/Experimental/ChaosCaching/Source/ChaosCaching/Public/Chaos/Sequencer/ChaosCacheObjectSpawner.h b/Engine/Plugins/Experimental/ChaosCaching/Source/ChaosCaching/Public/Chaos/Sequencer/ChaosCacheObjectSpawner.h
--- a/Engine/Plugins/Experimental/ChaosCaching/Source/ChaosCaching/Public/Chaos/Sequencer/ChaosCacheObjectSpawner.h
+++ b/Engine/Plugins/Experimental/ChaosCaching/Source/ChaosCaching/Public/Chaos/Sequencer/ChaosCacheObjectSpawner.h
@@ -1,22 +1,22 @@
-﻿// Copyright Epic Games, Inc. All Rights Reserved.
+// Copyright Epic Games, Inc. All Rights Reserved.
 
 #pragma once
 
 #include "LevelSequenceActorSpawner.h"
 
 /** Chaos cache Manager spawner to create a new cache manager from the spawnable template */
-class CHAOSCACHING_API FChaosCacheObjectSpawner : public FLevelSequenceActorSpawner
+class FChaosCacheObjectSpawner : public FLevelSequenceActorSpawner
 {
 public:
-	FChaosCacheObjectSpawner();
+	CHAOSCACHING_API FChaosCacheObjectSpawner();
 
 	/** Static method to create the object spawner */
-	static TSharedRef<IMovieSceneObjectSpawner> CreateObjectSpawner();
+	static CHAOSCACHING_API TSharedRef<IMovieSceneObjectSpawner> CreateObjectSpawner();
 
 	// IMovieSceneObjectSpawner interface
-	virtual UClass* GetSupportedTemplateType() const override;
-	virtual UObject* SpawnObject(FMovieSceneSpawnable& Spawnable, FMovieSceneSequenceIDRef TemplateID, IMovieScenePlayer& Player) override;
-	virtual void DestroySpawnedObject(UObject& Object) override;
+	CHAOSCACHING_API virtual UClass* GetSupportedTemplateType() const override;
+	CHAOSCACHING_API virtual UObject* SpawnObject(FMovieSceneSpawnable& Spawnable, FMovieSceneSequenceIDRef TemplateID, IMovieScenePlayer& Player) override;
+	CHAOSCACHING_API virtual void DestroySpawnedObject(UObject& Object) override;
 	virtual bool IsEditor() const override { return true; }
 	virtual int32 GetSpawnerPriority() const override { return 1; }
-};
\ No newline at end of file
+};
diff --git a/Engine/Plugins/Experimental/ChaosCaching/Source/ChaosCaching/Public/Chaos/Sequencer/MovieSceneChaosCacheSection.h b/Engine/Plugins/Experimental/ChaosCaching/Source/ChaosCaching/Public/Chaos/Sequencer/MovieSceneChaosCacheSection.h
--- a/Engine/Plugins/Experimental/ChaosCaching/Source/ChaosCaching/Public/Chaos/Sequencer/MovieSceneChaosCacheSection.h
+++ b/Engine/Plugins/Experimental/ChaosCaching/Source/ChaosCaching/Public/Chaos/Sequencer/MovieSceneChaosCacheSection.h
@@ -6,15 +6,15 @@
 #include "MovieSceneChaosCacheSection.generated.h"
 
 USTRUCT()
-struct CHAOSCACHING_API FMovieSceneChaosCacheParams : public FMovieSceneBaseCacheParams
+struct FMovieSceneChaosCacheParams : public FMovieSceneBaseCacheParams
 {
 	GENERATED_BODY()
 
-	FMovieSceneChaosCacheParams();
+	CHAOSCACHING_API FMovieSceneChaosCacheParams();
 	virtual ~FMovieSceneChaosCacheParams() override {}
 
 	/** Gets the animation sequence length, not modified by play rate */
-	virtual float GetSequenceLength() const override;
+	CHAOSCACHING_API virtual float GetSequenceLength() const override;
 	
 	/** The animation this section plays */
 	UPROPERTY(EditAnywhere, Category = "ChaosCache")
diff --git a/Engine/Plugins/Experimental/ChaosCloth/Source/Chaos/Public/ChaosCloth/ChaosClothConfig.h b/Engine/Plugins/Experimental/ChaosCloth/Source/Chaos/Public/ChaosCloth/ChaosClothConfig.h
--- a/Engine/Plugins/Experimental/ChaosCloth/Source/Chaos/Public/ChaosCloth/ChaosClothConfig.h
+++ b/Engine/Plugins/Experimental/ChaosCloth/Source/Chaos/Public/ChaosCloth/ChaosClothConfig.h
@@ -50,25 +50,25 @@ enum class EChaosClothTetherMode : uint8
 
 /** Holds initial, asset level config for clothing actors. */
 // Hiding categories that will be used in the future
-UCLASS(HideCategories = (Collision))
-class CHAOSCLOTH_API UChaosClothConfig : public UClothConfigCommon
+UCLASS(HideCategories = (Collision), MinimalAPI)
+class UChaosClothConfig : public UClothConfigCommon
 {
 	GENERATED_BODY()
 public:
-	UChaosClothConfig();
-	virtual ~UChaosClothConfig() override;
+	CHAOSCLOTH_API UChaosClothConfig();
+	CHAOSCLOTH_API virtual ~UChaosClothConfig() override;
 
 	/** Migrate from the legacy FClothConfig structure. */
-	virtual void MigrateFrom(const FClothConfig_Legacy&) override;
+	CHAOSCLOTH_API virtual void MigrateFrom(const FClothConfig_Legacy&) override;
 
 	/** Migrate from shared config. */
-	virtual void MigrateFrom(const UClothSharedConfigCommon* ClothSharedConfig) override;
+	CHAOSCLOTH_API virtual void MigrateFrom(const UClothSharedConfigCommon* ClothSharedConfig) override;
 
 	/** Serialize override used to set the current custom version. */
-	virtual void Serialize(FArchive& Ar) override;
+	CHAOSCLOTH_API virtual void Serialize(FArchive& Ar) override;
 
 	/** PostLoad override used to deal with updates/changes in properties. */
-	virtual void PostLoad() override;
+	CHAOSCLOTH_API virtual void PostLoad() override;
 
 	/** Return wherether to pre-compute Inverse Masses. */
 	virtual bool NeedsInverseMasses() const override { return false; }  // TODO: Chaos Cloth uses the mass mode enum, and this will require a little refactor to work
@@ -80,7 +80,7 @@ public:
 	virtual bool TethersUseGeodesicDistance() const override { return bUseGeodesicDistance; }
 
 	/** Return the mass value, from whichever mass mode (Density, UniformMass, or TotalMass) is selected. */
-	float GetMassValue() const;
+	CHAOSCLOTH_API float GetMassValue() const;
 
 	/**
 	 * How cloth particle mass is determined
@@ -405,26 +405,26 @@ public:
  * parameters that cannot change between the various clothing assets assigned
  * to a specific skeletal mesh. @seealso UChaosClothConfig.
  */
-UCLASS()
-class CHAOSCLOTH_API UChaosClothSharedSimConfig : public UClothSharedConfigCommon
+UCLASS(MinimalAPI)
+class UChaosClothSharedSimConfig : public UClothSharedConfigCommon
 {
 	GENERATED_BODY()
 public:
-	UChaosClothSharedSimConfig();
-	virtual ~UChaosClothSharedSimConfig() override;
+	CHAOSCLOTH_API UChaosClothSharedSimConfig();
+	CHAOSCLOTH_API virtual ~UChaosClothSharedSimConfig() override;
 
 	/** Migrate from the legacy FClothConfig structure. */
-	virtual void MigrateFrom(const FClothConfig_Legacy& ClothConfig) override;
+	CHAOSCLOTH_API virtual void MigrateFrom(const FClothConfig_Legacy& ClothConfig) override;
 
 	/** Serialize override used to set the current custom version. */
-	virtual void Serialize(FArchive& Ar) override;
+	CHAOSCLOTH_API virtual void Serialize(FArchive& Ar) override;
 
 	/** PostLoad override used to deal with updates/changes in properties. */
-	virtual void PostLoad() override;
+	CHAOSCLOTH_API virtual void PostLoad() override;
 
 #if WITH_EDITOR
 	/** Called after changes in any of the asset properties. */
-	virtual void PostEditChangeChainProperty(FPropertyChangedChainEvent& ChainEvent) override;
+	CHAOSCLOTH_API virtual void PostEditChangeChainProperty(FPropertyChangedChainEvent& ChainEvent) override;
 #endif
 
 	/**
diff --git a/Engine/Plugins/Experimental/ChaosCloth/Source/Chaos/Public/ChaosCloth/ChaosClothVisualization.h b/Engine/Plugins/Experimental/ChaosCloth/Source/Chaos/Public/ChaosCloth/ChaosClothVisualization.h
--- a/Engine/Plugins/Experimental/ChaosCloth/Source/Chaos/Public/ChaosCloth/ChaosClothVisualization.h
+++ b/Engine/Plugins/Experimental/ChaosCloth/Source/Chaos/Public/ChaosCloth/ChaosClothVisualization.h
@@ -18,39 +18,39 @@ namespace Chaos
 {
 	class FClothingSimulationSolver;
 
-	class CHAOSCLOTH_API FClothVisualization
+	class FClothVisualization
 #if WITH_EDITOR
 		: public FGCObject  // Add garbage collection for cloth material
 #endif  // #if WITH_EDITOR
 	{
 	public:
-		explicit FClothVisualization(const ::Chaos::FClothingSimulationSolver* InSolver = nullptr);
-		virtual ~FClothVisualization();
+		CHAOSCLOTH_API explicit FClothVisualization(const ::Chaos::FClothingSimulationSolver* InSolver = nullptr);
+		CHAOSCLOTH_API virtual ~FClothVisualization();
 
 #if CHAOS_DEBUG_DRAW
 		// Editor & runtime functions
-		void SetSolver(const ::Chaos::FClothingSimulationSolver* InSolver);
+		CHAOSCLOTH_API void SetSolver(const ::Chaos::FClothingSimulationSolver* InSolver);
 
-		void DrawPhysMeshWired(FPrimitiveDrawInterface* PDI = nullptr) const;
-		void DrawAnimMeshWired(FPrimitiveDrawInterface* PDI = nullptr) const;
-		void DrawAnimNormals(FPrimitiveDrawInterface* PDI = nullptr) const;
-		void DrawOpenEdges(FPrimitiveDrawInterface* PDI = nullptr) const;
-		void DrawPointNormals(FPrimitiveDrawInterface* PDI = nullptr) const;
-		void DrawPointVelocities(FPrimitiveDrawInterface* PDI = nullptr) const;
-		void DrawCollision(FPrimitiveDrawInterface* PDI = nullptr) const;
-		void DrawBackstops(FPrimitiveDrawInterface* PDI = nullptr) const;
-		void DrawBackstopDistances(FPrimitiveDrawInterface* PDI = nullptr) const;
-		void DrawMaxDistances(FPrimitiveDrawInterface* PDI = nullptr) const;
-		void DrawAnimDrive(FPrimitiveDrawInterface* PDI = nullptr) const;
-		void DrawEdgeConstraint(FPrimitiveDrawInterface* PDI = nullptr) const;
-		void DrawBendingConstraint(FPrimitiveDrawInterface* PDI = nullptr) const;
-		void DrawLongRangeConstraint(FPrimitiveDrawInterface* PDI = nullptr) const;
-		void DrawWindAndPressureForces(FPrimitiveDrawInterface* PDI = nullptr) const;
-		void DrawLocalSpace(FPrimitiveDrawInterface* PDI = nullptr) const;
-		void DrawSelfCollision(FPrimitiveDrawInterface* PDI = nullptr) const;
-		void DrawSelfIntersection(FPrimitiveDrawInterface* PDI = nullptr) const;
-		void DrawBounds(FPrimitiveDrawInterface* PDI = nullptr) const;
-		void DrawGravity(FPrimitiveDrawInterface* PDI = nullptr) const;
+		CHAOSCLOTH_API void DrawPhysMeshWired(FPrimitiveDrawInterface* PDI = nullptr) const;
+		CHAOSCLOTH_API void DrawAnimMeshWired(FPrimitiveDrawInterface* PDI = nullptr) const;
+		CHAOSCLOTH_API void DrawAnimNormals(FPrimitiveDrawInterface* PDI = nullptr) const;
+		CHAOSCLOTH_API void DrawOpenEdges(FPrimitiveDrawInterface* PDI = nullptr) const;
+		CHAOSCLOTH_API void DrawPointNormals(FPrimitiveDrawInterface* PDI = nullptr) const;
+		CHAOSCLOTH_API void DrawPointVelocities(FPrimitiveDrawInterface* PDI = nullptr) const;
+		CHAOSCLOTH_API void DrawCollision(FPrimitiveDrawInterface* PDI = nullptr) const;
+		CHAOSCLOTH_API void DrawBackstops(FPrimitiveDrawInterface* PDI = nullptr) const;
+		CHAOSCLOTH_API void DrawBackstopDistances(FPrimitiveDrawInterface* PDI = nullptr) const;
+		CHAOSCLOTH_API void DrawMaxDistances(FPrimitiveDrawInterface* PDI = nullptr) const;
+		CHAOSCLOTH_API void DrawAnimDrive(FPrimitiveDrawInterface* PDI = nullptr) const;
+		CHAOSCLOTH_API void DrawEdgeConstraint(FPrimitiveDrawInterface* PDI = nullptr) const;
+		CHAOSCLOTH_API void DrawBendingConstraint(FPrimitiveDrawInterface* PDI = nullptr) const;
+		CHAOSCLOTH_API void DrawLongRangeConstraint(FPrimitiveDrawInterface* PDI = nullptr) const;
+		CHAOSCLOTH_API void DrawWindAndPressureForces(FPrimitiveDrawInterface* PDI = nullptr) const;
+		CHAOSCLOTH_API void DrawLocalSpace(FPrimitiveDrawInterface* PDI = nullptr) const;
+		CHAOSCLOTH_API void DrawSelfCollision(FPrimitiveDrawInterface* PDI = nullptr) const;
+		CHAOSCLOTH_API void DrawSelfIntersection(FPrimitiveDrawInterface* PDI = nullptr) const;
+		CHAOSCLOTH_API void DrawBounds(FPrimitiveDrawInterface* PDI = nullptr) const;
+		CHAOSCLOTH_API void DrawGravity(FPrimitiveDrawInterface* PDI = nullptr) const;
 
 #else  // #if CHAOS_DEBUG_DRAW
 		void SetSolver(const ::Chaos::FClothingSimulationSolver* /*InSolver*/) {}
@@ -80,14 +80,14 @@ namespace Chaos
 
 #if WITH_EDITOR && CHAOS_DEBUG_DRAW
 		// Editor only functions
-		void DrawPhysMeshShaded(FPrimitiveDrawInterface* PDI) const;
-		void DrawParticleIndices(FCanvas* Canvas, const FSceneView* SceneView) const;
-		void DrawElementIndices(FCanvas* Canvas, const FSceneView* SceneView) const;
-		void DrawMaxDistanceValues(FCanvas* Canvas, const FSceneView* SceneView) const;
+		CHAOSCLOTH_API void DrawPhysMeshShaded(FPrimitiveDrawInterface* PDI) const;
+		CHAOSCLOTH_API void DrawParticleIndices(FCanvas* Canvas, const FSceneView* SceneView) const;
+		CHAOSCLOTH_API void DrawElementIndices(FCanvas* Canvas, const FSceneView* SceneView) const;
+		CHAOSCLOTH_API void DrawMaxDistanceValues(FCanvas* Canvas, const FSceneView* SceneView) const;
 
 	protected:
 		// FGCObject interface
-		virtual void AddReferencedObjects(FReferenceCollector& Collector) override;
+		CHAOSCLOTH_API virtual void AddReferencedObjects(FReferenceCollector& Collector) override;
 		virtual FString GetReferencerName() const override { return TEXT("UE::Chaos::Cloth::FVisualization"); }
 		// End of FGCObject interface
 
diff --git a/Engine/Plugins/Experimental/ChaosCloth/Source/Chaos/Public/ChaosCloth/ChaosClothingSimulation.h b/Engine/Plugins/Experimental/ChaosCloth/Source/Chaos/Public/ChaosCloth/ChaosClothingSimulation.h
--- a/Engine/Plugins/Experimental/ChaosCloth/Source/Chaos/Public/ChaosCloth/ChaosClothingSimulation.h
+++ b/Engine/Plugins/Experimental/ChaosCloth/Source/Chaos/Public/ChaosCloth/ChaosClothingSimulation.h
@@ -84,35 +84,35 @@ namespace Chaos
 
 #if WITH_EDITOR
 		// Editor only debug draw function
-		CHAOSCLOTH_API void DebugDrawPhysMeshShaded(FPrimitiveDrawInterface* PDI) const { Visualization.DrawPhysMeshShaded(PDI); }
-		CHAOSCLOTH_API void DebugDrawParticleIndices(FCanvas* Canvas, const FSceneView* SceneView) const { Visualization.DrawParticleIndices(Canvas, SceneView); }
-		CHAOSCLOTH_API void DebugDrawElementIndices(FCanvas* Canvas, const FSceneView* SceneView) const { Visualization.DrawElementIndices(Canvas, SceneView); }
-		CHAOSCLOTH_API void DebugDrawMaxDistanceValues(FCanvas* Canvas, const FSceneView* SceneView) const { Visualization.DrawMaxDistanceValues(Canvas, SceneView); }
+		void DebugDrawPhysMeshShaded(FPrimitiveDrawInterface* PDI) const { Visualization.DrawPhysMeshShaded(PDI); }
+		void DebugDrawParticleIndices(FCanvas* Canvas, const FSceneView* SceneView) const { Visualization.DrawParticleIndices(Canvas, SceneView); }
+		void DebugDrawElementIndices(FCanvas* Canvas, const FSceneView* SceneView) const { Visualization.DrawElementIndices(Canvas, SceneView); }
+		void DebugDrawMaxDistanceValues(FCanvas* Canvas, const FSceneView* SceneView) const { Visualization.DrawMaxDistanceValues(Canvas, SceneView); }
 #endif  // #if WITH_EDITOR
 
 #if CHAOS_DEBUG_DRAW
 		// Editor & runtime debug draw functions
-		CHAOSCLOTH_API void DebugDrawPhysMeshWired(FPrimitiveDrawInterface* PDI = nullptr) const { Visualization.DrawPhysMeshWired(PDI); }
-		CHAOSCLOTH_API void DebugDrawAnimMeshWired(FPrimitiveDrawInterface* PDI = nullptr) const { Visualization.DrawAnimMeshWired(PDI); }
-		CHAOSCLOTH_API void DebugDrawAnimNormals(FPrimitiveDrawInterface* PDI = nullptr) const { Visualization.DrawAnimNormals(PDI); }
-		CHAOSCLOTH_API void DebugDrawPointNormals(FPrimitiveDrawInterface* PDI = nullptr) const { Visualization.DrawPointNormals(PDI); }
-		CHAOSCLOTH_API void DebugDrawPointVelocities(FPrimitiveDrawInterface* PDI = nullptr) const { Visualization.DrawPointVelocities(PDI); }
-		CHAOSCLOTH_API void DebugDrawCollision(FPrimitiveDrawInterface* PDI = nullptr) const { Visualization.DrawCollision(PDI); }
-		CHAOSCLOTH_API void DebugDrawBackstops(FPrimitiveDrawInterface* PDI = nullptr) const { Visualization.DrawBackstops(PDI); }
-		CHAOSCLOTH_API void DebugDrawBackstopDistances(FPrimitiveDrawInterface* PDI = nullptr) const { Visualization.DrawBackstopDistances(PDI); }
-		CHAOSCLOTH_API void DebugDrawMaxDistances(FPrimitiveDrawInterface* PDI = nullptr) const { Visualization.DrawMaxDistances(PDI); }
-		CHAOSCLOTH_API void DebugDrawAnimDrive(FPrimitiveDrawInterface* PDI = nullptr) const { Visualization.DrawAnimDrive(PDI); }
-		CHAOSCLOTH_API void DebugDrawEdgeConstraint(FPrimitiveDrawInterface* PDI = nullptr) const { Visualization.DrawEdgeConstraint(PDI); }
-		CHAOSCLOTH_API void DebugDrawBendingConstraint(FPrimitiveDrawInterface* PDI = nullptr) const { Visualization.DrawBendingConstraint(PDI); }
-		CHAOSCLOTH_API void DebugDrawLongRangeConstraint(FPrimitiveDrawInterface* PDI = nullptr) const { Visualization.DrawLongRangeConstraint(PDI); }
-		CHAOSCLOTH_API void DebugDrawWindAndPressureForces(FPrimitiveDrawInterface* PDI = nullptr) const { Visualization.DrawWindAndPressureForces(PDI); }
+		void DebugDrawPhysMeshWired(FPrimitiveDrawInterface* PDI = nullptr) const { Visualization.DrawPhysMeshWired(PDI); }
+		void DebugDrawAnimMeshWired(FPrimitiveDrawInterface* PDI = nullptr) const { Visualization.DrawAnimMeshWired(PDI); }
+		void DebugDrawAnimNormals(FPrimitiveDrawInterface* PDI = nullptr) const { Visualization.DrawAnimNormals(PDI); }
+		void DebugDrawPointNormals(FPrimitiveDrawInterface* PDI = nullptr) const { Visualization.DrawPointNormals(PDI); }
+		void DebugDrawPointVelocities(FPrimitiveDrawInterface* PDI = nullptr) const { Visualization.DrawPointVelocities(PDI); }
+		void DebugDrawCollision(FPrimitiveDrawInterface* PDI = nullptr) const { Visualization.DrawCollision(PDI); }
+		void DebugDrawBackstops(FPrimitiveDrawInterface* PDI = nullptr) const { Visualization.DrawBackstops(PDI); }
+		void DebugDrawBackstopDistances(FPrimitiveDrawInterface* PDI = nullptr) const { Visualization.DrawBackstopDistances(PDI); }
+		void DebugDrawMaxDistances(FPrimitiveDrawInterface* PDI = nullptr) const { Visualization.DrawMaxDistances(PDI); }
+		void DebugDrawAnimDrive(FPrimitiveDrawInterface* PDI = nullptr) const { Visualization.DrawAnimDrive(PDI); }
+		void DebugDrawEdgeConstraint(FPrimitiveDrawInterface* PDI = nullptr) const { Visualization.DrawEdgeConstraint(PDI); }
+		void DebugDrawBendingConstraint(FPrimitiveDrawInterface* PDI = nullptr) const { Visualization.DrawBendingConstraint(PDI); }
+		void DebugDrawLongRangeConstraint(FPrimitiveDrawInterface* PDI = nullptr) const { Visualization.DrawLongRangeConstraint(PDI); }
+		void DebugDrawWindAndPressureForces(FPrimitiveDrawInterface* PDI = nullptr) const { Visualization.DrawWindAndPressureForces(PDI); }
 		UE_DEPRECATED(5.1, "DebugDrawWindForces has been renamed DebugDrawWindAndPressureForces.")
-		CHAOSCLOTH_API void DebugDrawWindForces(FPrimitiveDrawInterface* PDI = nullptr) const { Visualization.DrawWindAndPressureForces(PDI); }
-		CHAOSCLOTH_API void DebugDrawLocalSpace(FPrimitiveDrawInterface* PDI = nullptr) const { Visualization.DrawLocalSpace(PDI); }
-		CHAOSCLOTH_API void DebugDrawSelfCollision(FPrimitiveDrawInterface* PDI = nullptr) const { Visualization.DrawSelfCollision(PDI); }
-		CHAOSCLOTH_API void DebugDrawSelfIntersection(FPrimitiveDrawInterface* PDI = nullptr) const { Visualization.DrawSelfIntersection(PDI); }
-		CHAOSCLOTH_API void DebugDrawBounds(FPrimitiveDrawInterface* PDI = nullptr) const { Visualization.DrawBounds(PDI); }
-		CHAOSCLOTH_API void DebugDrawGravity(FPrimitiveDrawInterface* PDI = nullptr) const { Visualization.DrawGravity(PDI); }
+		void DebugDrawWindForces(FPrimitiveDrawInterface* PDI = nullptr) const { Visualization.DrawWindAndPressureForces(PDI); }
+		void DebugDrawLocalSpace(FPrimitiveDrawInterface* PDI = nullptr) const { Visualization.DrawLocalSpace(PDI); }
+		void DebugDrawSelfCollision(FPrimitiveDrawInterface* PDI = nullptr) const { Visualization.DrawSelfCollision(PDI); }
+		void DebugDrawSelfIntersection(FPrimitiveDrawInterface* PDI = nullptr) const { Visualization.DrawSelfIntersection(PDI); }
+		void DebugDrawBounds(FPrimitiveDrawInterface* PDI = nullptr) const { Visualization.DrawBounds(PDI); }
+		void DebugDrawGravity(FPrimitiveDrawInterface* PDI = nullptr) const { Visualization.DrawGravity(PDI); }
 #endif  // #if CHAOS_DEBUG_DRAW
 
 	private:
diff --git a/Engine/Plugins/Experimental/ChaosCloth/Source/Chaos/Public/ChaosCloth/ChaosClothingSimulationCloth.h b/Engine/Plugins/Experimental/ChaosCloth/Source/Chaos/Public/ChaosCloth/ChaosClothingSimulationCloth.h
--- a/Engine/Plugins/Experimental/ChaosCloth/Source/Chaos/Public/ChaosCloth/ChaosClothingSimulationCloth.h
+++ b/Engine/Plugins/Experimental/ChaosCloth/Source/Chaos/Public/ChaosCloth/ChaosClothingSimulationCloth.h
@@ -18,7 +18,7 @@ namespace Chaos
 	class FClothingSimulationConfig;
 
 	// Cloth simulation node
-	class CHAOSCLOTH_API FClothingSimulationCloth final
+	class FClothingSimulationCloth final
 	{
 	public:
 		enum EMassMode
@@ -31,7 +31,7 @@ namespace Chaos
 		UE_DEPRECATED(5.3, "ETetherMode has been replaced with bUseGeodesicTethers.")
 PRAGMA_DISABLE_DEPRECATION_WARNINGS
 		typedef FClothConstraints::ETetherMode ETetherMode;
-PRAGMA_ENABLE_DEPRECATION_WARNINGS
+CHAOSCLOTH_API PRAGMA_ENABLE_DEPRECATION_WARNINGS
 
 		FClothingSimulationCloth(
 			FClothingSimulationConfig* InConfig,
@@ -95,7 +95,7 @@ PRAGMA_ENABLE_DEPRECATION_WARNINGS
 			int32 InLODIndexOverride,
 			const TVec2<FRealSingle>& EdgeDampingRatio = TVec2<FRealSingle>(0.f),
 			const TVec2<FRealSingle>& BendDampingRatio = TVec2<FRealSingle>(0.f));
-		~FClothingSimulationCloth();
+		CHAOSCLOTH_API ~FClothingSimulationCloth();
 
 		FClothingSimulationCloth(const FClothingSimulationCloth&) = delete;
 		FClothingSimulationCloth(FClothingSimulationCloth&&) = delete;
@@ -111,17 +111,17 @@ PRAGMA_ENABLE_DEPRECATION_WARNINGS
 		// ---- Animatable property setters ----
 		void SetMaxDistancesMultiplier(FRealSingle InMaxDistancesMultiplier) { MaxDistancesMultiplier = InMaxDistancesMultiplier; }
 
-		void SetMaterialProperties(const TVec2<FRealSingle>& InEdgeStiffness, const TVec2<FRealSingle>& InBendingStiffness, const TVec2<FRealSingle>& InAreaStiffness);
-		void SetLongRangeAttachmentProperties(const TVec2<FRealSingle>& InTetherStiffness, const TVec2<FRealSingle>& InTetherScale);
-		void SetCollisionProperties(FRealSingle InCollisionThickness, FRealSingle InFrictionCoefficient, bool bInUseCCD, FRealSingle InSelfCollisionThickness);
-		void SetBackstopProperties(bool bInEnableBackstop);
-		void SetDampingProperties(FRealSingle InDampingCoefficient, FRealSingle InLocalDampingCoefficient = 0.f);
-		void SetAerodynamicsProperties(const TVec2<FRealSingle>& InDrag, const TVec2<FRealSingle>& InLift, FRealSingle InAirDensity, const FVec3& InWindVelocity);  // AirDensity is here in kg/cm^3 for legacy reason (kg/m^3 in UI)
-		void SetPressureProperties(const TVec2<FRealSingle>& InPressure);
-		void SetGravityProperties(FRealSingle InGravityScale, bool bInUseGravityOverride, const FVec3& InGravityOverride);
-		void SetAnimDriveProperties(const TVec2<FRealSingle>& InAnimDriveStiffness, const TVec2<FRealSingle>& InAnimDriveDamping);
-		void GetAnimDriveProperties(TVec2<FRealSingle>& OutAnimDriveStiffness, TVec2<FRealSingle>& OutAnimDriveDamping);
-		void SetVelocityScaleProperties(const FVec3& InLinearVelocityScale, FRealSingle InAngularVelocityScale, FRealSingle InFictitiousAngularScale);
+		CHAOSCLOTH_API void SetMaterialProperties(const TVec2<FRealSingle>& InEdgeStiffness, const TVec2<FRealSingle>& InBendingStiffness, const TVec2<FRealSingle>& InAreaStiffness);
+		CHAOSCLOTH_API void SetLongRangeAttachmentProperties(const TVec2<FRealSingle>& InTetherStiffness, const TVec2<FRealSingle>& InTetherScale);
+		CHAOSCLOTH_API void SetCollisionProperties(FRealSingle InCollisionThickness, FRealSingle InFrictionCoefficient, bool bInUseCCD, FRealSingle InSelfCollisionThickness);
+		CHAOSCLOTH_API void SetBackstopProperties(bool bInEnableBackstop);
+		CHAOSCLOTH_API void SetDampingProperties(FRealSingle InDampingCoefficient, FRealSingle InLocalDampingCoefficient = 0.f);
+		CHAOSCLOTH_API void SetAerodynamicsProperties(const TVec2<FRealSingle>& InDrag, const TVec2<FRealSingle>& InLift, FRealSingle InAirDensity, const FVec3& InWindVelocity);  // AirDensity is here in kg/cm^3 for legacy reason (kg/m^3 in UI)
+		CHAOSCLOTH_API void SetPressureProperties(const TVec2<FRealSingle>& InPressure);
+		CHAOSCLOTH_API void SetGravityProperties(FRealSingle InGravityScale, bool bInUseGravityOverride, const FVec3& InGravityOverride);
+		CHAOSCLOTH_API void SetAnimDriveProperties(const TVec2<FRealSingle>& InAnimDriveStiffness, const TVec2<FRealSingle>& InAnimDriveDamping);
+		CHAOSCLOTH_API void GetAnimDriveProperties(TVec2<FRealSingle>& OutAnimDriveStiffness, TVec2<FRealSingle>& OutAnimDriveDamping);
+		CHAOSCLOTH_API void SetVelocityScaleProperties(const FVec3& InLinearVelocityScale, FRealSingle InAngularVelocityScale, FRealSingle InFictitiousAngularScale);
 
 		void Reset() { bNeedsReset = true; }
 		void Teleport() { bNeedsTeleport = true; }
@@ -130,70 +130,70 @@ PRAGMA_ENABLE_DEPRECATION_WARNINGS
 		// ---- Node property getters/setters
 		PRAGMA_DISABLE_DEPRECATION_WARNINGS  // TODO: CHAOS_IS_CLOTHINGSIMULATIONMESH_ABSTRACT
 		FClothingSimulationMesh* GetMesh() const { return Mesh; }
-		void SetMesh(FClothingSimulationMesh* InMesh);
+		CHAOSCLOTH_API void SetMesh(FClothingSimulationMesh* InMesh);
 		PRAGMA_ENABLE_DEPRECATION_WARNINGS
 
 		FClothingSimulationConfig* GetConfig() const { return Config; }
-		void SetConfig(FClothingSimulationConfig* InConfig);
+		CHAOSCLOTH_API void SetConfig(FClothingSimulationConfig* InConfig);
 
 		const TArray<FClothingSimulationCollider*>& GetColliders() const { return Colliders; }
-		void SetColliders(TArray<FClothingSimulationCollider*>&& InColliders);
-		void AddCollider(FClothingSimulationCollider* InCollider);
-		void RemoveCollider(FClothingSimulationCollider* InCollider);
-		void RemoveColliders();
+		CHAOSCLOTH_API void SetColliders(TArray<FClothingSimulationCollider*>&& InColliders);
+		CHAOSCLOTH_API void AddCollider(FClothingSimulationCollider* InCollider);
+		CHAOSCLOTH_API void RemoveCollider(FClothingSimulationCollider* InCollider);
+		CHAOSCLOTH_API void RemoveColliders();
 		// ---- End of the Node property getters/setters
 
 		// ---- Debugging/visualization functions
 		// Return the solver's input positions for this cloth source current LOD, not thread safe, call must be done right after the solver update.
-		TConstArrayView<Softs::FSolverVec3> GetAnimationPositions(const FClothingSimulationSolver* Solver) const;
+		CHAOSCLOTH_API TConstArrayView<Softs::FSolverVec3> GetAnimationPositions(const FClothingSimulationSolver* Solver) const;
 		// Return the solver's input normals for this cloth source current LOD, not thread safe, call must be done right after the solver update.
-		TConstArrayView<Softs::FSolverVec3> GetAnimationNormals(const FClothingSimulationSolver* Solver) const;
+		CHAOSCLOTH_API TConstArrayView<Softs::FSolverVec3> GetAnimationNormals(const FClothingSimulationSolver* Solver) const;
 		// Return the solver's positions for this cloth current LOD, not thread safe, call must be done right after the solver update.
-		TConstArrayView<Softs::FSolverVec3> GetParticlePositions(const FClothingSimulationSolver* Solver) const;
+		CHAOSCLOTH_API TConstArrayView<Softs::FSolverVec3> GetParticlePositions(const FClothingSimulationSolver* Solver) const;
 		// Return the solver's velocities for this cloth current LOD, not thread safe, call must be done right after the solver update.
-		TConstArrayView<Softs::FSolverVec3> GetParticleVelocities(const FClothingSimulationSolver* Solver) const;
+		CHAOSCLOTH_API TConstArrayView<Softs::FSolverVec3> GetParticleVelocities(const FClothingSimulationSolver* Solver) const;
 		// Return the solver's normals for this cloth current LOD, not thread safe, call must be done right after the solver update.
-		TConstArrayView<Softs::FSolverVec3> GetParticleNormals(const FClothingSimulationSolver* Solver) const;
+		CHAOSCLOTH_API TConstArrayView<Softs::FSolverVec3> GetParticleNormals(const FClothingSimulationSolver* Solver) const;
 		// Return the solver's inverse masses for this cloth current LOD, not thread safe, call must be done right after the solver update.
-		TConstArrayView<Softs::FSolverReal> GetParticleInvMasses(const FClothingSimulationSolver* Solver) const;
+		CHAOSCLOTH_API TConstArrayView<Softs::FSolverReal> GetParticleInvMasses(const FClothingSimulationSolver* Solver) const;
 		// Return the current gravity as applied by the solver using the various overrides, not thread safe, call must be done right after the solver update.
-		TVec3<FRealSingle> GetGravity(const FClothingSimulationSolver* Solver) const;
+		CHAOSCLOTH_API TVec3<FRealSingle> GetGravity(const FClothingSimulationSolver* Solver) const;
 		// Return the current bounding box based on a given solver, not thread safe, call must be done right after the solver update.
-		FAABB3 CalculateBoundingBox(const FClothingSimulationSolver* Solver) const;
+		CHAOSCLOTH_API FAABB3 CalculateBoundingBox(const FClothingSimulationSolver* Solver) const;
 		// Return the current LOD offset in the solver's particle array, or INDEX_NONE if no LOD is currently selected.
-		int32 GetOffset(const FClothingSimulationSolver* Solver) const;
+		CHAOSCLOTH_API int32 GetOffset(const FClothingSimulationSolver* Solver) const;
 		// Return the current LOD num particles, or 0 if no LOD is currently selected.
-		int32 GetNumParticles(const FClothingSimulationSolver* Solver) const;
+		CHAOSCLOTH_API int32 GetNumParticles(const FClothingSimulationSolver* Solver) const;
 		// Return the current LOD mesh.
-		const FTriangleMesh& GetTriangleMesh(const FClothingSimulationSolver* Solver) const;
+		CHAOSCLOTH_API const FTriangleMesh& GetTriangleMesh(const FClothingSimulationSolver* Solver) const;
 		// Return the weight map of the specified name if available on the current LOD, or an empty array view otherwise.
-		TConstArrayView<FRealSingle> GetWeightMapByName(const FClothingSimulationSolver* Solver, const FString& Name) const;
+		CHAOSCLOTH_API TConstArrayView<FRealSingle> GetWeightMapByName(const FClothingSimulationSolver* Solver, const FString& Name) const;
 		// Return the weight map of the specified property name if it exists and is available on the current LOD, or an empty array view otherwise.
-		TConstArrayView<FRealSingle> GetWeightMapByProperty(const FClothingSimulationSolver* Solver, const FString& Property) const;
+		CHAOSCLOTH_API TConstArrayView<FRealSingle> GetWeightMapByProperty(const FClothingSimulationSolver* Solver, const FString& Property) const;
 		UE_DEPRECATED(5.3, "Returns an empty array from 5.3. Update your code with GetWeightMapByName and GetWeightMapByProperty to return the current LOD weight map instead.")
-		const TArray<TConstArrayView<FRealSingle>>& GetWeightMaps(const FClothingSimulationSolver* Solver) const;
+		CHAOSCLOTH_API const TArray<TConstArrayView<FRealSingle>>& GetWeightMaps(const FClothingSimulationSolver* Solver) const;
 		// Return the current LOD tethers.
-		const TArray<TConstArrayView<TTuple<int32, int32, float>>>& GetTethers(const FClothingSimulationSolver* Solver) const;
+		CHAOSCLOTH_API const TArray<TConstArrayView<TTuple<int32, int32, float>>>& GetTethers(const FClothingSimulationSolver* Solver) const;
 		// Return the reference bone index for this cloth.
-		int32 GetReferenceBoneIndex() const;
+		CHAOSCLOTH_API int32 GetReferenceBoneIndex() const;
 		// Return the local reference space transform for this cloth.
 		const FRigidTransform3& GetReferenceSpaceTransform() const { return ReferenceSpaceTransform;  }
 		// ---- End of the debugging/visualization functions
 
 		// ---- Solver interface ----
-		void Add(FClothingSimulationSolver* Solver);
-		void Remove(FClothingSimulationSolver* Solver);
+		CHAOSCLOTH_API void Add(FClothingSimulationSolver* Solver);
+		CHAOSCLOTH_API void Remove(FClothingSimulationSolver* Solver);
 
-		void PreUpdate(FClothingSimulationSolver* Solver);
-		void Update(FClothingSimulationSolver* Solver);
-		void PostUpdate(FClothingSimulationSolver* Solver);
+		CHAOSCLOTH_API void PreUpdate(FClothingSimulationSolver* Solver);
+		CHAOSCLOTH_API void Update(FClothingSimulationSolver* Solver);
+		CHAOSCLOTH_API void PostUpdate(FClothingSimulationSolver* Solver);
 
-		void UpdateFromCache(const FClothingSimulationCacheData& CacheData);
+		CHAOSCLOTH_API void UpdateFromCache(const FClothingSimulationCacheData& CacheData);
 		// ---- End of the Solver interface ----
 
 	private:
-		int32 GetNumParticles(int32 InLODIndex) const;
-		int32 GetOffset(const FClothingSimulationSolver* Solver, int32 InLODIndex) const;
+		CHAOSCLOTH_API int32 GetNumParticles(int32 InLODIndex) const;
+		CHAOSCLOTH_API int32 GetOffset(const FClothingSimulationSolver* Solver, int32 InLODIndex) const;
 
 	private:
 		struct FLODData;
diff --git a/Engine/Plugins/Experimental/ChaosCloth/Source/Chaos/Public/ChaosCloth/ChaosClothingSimulationCollider.h b/Engine/Plugins/Experimental/ChaosCloth/Source/Chaos/Public/ChaosCloth/ChaosClothingSimulationCollider.h
--- a/Engine/Plugins/Experimental/ChaosCloth/Source/Chaos/Public/ChaosCloth/ChaosClothingSimulationCollider.h
+++ b/Engine/Plugins/Experimental/ChaosCloth/Source/Chaos/Public/ChaosCloth/ChaosClothingSimulationCollider.h
@@ -21,7 +21,7 @@ namespace Chaos
 	template<typename T> class TWeightedLatticeImplicitObject;
 
 	// Collider simulation node
-	class CHAOSCLOTH_API FClothingSimulationCollider final
+	class FClothingSimulationCollider final
 	{
 	public:
 		enum class ECollisionDataType : int32
@@ -31,26 +31,26 @@ namespace Chaos
 			LODs,  // LODIndex based start slot for LODs collisions
 		};
 
-		FClothingSimulationCollider(const UPhysicsAsset* InPhysicsAsset, const FReferenceSkeleton* InReferenceSkeleton);
+		CHAOSCLOTH_API FClothingSimulationCollider(const UPhysicsAsset* InPhysicsAsset, const FReferenceSkeleton* InReferenceSkeleton);
 
 		UE_DEPRECATED(5.2, "Use FClothingSimulationCollider(const UPhysicsAsset* InPhysicsAsset, const FReferenceSkeleton* InReferenceSkeleton) instead.")
-		FClothingSimulationCollider(
+		CHAOSCLOTH_API FClothingSimulationCollider(
 			const UClothingAssetCommon* InAsset,  // Cloth asset for collision data, can be nullptr
 			const USkeletalMeshComponent* InSkeletalMeshComponent,  // For asset LODs management, can be nullptr
 			bool bInUseLODIndexOverride,
 			int32 InLODIndexOverride);
 
-		~FClothingSimulationCollider();
+		CHAOSCLOTH_API ~FClothingSimulationCollider();
 
 		FClothingSimulationCollider(const FClothingSimulationCollider&) = delete;
 		FClothingSimulationCollider(FClothingSimulationCollider&&) = delete;
 		FClothingSimulationCollider& operator=(const FClothingSimulationCollider&) = delete;
 		FClothingSimulationCollider& operator=(FClothingSimulationCollider&&) = delete;
 
-		int32 GetNumGeometries() const;
+		CHAOSCLOTH_API int32 GetNumGeometries() const;
 
 		// Return source (untransformed) collision data for LODless, external and active LODs.
-		FClothCollisionData GetCollisionData(const FClothingSimulationSolver* Solver, const FClothingSimulationCloth* Cloth) const;
+		CHAOSCLOTH_API FClothCollisionData GetCollisionData(const FClothingSimulationSolver* Solver, const FClothingSimulationCloth* Cloth) const;
 
 		// ---- Animatable property setters ----
 		// Set external collision data, will only get updated when used as a Solver Collider TODO: Subclass collider?
@@ -58,29 +58,29 @@ namespace Chaos
 		// ---- End of the animatable property setters ----
 
 		// ---- Cloth interface ----
-		void Add(FClothingSimulationSolver* Solver, FClothingSimulationCloth* Cloth);
-		void Remove(FClothingSimulationSolver* Solver, FClothingSimulationCloth* Cloth);
+		CHAOSCLOTH_API void Add(FClothingSimulationSolver* Solver, FClothingSimulationCloth* Cloth);
+		CHAOSCLOTH_API void Remove(FClothingSimulationSolver* Solver, FClothingSimulationCloth* Cloth);
 
-		void PreUpdate(FClothingSimulationSolver* Solver, FClothingSimulationCloth* Cloth);
-		void Update(FClothingSimulationSolver* Solver, FClothingSimulationCloth* Cloth);
-		void ResetStartPose(FClothingSimulationSolver* Solver, FClothingSimulationCloth* Cloth);
+		CHAOSCLOTH_API void PreUpdate(FClothingSimulationSolver* Solver, FClothingSimulationCloth* Cloth);
+		CHAOSCLOTH_API void Update(FClothingSimulationSolver* Solver, FClothingSimulationCloth* Cloth);
+		CHAOSCLOTH_API void ResetStartPose(FClothingSimulationSolver* Solver, FClothingSimulationCloth* Cloth);
 		// ---- End of the Cloth interface ----
 
 		// ---- Debugging and visualization functions ----
 		// Return current active LOD collision particles translations, not thread safe, to use after solver update.
-		TConstArrayView<Softs::FSolverVec3> GetCollisionTranslations(const FClothingSimulationSolver* Solver, const FClothingSimulationCloth* Cloth, ECollisionDataType CollisionDataType) const;
+		CHAOSCLOTH_API TConstArrayView<Softs::FSolverVec3> GetCollisionTranslations(const FClothingSimulationSolver* Solver, const FClothingSimulationCloth* Cloth, ECollisionDataType CollisionDataType) const;
 
 		// Return current active LOD collision particles rotations, not thread safe, to use after solver update.
-		TConstArrayView<Softs::FSolverRotation3> GetCollisionRotations(const FClothingSimulationSolver* Solver, const FClothingSimulationCloth* Cloth, ECollisionDataType CollisionDataType) const;
+		CHAOSCLOTH_API TConstArrayView<Softs::FSolverRotation3> GetCollisionRotations(const FClothingSimulationSolver* Solver, const FClothingSimulationCloth* Cloth, ECollisionDataType CollisionDataType) const;
 
 		// Return current active LOD previous frame collision particles transforms, not thread safe, to use after solver update.
-		TConstArrayView<Softs::FSolverRigidTransform3> GetOldCollisionTransforms(const FClothingSimulationSolver* Solver, const FClothingSimulationCloth* Cloth, ECollisionDataType CollisionDataType) const;
+		CHAOSCLOTH_API TConstArrayView<Softs::FSolverRigidTransform3> GetOldCollisionTransforms(const FClothingSimulationSolver* Solver, const FClothingSimulationCloth* Cloth, ECollisionDataType CollisionDataType) const;
 
 		// Return current active LOD collision geometries, not thread safe, to use after solver update.
-		TConstArrayView<TUniquePtr<FImplicitObject>> GetCollisionGeometries(const FClothingSimulationSolver* Solver, const FClothingSimulationCloth* Cloth, ECollisionDataType CollisionDataType) const;
+		CHAOSCLOTH_API TConstArrayView<TUniquePtr<FImplicitObject>> GetCollisionGeometries(const FClothingSimulationSolver* Solver, const FClothingSimulationCloth* Cloth, ECollisionDataType CollisionDataType) const;
 
 		// Return whether the collision has been hit by a particle during CCD.
-		TConstArrayView<bool> GetCollisionStatus(const FClothingSimulationSolver* Solver, const FClothingSimulationCloth* Cloth, ECollisionDataType CollisionDataType) const;
+		CHAOSCLOTH_API TConstArrayView<bool> GetCollisionStatus(const FClothingSimulationSolver* Solver, const FClothingSimulationCloth* Cloth, ECollisionDataType CollisionDataType) const;
 		// ---- End of the debugging and visualization functions ----
 
 	private:
@@ -98,15 +98,15 @@ namespace Chaos
 			TArray<int32> MappedSkinnedBones;
 		};
 
-		void ExtractPhysicsAssetCollision(FClothCollisionData& ClothCollisionData, TArray<FLevelSetCollisionData>& LevelSetCollisions, TArray<FSkinnedLevelSetCollisionData>& SkinnedLevelSetCollisions, TArray<int32>& UsedBoneIndices);
+		CHAOSCLOTH_API void ExtractPhysicsAssetCollision(FClothCollisionData& ClothCollisionData, TArray<FLevelSetCollisionData>& LevelSetCollisions, TArray<FSkinnedLevelSetCollisionData>& SkinnedLevelSetCollisions, TArray<int32>& UsedBoneIndices);
 
-		int32 GetNumGeometries(int32 InSlotIndex) const;
+		CHAOSCLOTH_API int32 GetNumGeometries(int32 InSlotIndex) const;
 
 		// Return the collision particle offset for the specified slot being LODLess, external, or any of the LODs collision.
-		int32 GetOffset(const FClothingSimulationSolver* Solver, const FClothingSimulationCloth* Cloth, int32 InSlotIndex) const;
+		CHAOSCLOTH_API int32 GetOffset(const FClothingSimulationSolver* Solver, const FClothingSimulationCloth* Cloth, int32 InSlotIndex) const;
 
 		// Return the collision particle offset and number of geometries for the specified type if valid. If ECollisionDataType::LODs is asked, then the offset returned is for the current LOD.
-		bool GetOffsetAndNumGeometries(const FClothingSimulationSolver* Solver, const FClothingSimulationCloth* Cloth, ECollisionDataType CollisionDataType, int32& OutOffset, int32& OutNumGeometries) const;
+		CHAOSCLOTH_API bool GetOffsetAndNumGeometries(const FClothingSimulationSolver* Solver, const FClothingSimulationCloth* Cloth, ECollisionDataType CollisionDataType, int32& OutOffset, int32& OutNumGeometries) const;
 
 	private:
 		typedef TPair<const FClothingSimulationSolver*, const FClothingSimulationCloth*> FSolverClothPair;
@@ -126,4 +126,4 @@ namespace Chaos
 		// Initial scale
 		FReal Scale = 1.f;
 	};
-} // namespace Chaos
\ No newline at end of file
+} // namespace Chaos
diff --git a/Engine/Plugins/Experimental/ChaosCloth/Source/Chaos/Public/ChaosCloth/ChaosClothingSimulationConfig.h b/Engine/Plugins/Experimental/ChaosCloth/Source/Chaos/Public/ChaosCloth/ChaosClothingSimulationConfig.h
--- a/Engine/Plugins/Experimental/ChaosCloth/Source/Chaos/Public/ChaosCloth/ChaosClothingSimulationConfig.h
+++ b/Engine/Plugins/Experimental/ChaosCloth/Source/Chaos/Public/ChaosCloth/ChaosClothingSimulationConfig.h
@@ -18,13 +18,13 @@ namespace Chaos
 	}
 
 	// Cloth simulation properties
-	class CHAOSCLOTH_API FClothingSimulationConfig final
+	class FClothingSimulationConfig final
 	{
 	public:
-		FClothingSimulationConfig();
-		FClothingSimulationConfig(const TSharedPtr<const FManagedArrayCollection>& InPropertyCollection);
+		CHAOSCLOTH_API FClothingSimulationConfig();
+		CHAOSCLOTH_API FClothingSimulationConfig(const TSharedPtr<const FManagedArrayCollection>& InPropertyCollection);
 
-		~FClothingSimulationConfig();
+		CHAOSCLOTH_API ~FClothingSimulationConfig();
 
 		FClothingSimulationConfig(const FClothingSimulationConfig&) = delete;
 		FClothingSimulationConfig(FClothingSimulationConfig&&) = delete;
@@ -37,16 +37,16 @@ namespace Chaos
 		 * @param ClothSharedConfig The cloth solver shared config UObject.
 		 * @param bUseLegacyConfig Whether to make the config a legacy cloth config, so that the constraints disable themselves with missing masks, ...etc.
 		 */
-		void Initialize(const UChaosClothConfig* ClothConfig, const UChaosClothSharedSimConfig* ClothSharedConfig, bool bUseLegacyConfig = false);
+		CHAOSCLOTH_API void Initialize(const UChaosClothConfig* ClothConfig, const UChaosClothSharedSimConfig* ClothSharedConfig, bool bUseLegacyConfig = false);
 
 		/** Initialize config from a property collection. */
-		void Initialize(const TSharedPtr<const FManagedArrayCollection>& InPropertyCollection);
+		CHAOSCLOTH_API void Initialize(const TSharedPtr<const FManagedArrayCollection>& InPropertyCollection);
 
 		/** Return a property collection facade for reading properties from this configuration. */
-		const Softs::FCollectionPropertyConstFacade& GetProperties() const;
+		CHAOSCLOTH_API const Softs::FCollectionPropertyConstFacade& GetProperties() const;
 
 		/** Return a property collection facade for setting properties for this configuration. */
-		Softs::FCollectionPropertyFacade& GetProperties();
+		CHAOSCLOTH_API Softs::FCollectionPropertyFacade& GetProperties();
 
 		/** Return this configuration's internal property collection. */
 		TSharedPtr<const FManagedArrayCollection> GetPropertyCollection() const { return TSharedPtr<const FManagedArrayCollection>(PropertyCollection); }
diff --git a/Engine/Plugins/Experimental/ChaosCloth/Source/Chaos/Public/ChaosCloth/ChaosClothingSimulationFactory.h b/Engine/Plugins/Experimental/ChaosCloth/Source/Chaos/Public/ChaosCloth/ChaosClothingSimulationFactory.h
--- a/Engine/Plugins/Experimental/ChaosCloth/Source/Chaos/Public/ChaosCloth/ChaosClothingSimulationFactory.h
+++ b/Engine/Plugins/Experimental/ChaosCloth/Source/Chaos/Public/ChaosCloth/ChaosClothingSimulationFactory.h
@@ -6,18 +6,18 @@
 
 #include "ChaosClothingSimulationFactory.generated.h"
 
-UCLASS()
-class CHAOSCLOTH_API UChaosClothingSimulationFactory final : public UClothingSimulationFactory
+UCLASS(MinimalAPI)
+class UChaosClothingSimulationFactory final : public UClothingSimulationFactory
 {
 	GENERATED_BODY()
 public:
-	virtual IClothingSimulation* CreateSimulation() override;
-	virtual void DestroySimulation(IClothingSimulation* InSimulation) override;
-	virtual bool SupportsAsset(UClothingAssetBase* InAsset) override;
+	CHAOSCLOTH_API virtual IClothingSimulation* CreateSimulation() override;
+	CHAOSCLOTH_API virtual void DestroySimulation(IClothingSimulation* InSimulation) override;
+	CHAOSCLOTH_API virtual bool SupportsAsset(UClothingAssetBase* InAsset) override;
 
-	virtual bool SupportsRuntimeInteraction() override;
-	virtual UClothingSimulationInteractor* CreateInteractor() override;
+	CHAOSCLOTH_API virtual bool SupportsRuntimeInteraction() override;
+	CHAOSCLOTH_API virtual UClothingSimulationInteractor* CreateInteractor() override;
 
-	virtual TArrayView<const TSubclassOf<UClothConfigBase>> GetClothConfigClasses() const override;
-	const UEnum* GetWeightMapTargetEnum() const override;
+	CHAOSCLOTH_API virtual TArrayView<const TSubclassOf<UClothConfigBase>> GetClothConfigClasses() const override;
+	CHAOSCLOTH_API const UEnum* GetWeightMapTargetEnum() const override;
 };
diff --git a/Engine/Plugins/Experimental/ChaosCloth/Source/Chaos/Public/ChaosCloth/ChaosClothingSimulationInteractor.h b/Engine/Plugins/Experimental/ChaosCloth/Source/Chaos/Public/ChaosCloth/ChaosClothingSimulationInteractor.h
--- a/Engine/Plugins/Experimental/ChaosCloth/Source/Chaos/Public/ChaosCloth/ChaosClothingSimulationInteractor.h
+++ b/Engine/Plugins/Experimental/ChaosCloth/Source/Chaos/Public/ChaosCloth/ChaosClothingSimulationInteractor.h
@@ -16,84 +16,84 @@ class FClothingSimulationContextCommon;
 DECLARE_DELEGATE_OneParam(FChaosClothingInteractorCommand, Chaos::FClothingSimulationCloth*)
 DECLARE_DELEGATE_TwoParams(FChaosClothingSimulationInteractorCommand, Chaos::FClothingSimulation*, FClothingSimulationContextCommon*)
 
-UCLASS(BlueprintType)
-class CHAOSCLOTH_API UChaosClothingInteractor : public UClothingInteractor
+UCLASS(BlueprintType, MinimalAPI)
+class UChaosClothingInteractor : public UClothingInteractor
 {
 	GENERATED_BODY()
 
 public:
-	virtual void Sync(IClothingSimulation* Simulation);
+	CHAOSCLOTH_API virtual void Sync(IClothingSimulation* Simulation);
 
 	UFUNCTION(BlueprintCallable, Category = ClothingSimulation, Meta = (Keywords = "Chaos Edge Bending Area Stiffness"))
-	void SetMaterialLinear(float EdgeStiffness = 1.f, float BendingStiffness = 1.f, float AreaStiffness = 1.f);
+	CHAOSCLOTH_API void SetMaterialLinear(float EdgeStiffness = 1.f, float BendingStiffness = 1.f, float AreaStiffness = 1.f);
 
 	UFUNCTION(BlueprintCallable, Category = ClothingSimulation, Meta = (Keywords = "Chaos Edge Bending Area Stiffness"))
-	void SetMaterial(FVector2D EdgeStiffness = FVector2D(1.f, 1.f), FVector2D BendingStiffness = FVector2D(1.f, 1.f), FVector2D AreaStiffness = FVector2D(1.f, 1.f));
+	CHAOSCLOTH_API void SetMaterial(FVector2D EdgeStiffness = FVector2D(1.f, 1.f), FVector2D BendingStiffness = FVector2D(1.f, 1.f), FVector2D AreaStiffness = FVector2D(1.f, 1.f));
 
 	UFUNCTION(BlueprintCallable, Category = ClothingSimulation, Meta = (Keywords = "Chaos Tether Stiffness Scale"))
-	void SetLongRangeAttachmentLinear(float TetherStiffness = 1.f, float TetherScale = 1.f);
+	CHAOSCLOTH_API void SetLongRangeAttachmentLinear(float TetherStiffness = 1.f, float TetherScale = 1.f);
 
 	UFUNCTION(BlueprintCallable, Category = ClothingSimulation, Meta = (Keywords = "Chaos Tether Stiffness Scale"))
-	void SetLongRangeAttachment(FVector2D TetherStiffness = FVector2D(1.f, 1.f), FVector2D TetherScale = FVector2D(1.f, 1.f));
+	CHAOSCLOTH_API void SetLongRangeAttachment(FVector2D TetherStiffness = FVector2D(1.f, 1.f), FVector2D TetherScale = FVector2D(1.f, 1.f));
 
 	UFUNCTION(BlueprintCallable, Category = ClothingSimulation, Meta = (Keywords = "Chaos Self Thickness Friction Coefficient"))
-	void SetCollision(float CollisionThickness = 1.f, float FrictionCoefficient = 0.8f, bool bUseCCD = false, float SelfCollisionThickness = 2.f);
+	CHAOSCLOTH_API void SetCollision(float CollisionThickness = 1.f, float FrictionCoefficient = 0.8f, bool bUseCCD = false, float SelfCollisionThickness = 2.f);
 
 	UFUNCTION(BlueprintCallable, Category = ClothingSimulation, Meta = (Keywords = "Chaos Backstop Enable Disable Use"))
-	void SetBackstop(bool bEnabled = true);
+	CHAOSCLOTH_API void SetBackstop(bool bEnabled = true);
 
 	UFUNCTION(BlueprintCallable, Category = ClothingSimulation, Meta = (Keywords = "Chaos Coefficient"))
-	void SetDamping(float DampingCoefficient = 0.01f, float LocalDampingCoefficient = 0.f);
+	CHAOSCLOTH_API void SetDamping(float DampingCoefficient = 0.01f, float LocalDampingCoefficient = 0.f);
 
 	// Deprecated. This function cannot set different Low and High values for the Drag and Lift weight maps. Use SetWind instead.
 	UFUNCTION(BlueprintCallable, Category = ClothingSimulation, Meta = (Keywords = "Chaos Wind Drag Lift Coefficient"))
-	void SetAerodynamics(float DragCoefficient = 0.07f, float LiftCoefficient = 0.035f, FVector WindVelocity = FVector(0.f, 0.f, 0.f));
+	CHAOSCLOTH_API void SetAerodynamics(float DragCoefficient = 0.07f, float LiftCoefficient = 0.035f, FVector WindVelocity = FVector(0.f, 0.f, 0.f));
 
 	UFUNCTION(BlueprintCallable, Category = ClothingSimulation, Meta = (Keywords = "Chaos Wind Drag Lift Coefficient Air Density Velocity"))
-	void SetWind(FVector2D Drag = FVector2D(0.07f, 0.5f), FVector2D Lift = FVector2D(0.07f, 0.5f), float AirDensity = 1.225e-6f, FVector WindVelocity = FVector(0.f, 0.f, 0.f));
+	CHAOSCLOTH_API void SetWind(FVector2D Drag = FVector2D(0.07f, 0.5f), FVector2D Lift = FVector2D(0.07f, 0.5f), float AirDensity = 1.225e-6f, FVector WindVelocity = FVector(0.f, 0.f, 0.f));
 
 	UFUNCTION(BlueprintCallable, Category = ClothingSimulation, Meta = (Keywords = "Chaos Pressure"))
-	void SetPressure(FVector2D Pressure = FVector2D(0.f, 1.f));
+	CHAOSCLOTH_API void SetPressure(FVector2D Pressure = FVector2D(0.f, 1.f));
 
 	UFUNCTION(BlueprintCallable, Category = ClothingSimulation, Meta = (Keywords = "Chaos Scale Override"))
-	void SetGravity(float GravityScale = 1.f, bool bIsGravityOverridden = false, FVector GravityOverride = FVector(0.f, 0.f, -981.f));
+	CHAOSCLOTH_API void SetGravity(float GravityScale = 1.f, bool bIsGravityOverridden = false, FVector GravityOverride = FVector(0.f, 0.f, -981.f));
 
 	UFUNCTION(BlueprintCallable, Category = ClothingSimulation, Meta = (Keywords = "Chaos Stiffness"))
-	void SetAnimDriveLinear(float AnimDriveStiffness = 0.f);
+	CHAOSCLOTH_API void SetAnimDriveLinear(float AnimDriveStiffness = 0.f);
 
 	UFUNCTION(BlueprintCallable, Category = ClothingSimulation, Meta = (Keywords = "Chaos Stiffness Damping"))
-	void SetAnimDrive(FVector2D AnimDriveStiffness = FVector2D(0.f, 1.f), FVector2D AnimDriveDamping = FVector2D(0.f, 1.f));
+	CHAOSCLOTH_API void SetAnimDrive(FVector2D AnimDriveStiffness = FVector2D(0.f, 1.f), FVector2D AnimDriveDamping = FVector2D(0.f, 1.f));
 
 	UFUNCTION(BlueprintCallable, Category = ClothingSimulation, Meta = (Keywords = "Chaos Linear Angular Velocity Scale"))
-	void SetVelocityScale(FVector LinearVelocityScale = FVector(0.75f, 0.75f, 0.75f), float AngularVelocityScale = 0.75f, float FictitiousAngularScale = 1.f);
+	CHAOSCLOTH_API void SetVelocityScale(FVector LinearVelocityScale = FVector(0.75f, 0.75f, 0.75f), float AngularVelocityScale = 0.75f, float FictitiousAngularScale = 1.f);
 
 	UFUNCTION(BlueprintCallable, Category = ClothingSimulation, Meta = (Keywords = "Chaos"))
-	void ResetAndTeleport(bool bReset = false, bool bTeleport = false);
+	CHAOSCLOTH_API void ResetAndTeleport(bool bReset = false, bool bTeleport = false);
 
 private:
 	// Cloth command queue processed when we hit a sync
 	TArray<FChaosClothingInteractorCommand> Commands;
 };
 
-UCLASS(BlueprintType)
-class CHAOSCLOTH_API UChaosClothingSimulationInteractor : public UClothingSimulationInteractor
+UCLASS(BlueprintType, MinimalAPI)
+class UChaosClothingSimulationInteractor : public UClothingSimulationInteractor
 {
 	GENERATED_BODY()
 public:
 	// UClothingSimulationInteractor interface
-	virtual void Sync(IClothingSimulation* InSimulation, IClothingSimulationContext* InContext) override;
+	CHAOSCLOTH_API virtual void Sync(IClothingSimulation* InSimulation, IClothingSimulationContext* InContext) override;
 
-	virtual void PhysicsAssetUpdated() override;
-	virtual void ClothConfigUpdated() override;
-	virtual void SetAnimDriveSpringStiffness(float Stiffness) override;
-	virtual void EnableGravityOverride(const FVector& Gravity) override;
-	virtual void DisableGravityOverride() override;
-	virtual void SetNumIterations(int32 NumIterations = 1) override;
-	virtual void SetMaxNumIterations(int32 MaxNumIterations = 10) override;
-	virtual void SetNumSubsteps(int32 NumSubsteps = 1) override;
+	CHAOSCLOTH_API virtual void PhysicsAssetUpdated() override;
+	CHAOSCLOTH_API virtual void ClothConfigUpdated() override;
+	CHAOSCLOTH_API virtual void SetAnimDriveSpringStiffness(float Stiffness) override;
+	CHAOSCLOTH_API virtual void EnableGravityOverride(const FVector& Gravity) override;
+	CHAOSCLOTH_API virtual void DisableGravityOverride() override;
+	CHAOSCLOTH_API virtual void SetNumIterations(int32 NumIterations = 1) override;
+	CHAOSCLOTH_API virtual void SetMaxNumIterations(int32 MaxNumIterations = 10) override;
+	CHAOSCLOTH_API virtual void SetNumSubsteps(int32 NumSubsteps = 1) override;
 
 protected:
-	virtual UClothingInteractor* CreateClothingInteractor() override;
+	CHAOSCLOTH_API virtual UClothingInteractor* CreateClothingInteractor() override;
 
 private:
 	// Simulation command queue processed when we hit a sync
diff --git a/Engine/Plugins/Experimental/ChaosCloth/Source/Chaos/Public/ChaosCloth/ChaosClothingSimulationSolver.h b/Engine/Plugins/Experimental/ChaosCloth/Source/Chaos/Public/ChaosCloth/ChaosClothingSimulationSolver.h
--- a/Engine/Plugins/Experimental/ChaosCloth/Source/Chaos/Public/ChaosCloth/ChaosClothingSimulationSolver.h
+++ b/Engine/Plugins/Experimental/ChaosCloth/Source/Chaos/Public/ChaosCloth/ChaosClothingSimulationSolver.h
@@ -28,11 +28,11 @@ namespace Chaos
 	}
 
 	// Solver simulation node
-	class CHAOSCLOTH_API FClothingSimulationSolver final : public FPhysicsSolverEvents
+	class FClothingSimulationSolver final : public FPhysicsSolverEvents
 	{
 	public:
-		FClothingSimulationSolver();
-		~FClothingSimulationSolver();
+		CHAOSCLOTH_API FClothingSimulationSolver();
+		CHAOSCLOTH_API ~FClothingSimulationSolver();
 		
 		FClothingSimulationSolver(const FClothingSimulationSolver&) = delete;
 		FClothingSimulationSolver(FClothingSimulationSolver&&) = delete;
@@ -40,7 +40,7 @@ namespace Chaos
 		FClothingSimulationSolver& operator=(FClothingSimulationSolver&&) = delete;
 
 		// ---- Animatable property setters ----
-		void SetLocalSpaceLocation(const FVec3& InLocalSpaceLocation, bool bReset = false);
+		CHAOSCLOTH_API void SetLocalSpaceLocation(const FVec3& InLocalSpaceLocation, bool bReset = false);
 		const FVec3& GetLocalSpaceLocation() const { return LocalSpaceLocation; }
 		void SetLocalSpaceRotation(const FQuat& InLocalSpaceRotation) { LocalSpaceRotation = InLocalSpaceRotation; }
 		const FRotation3& GetLocalSpaceRotation() const { return LocalSpaceRotation; }
@@ -53,28 +53,28 @@ namespace Chaos
 		void SetGravity(const TVec3<FRealSingle>& InGravity) { Gravity = InGravity; }
 		const TVec3<FRealSingle>& GetGravity() const { return Gravity; }
 
-		void SetWindVelocity(const TVec3<FRealSingle>& InWindVelocity, FRealSingle InLegacyWindAdaption = (FRealSingle)0.);
+		CHAOSCLOTH_API void SetWindVelocity(const TVec3<FRealSingle>& InWindVelocity, FRealSingle InLegacyWindAdaption = (FRealSingle)0.);
 		const TVec3<FRealSingle>& GetWindVelocity() const { return WindVelocity; }
 
-		void SetNumIterations(int32 InNumIterations);
-		int32 GetNumIterations() const;
-		void SetMaxNumIterations(int32 InMaxNumIterations);
-		int32 GetMaxNumIterations() const;
-		void SetNumSubsteps(int32 InNumSubsteps);
-		int32 GetNumSubsteps() const;
+		CHAOSCLOTH_API void SetNumIterations(int32 InNumIterations);
+		CHAOSCLOTH_API int32 GetNumIterations() const;
+		CHAOSCLOTH_API void SetMaxNumIterations(int32 InMaxNumIterations);
+		CHAOSCLOTH_API int32 GetMaxNumIterations() const;
+		CHAOSCLOTH_API void SetNumSubsteps(int32 InNumSubsteps);
+		CHAOSCLOTH_API int32 GetNumSubsteps() const;
 
 		void SetEnableSolver(bool InbEnableSolver) { bEnableSolver = InbEnableSolver; }
 		bool GetEnableSolver() const { return bEnableSolver; }
 		// ---- End of the animatable property setters ----
 
 		// ---- Object management functions ----
-		void SetCloths(TArray<FClothingSimulationCloth*>&& InCloths); 
-		void AddCloth(FClothingSimulationCloth* InCloth);
-		void RemoveCloth(FClothingSimulationCloth* InCloth);
-		void RemoveCloths();
+		CHAOSCLOTH_API void SetCloths(TArray<FClothingSimulationCloth*>&& InCloths); 
+		CHAOSCLOTH_API void AddCloth(FClothingSimulationCloth* InCloth);
+		CHAOSCLOTH_API void RemoveCloth(FClothingSimulationCloth* InCloth);
+		CHAOSCLOTH_API void RemoveCloths();
 
-		void RefreshCloth(FClothingSimulationCloth* InCloth);
-		void RefreshCloths();
+		CHAOSCLOTH_API void RefreshCloth(FClothingSimulationCloth* InCloth);
+		CHAOSCLOTH_API void RefreshCloths();
 
 		TConstArrayView<const FClothingSimulationCloth*> GetCloths() const { return Cloths; }
 
@@ -89,36 +89,36 @@ namespace Chaos
 		void SetConfig(FClothingSimulationConfig* InConfig) { Config = InConfig; }
 
 		/** Advance the simulation. */
-		void Update(FSolverReal InDeltaTime);
+		CHAOSCLOTH_API void Update(FSolverReal InDeltaTime);
 
 		/** Return the last delta time used for advancing the simulation. */
 		FSolverReal GetDeltaTime() const { return DeltaTime; }
 
 		/** Set the cached positions onto the particles */
-		void UpdateFromCache(const FClothingSimulationCacheData& CacheData);
+		CHAOSCLOTH_API void UpdateFromCache(const FClothingSimulationCacheData& CacheData);
 		UE_DEPRECATED(5.3, "Use UpdateFromCache(CacheData) instead")
-		void UpdateFromCache(const TArray<FVector>& CachedPositions, const TArray<FVector>& CachedVelocities);
+		CHAOSCLOTH_API void UpdateFromCache(const TArray<FVector>& CachedPositions, const TArray<FVector>& CachedVelocities);
 
 		// Return the actual of number of iterations used by the Evolution solver after the update (different from the number of iterations, depends on frame rate)
-		int32 GetNumUsedIterations() const;
+		CHAOSCLOTH_API int32 GetNumUsedIterations() const;
 
-		FBoxSphereBounds CalculateBounds() const;
+		CHAOSCLOTH_API FBoxSphereBounds CalculateBounds() const;
 		// ---- End of the object management functions ----
 
 		// ---- Cloth interface ----
-		int32 AddParticles(int32 NumParticles, uint32 GroupId);
-		void EnableParticles(int32 Offset, bool bEnable);
+		CHAOSCLOTH_API int32 AddParticles(int32 NumParticles, uint32 GroupId);
+		CHAOSCLOTH_API void EnableParticles(int32 Offset, bool bEnable);
 
-		void ResetStartPose(int32 Offset, int32 NumParticles);
+		CHAOSCLOTH_API void ResetStartPose(int32 Offset, int32 NumParticles);
 
 		// Get the current solver time
 		FSolverReal GetTime() const { return Time; }
-		void SetParticleMassUniform(int32 Offset, FRealSingle UniformMass, FRealSingle MinPerParticleMass, const FTriangleMesh& Mesh, const TFunctionRef<bool(int32)>& KinematicPredicate);
-		void SetParticleMassFromTotalMass(int32 Offset, FRealSingle TotalMass, FRealSingle MinPerParticleMass, const FTriangleMesh& Mesh, const TFunctionRef<bool(int32)>& KinematicPredicate);
-		void SetParticleMassFromDensity(int32 Offset, FRealSingle Density, FRealSingle MinPerParticleMass, const FTriangleMesh& Mesh, const TFunctionRef<bool(int32)>& KinematicPredicate);
+		CHAOSCLOTH_API void SetParticleMassUniform(int32 Offset, FRealSingle UniformMass, FRealSingle MinPerParticleMass, const FTriangleMesh& Mesh, const TFunctionRef<bool(int32)>& KinematicPredicate);
+		CHAOSCLOTH_API void SetParticleMassFromTotalMass(int32 Offset, FRealSingle TotalMass, FRealSingle MinPerParticleMass, const FTriangleMesh& Mesh, const TFunctionRef<bool(int32)>& KinematicPredicate);
+		CHAOSCLOTH_API void SetParticleMassFromDensity(int32 Offset, FRealSingle Density, FRealSingle MinPerParticleMass, const FTriangleMesh& Mesh, const TFunctionRef<bool(int32)>& KinematicPredicate);
 
 		// Set the amount of velocity allowed to filter from the given change in reference space transform, including local simulation space.
-		void SetReferenceVelocityScale(uint32 GroupId,
+		CHAOSCLOTH_API void SetReferenceVelocityScale(uint32 GroupId,
 			const FRigidTransform3& OldReferenceSpaceTransform,
 			const FRigidTransform3& ReferenceSpaceTransform,
 			const TVec3<FRealSingle>& LinearVelocityScale,
@@ -126,7 +126,7 @@ namespace Chaos
 			FRealSingle FictitiousAngularScale);
 
 		// Set general cloth simulation properties.
-		void SetProperties(
+		CHAOSCLOTH_API void SetProperties(
 			uint32 GroupId,
 			FRealSingle DampingCoefficient,
 			FRealSingle MotionDampingCoefficient,
@@ -134,22 +134,22 @@ namespace Chaos
 			FRealSingle FrictionCoefficient);
 
 		// Set whether to use continuous collision detection.
-		void SetUseCCD(uint32 GroupId, bool bUseCCD);
+		CHAOSCLOTH_API void SetUseCCD(uint32 GroupId, bool bUseCCD);
 
 		// Set per group gravity, used to override solver's gravity. Must be called during cloth update.
-		void SetGravity(uint32 GroupId, const TVec3<FRealSingle>& Gravity);
+		CHAOSCLOTH_API void SetGravity(uint32 GroupId, const TVec3<FRealSingle>& Gravity);
 
 		// Set per group wind velocity, used to override solver's wind velocity. Must be called during cloth update.
-		void SetWindVelocity(uint32 GroupId, const TVec3<FRealSingle>& InWindVelocity);
+		CHAOSCLOTH_API void SetWindVelocity(uint32 GroupId, const TVec3<FRealSingle>& InWindVelocity);
 
 		// Set the geometry affected by the wind and pressure
-		void SetWindAndPressureGeometry(
+		CHAOSCLOTH_API void SetWindAndPressureGeometry(
 			uint32 GroupId,
 			const FTriangleMesh& TriangleMesh,
 			const Softs::FCollectionPropertyConstFacade& PropertyCollection,
 			const TMap<FString, TConstArrayView<FRealSingle>>& WeightMaps);
 
-		void SetWindAndPressureGeometry(
+		CHAOSCLOTH_API void SetWindAndPressureGeometry(
 			uint32 GroupId,
 			const FTriangleMesh& TriangleMesh,
 			const TConstArrayView<FRealSingle>& DragMultipliers,
@@ -157,13 +157,13 @@ namespace Chaos
 			const TConstArrayView<FRealSingle>& PressureMultipliers);
 
 		// Set the wind and pressure properties.
-		void SetWindAndPressureProperties(
+		CHAOSCLOTH_API void SetWindAndPressureProperties(
 			uint32 GroupId,
 			const Softs::FCollectionPropertyConstFacade& PropertyCollection,
 			const TMap<FString, TConstArrayView<FRealSingle>>& WeightMaps,
 			bool bEnableAerodynamics);
 
-		void SetWindAndPressureProperties(
+		CHAOSCLOTH_API void SetWindAndPressureProperties(
 			uint32 GroupId,
 			const TVec2<FRealSingle>& Drag,
 			const TVec2<FRealSingle>& Lift,
@@ -171,10 +171,10 @@ namespace Chaos
 			const TVec2<FRealSingle>& Pressure = TVec2<FRealSingle>::ZeroVector);
 
 		// Return the wind velocity and pressure field associated with a given group id.
-		const Softs::FVelocityAndPressureField& GetWindVelocityAndPressureField(uint32 GroupId) const;
+		CHAOSCLOTH_API const Softs::FVelocityAndPressureField& GetWindVelocityAndPressureField(uint32 GroupId) const;
 
 		// Add external forces to the particles
-		void AddExternalForces(uint32 GroupId, bool bUseLegacyWind);
+		CHAOSCLOTH_API void AddExternalForces(uint32 GroupId, bool bUseLegacyWind);
 
 		const Softs::FSolverVec3* GetOldAnimationPositions(int32 Offset) const { return OldAnimationPositions.GetData() + Offset; }
 		Softs::FSolverVec3* GetOldAnimationPositions(int32 Offset) { return OldAnimationPositions.GetData() + Offset; }
@@ -192,23 +192,23 @@ namespace Chaos
 		const Softs::FSolverVec3* GetAnimationVelocities(int32 Offset) const { return AnimationVelocities.GetData() + Offset; }
 		Softs::FSolverVec3* GetAnimationVelocities(int32 Offset) { return AnimationVelocities.GetData() + Offset; }
 		Softs::FSolverVec3* GetNormals(int32 Offset) { return Normals.GetData() + Offset; }
-		const Softs::FPAndInvM* GetParticlePandInvMs(int32 Offset) const;
-		Softs::FPAndInvM* GetParticlePandInvMs(int32 Offset);
-		const Softs::FSolverVec3* GetParticleXs(int32 Offset) const;
-		Softs::FSolverVec3* GetParticleXs(int32 Offset);
-		const Softs::FSolverVec3* GetParticleVs(int32 Offset) const;
-		Softs::FSolverVec3* GetParticleVs(int32 Offset);
-		const Softs::FSolverReal* GetParticleInvMasses(int32 Offset) const;
+		CHAOSCLOTH_API const Softs::FPAndInvM* GetParticlePandInvMs(int32 Offset) const;
+		CHAOSCLOTH_API Softs::FPAndInvM* GetParticlePandInvMs(int32 Offset);
+		CHAOSCLOTH_API const Softs::FSolverVec3* GetParticleXs(int32 Offset) const;
+		CHAOSCLOTH_API Softs::FSolverVec3* GetParticleXs(int32 Offset);
+		CHAOSCLOTH_API const Softs::FSolverVec3* GetParticleVs(int32 Offset) const;
+		CHAOSCLOTH_API Softs::FSolverVec3* GetParticleVs(int32 Offset);
+		CHAOSCLOTH_API const Softs::FSolverReal* GetParticleInvMasses(int32 Offset) const;
 		const FClothConstraints& GetClothConstraints(int32 Offset) const { return *ClothsConstraints.FindChecked(Offset); }
 		FClothConstraints& GetClothConstraints(int32 Offset) { return *ClothsConstraints.FindChecked(Offset); }
-		uint32 GetNumParticles() const;
+		CHAOSCLOTH_API uint32 GetNumParticles() const;
 		// ---- End of the Cloth interface ----
 
 		// ---- Collider interface ----
-		int32 AddCollisionParticles(int32 NumCollisionParticles, uint32 GroupId, int32 RecycledOffset = 0);
-		void EnableCollisionParticles(int32 Offset, bool bEnable);
+		CHAOSCLOTH_API int32 AddCollisionParticles(int32 NumCollisionParticles, uint32 GroupId, int32 RecycledOffset = 0);
+		CHAOSCLOTH_API void EnableCollisionParticles(int32 Offset, bool bEnable);
 
-		void ResetCollisionStartPose(int32 Offset, int32 NumCollisionParticles);
+		CHAOSCLOTH_API void ResetCollisionStartPose(int32 Offset, int32 NumCollisionParticles);
 
 		const int32* GetCollisionBoneIndices(int32 Offset) const { return CollisionBoneIndices.GetData() + Offset; }
 		int32* GetCollisionBoneIndices(int32 Offset) { return CollisionBoneIndices.GetData() + Offset; }
@@ -218,16 +218,16 @@ namespace Chaos
 		Softs::FSolverRigidTransform3* GetOldCollisionTransforms(int32 Offset) { return OldCollisionTransforms.GetData() + Offset; }
 		const Softs::FSolverRigidTransform3* GetCollisionTransforms(int32 Offset) const { return CollisionTransforms.GetData() + Offset; }
 		Softs::FSolverRigidTransform3* GetCollisionTransforms(int32 Offset) { return CollisionTransforms.GetData() + Offset; }
-		const Softs::FSolverVec3* GetCollisionParticleXs(int32 Offset) const;
-		Softs::FSolverVec3* GetCollisionParticleXs(int32 Offset);
-		const Softs::FSolverRotation3* GetCollisionParticleRs(int32 Offset) const;
-		Softs::FSolverRotation3* GetCollisionParticleRs(int32 Offset);
-		void SetCollisionGeometry(int32 Offset, int32 Index, TUniquePtr<FImplicitObject>&& Geometry);
-		const TUniquePtr<FImplicitObject>* GetCollisionGeometries(int32 Offset) const;
-		const bool* GetCollisionStatus(int32 Offset) const;
-		const TArray<Softs::FSolverVec3>& GetCollisionContacts() const;
-		const TArray<Softs::FSolverVec3>& GetCollisionNormals() const;
-		const TArray<Softs::FSolverReal>& GetCollisionPhis() const;
+		CHAOSCLOTH_API const Softs::FSolverVec3* GetCollisionParticleXs(int32 Offset) const;
+		CHAOSCLOTH_API Softs::FSolverVec3* GetCollisionParticleXs(int32 Offset);
+		CHAOSCLOTH_API const Softs::FSolverRotation3* GetCollisionParticleRs(int32 Offset) const;
+		CHAOSCLOTH_API Softs::FSolverRotation3* GetCollisionParticleRs(int32 Offset);
+		CHAOSCLOTH_API void SetCollisionGeometry(int32 Offset, int32 Index, TUniquePtr<FImplicitObject>&& Geometry);
+		CHAOSCLOTH_API const TUniquePtr<FImplicitObject>* GetCollisionGeometries(int32 Offset) const;
+		CHAOSCLOTH_API const bool* GetCollisionStatus(int32 Offset) const;
+		CHAOSCLOTH_API const TArray<Softs::FSolverVec3>& GetCollisionContacts() const;
+		CHAOSCLOTH_API const TArray<Softs::FSolverVec3>& GetCollisionNormals() const;
+		CHAOSCLOTH_API const TArray<Softs::FSolverReal>& GetCollisionPhis() const;
 		// ---- End of the Collider interface ----
 
 		// ---- Field interface ----
@@ -236,16 +236,16 @@ namespace Chaos
 		// ---- End of the Field interface ----
 
 	private:
-		void ResetParticles();
-		void ResetCollisionParticles(int32 InCollisionParticlesOffset = 0);
-		void ApplyPreSimulationTransforms();
-		void PreSubstep(const Softs::FSolverReal InterpolationAlpha);
-		Softs::FSolverReal SetParticleMassPerArea(int32 Offset, int32 Size, const FTriangleMesh& Mesh);
-		void ParticleMassUpdateDensity(const FTriangleMesh& Mesh, Softs::FSolverReal Density);
-		void ParticleMassClampAndKinematicStateUpdate(int32 Offset, int32 Size, Softs::FSolverReal MinPerParticleMass, const TFunctionRef<bool(int32)>& KinematicPredicate);
+		CHAOSCLOTH_API void ResetParticles();
+		CHAOSCLOTH_API void ResetCollisionParticles(int32 InCollisionParticlesOffset = 0);
+		CHAOSCLOTH_API void ApplyPreSimulationTransforms();
+		CHAOSCLOTH_API void PreSubstep(const Softs::FSolverReal InterpolationAlpha);
+		CHAOSCLOTH_API Softs::FSolverReal SetParticleMassPerArea(int32 Offset, int32 Size, const FTriangleMesh& Mesh);
+		CHAOSCLOTH_API void ParticleMassUpdateDensity(const FTriangleMesh& Mesh, Softs::FSolverReal Density);
+		CHAOSCLOTH_API void ParticleMassClampAndKinematicStateUpdate(int32 Offset, int32 Size, Softs::FSolverReal MinPerParticleMass, const TFunctionRef<bool(int32)>& KinematicPredicate);
 
 		// Update the solver field forces/velocities at the particles location
-		void UpdateSolverField();
+		CHAOSCLOTH_API void UpdateSolverField();
 
 	private:
 		TUniquePtr<Softs::FPBDEvolution> Evolution;
diff --git a/Engine/Source/Runtime/Experimental/ChaosCore/Public/ChaosCoreModule.h b/Engine/Source/Runtime/Experimental/ChaosCore/Public/ChaosCoreModule.h
--- a/Engine/Source/Runtime/Experimental/ChaosCore/Public/ChaosCoreModule.h
+++ b/Engine/Source/Runtime/Experimental/ChaosCore/Public/ChaosCoreModule.h
@@ -4,7 +4,7 @@
 
 #include "Modules/ModuleInterface.h"
 
-class CHAOSCORE_API FChaosCoreEngineModule : public IModuleInterface
+class FChaosCoreEngineModule : public IModuleInterface
 {
 public:
 
diff --git a/Engine/Source/Runtime/Experimental/ChaosSolverEngine/Public/Chaos/ChaosDebugDrawComponent.h b/Engine/Source/Runtime/Experimental/ChaosSolverEngine/Public/Chaos/ChaosDebugDrawComponent.h
--- a/Engine/Source/Runtime/Experimental/ChaosSolverEngine/Public/Chaos/ChaosDebugDrawComponent.h
+++ b/Engine/Source/Runtime/Experimental/ChaosSolverEngine/Public/Chaos/ChaosDebugDrawComponent.h
@@ -10,23 +10,23 @@
 #include "ChaosDebugDrawComponent.generated.h"
 
 
-UCLASS(BlueprintType, ClassGroup = Chaos, meta = (BlueprintSpawnableComponent))
-class CHAOSSOLVERENGINE_API UChaosDebugDrawComponent : public UActorComponent
+UCLASS(BlueprintType, ClassGroup = Chaos, meta = (BlueprintSpawnableComponent), MinimalAPI)
+class UChaosDebugDrawComponent : public UActorComponent
 {
 	GENERATED_BODY()
 
 public:
 
-	UChaosDebugDrawComponent();
+	CHAOSSOLVERENGINE_API UChaosDebugDrawComponent();
 
 	//~ Begin UActorComponent interface
-	virtual void TickComponent(float DeltaTime, enum ELevelTick TickType, FActorComponentTickFunction *ThisTickFunction) override;
-	virtual void BeginPlay() override;
-	virtual void EndPlay(const EEndPlayReason::Type EndPlayReason) override;
-	virtual void BeginDestroy() override;
+	CHAOSSOLVERENGINE_API virtual void TickComponent(float DeltaTime, enum ELevelTick TickType, FActorComponentTickFunction *ThisTickFunction) override;
+	CHAOSSOLVERENGINE_API virtual void BeginPlay() override;
+	CHAOSSOLVERENGINE_API virtual void EndPlay(const EEndPlayReason::Type EndPlayReason) override;
+	CHAOSSOLVERENGINE_API virtual void BeginDestroy() override;
 	//~ End UActorComponent interface
 
-	static void BindWorldDelegates();
+	static CHAOSSOLVERENGINE_API void BindWorldDelegates();
 
 private:
 	static void HandlePostWorldInitialization(UWorld* World, const UWorld::InitializationValues IVS);
diff --git a/Engine/Source/Runtime/Experimental/ChaosSolverEngine/Public/Chaos/ChaosEventListenerComponent.h b/Engine/Source/Runtime/Experimental/ChaosSolverEngine/Public/Chaos/ChaosEventListenerComponent.h
--- a/Engine/Source/Runtime/Experimental/ChaosSolverEngine/Public/Chaos/ChaosEventListenerComponent.h
+++ b/Engine/Source/Runtime/Experimental/ChaosSolverEngine/Public/Chaos/ChaosEventListenerComponent.h
@@ -20,17 +20,17 @@ class AChaosSolverActor;
  * Base class for listeners that query and respond to a frame's physics data (collision events, break events, etc).
  */
 
-UCLASS(BlueprintType, ClassGroup = Chaos, meta = (BlueprintSpawnableComponent))
-class CHAOSSOLVERENGINE_API UChaosEventListenerComponent : public UActorComponent
+UCLASS(BlueprintType, ClassGroup = Chaos, meta = (BlueprintSpawnableComponent), MinimalAPI)
+class UChaosEventListenerComponent : public UActorComponent
 {
 	GENERATED_BODY()
 
 public:
 
-	UChaosEventListenerComponent();
+	CHAOSSOLVERENGINE_API UChaosEventListenerComponent();
 
 	//~ Begin UActorComponent interface
-	virtual void TickComponent(float DeltaTime, enum ELevelTick TickType, FActorComponentTickFunction *ThisTickFunction) override;
+	CHAOSSOLVERENGINE_API virtual void TickComponent(float DeltaTime, enum ELevelTick TickType, FActorComponentTickFunction *ThisTickFunction) override;
 	//~ End UActorComponent interface
 
 
diff --git a/Engine/Source/Runtime/Experimental/ChaosSolverEngine/Public/Chaos/ChaosGameplayEventDispatcher.h b/Engine/Source/Runtime/Experimental/ChaosSolverEngine/Public/Chaos/ChaosGameplayEventDispatcher.h
--- a/Engine/Source/Runtime/Experimental/ChaosSolverEngine/Public/Chaos/ChaosGameplayEventDispatcher.h
+++ b/Engine/Source/Runtime/Experimental/ChaosSolverEngine/Public/Chaos/ChaosGameplayEventDispatcher.h
@@ -24,7 +24,7 @@ typedef TFunction<void(const FChaosBreakEvent&)> FOnBreakEventCallback;
 
 /** UStruct wrapper so we can store the TFunction in a TMap */
 USTRUCT()
-struct CHAOSSOLVERENGINE_API FBreakEventCallbackWrapper
+struct FBreakEventCallbackWrapper
 {
 	GENERATED_BODY()
 
@@ -36,7 +36,7 @@ typedef TFunction<void(const FChaosRemovalEvent&)> FOnRemovalEventCallback;
 
 /** UStruct wrapper so we can store the TFunction in a TMap */
 USTRUCT()
-struct CHAOSSOLVERENGINE_API FRemovalEventCallbackWrapper
+struct FRemovalEventCallbackWrapper
 {
 	GENERATED_BODY()
 
@@ -48,7 +48,7 @@ typedef TFunction<void(const FChaosCrumblingEvent&)> FOnCrumblingEventCallback;
 
 /** UStruct wrapper so we can store the TFunction in a TMap */
 USTRUCT()
-struct CHAOSSOLVERENGINE_API FCrumblingEventCallbackWrapper
+struct FCrumblingEventCallbackWrapper
 {
 	GENERATED_BODY()
 
@@ -76,20 +76,20 @@ struct FChaosPendingCollisionNotify
 };
 
 
-UCLASS()
-class CHAOSSOLVERENGINE_API UChaosGameplayEventDispatcher : public UChaosEventListenerComponent
+UCLASS(MinimalAPI)
+class UChaosGameplayEventDispatcher : public UChaosEventListenerComponent
 {
 	GENERATED_BODY()
 
 public:
 
-	UChaosGameplayEventDispatcher();
+	CHAOSSOLVERENGINE_API UChaosGameplayEventDispatcher();
 
-	virtual void OnRegister() override;
-	virtual void OnUnregister() override;
+	CHAOSSOLVERENGINE_API virtual void OnRegister() override;
+	CHAOSSOLVERENGINE_API virtual void OnUnregister() override;
 
-	void RegisterChaosEvents();
-	void UnregisterChaosEvents();
+	CHAOSSOLVERENGINE_API void RegisterChaosEvents();
+	CHAOSSOLVERENGINE_API void UnregisterChaosEvents();
 
 private:
 
@@ -131,17 +131,17 @@ public:
 	 * @param ComponentToListenTo	The component whose collisions will be reported
 	 * @param ObjectToNotify		The object that will receive the notifications. Should be a PrimitiveComponent or implement IChaosNotifyHandlerInterface, or both.
 	 */
-	void RegisterForCollisionEvents(UPrimitiveComponent* ComponentToListenTo, UObject* ObjectToNotify);
-	void UnRegisterForCollisionEvents(UPrimitiveComponent* ComponentToListenTo, UObject* ObjectToNotify);
+	CHAOSSOLVERENGINE_API void RegisterForCollisionEvents(UPrimitiveComponent* ComponentToListenTo, UObject* ObjectToNotify);
+	CHAOSSOLVERENGINE_API void UnRegisterForCollisionEvents(UPrimitiveComponent* ComponentToListenTo, UObject* ObjectToNotify);
 
-	void RegisterForBreakEvents(UPrimitiveComponent* Component, FOnBreakEventCallback InFunc);
-	void UnRegisterForBreakEvents(UPrimitiveComponent* Component);
+	CHAOSSOLVERENGINE_API void RegisterForBreakEvents(UPrimitiveComponent* Component, FOnBreakEventCallback InFunc);
+	CHAOSSOLVERENGINE_API void UnRegisterForBreakEvents(UPrimitiveComponent* Component);
 
-	void RegisterForRemovalEvents(UPrimitiveComponent* Component, FOnRemovalEventCallback InFunc);
-	void UnRegisterForRemovalEvents(UPrimitiveComponent* Component);
+	CHAOSSOLVERENGINE_API void RegisterForRemovalEvents(UPrimitiveComponent* Component, FOnRemovalEventCallback InFunc);
+	CHAOSSOLVERENGINE_API void UnRegisterForRemovalEvents(UPrimitiveComponent* Component);
 
-	void RegisterForCrumblingEvents(UPrimitiveComponent* Component, FOnCrumblingEventCallback InFunc);
-	void UnRegisterForCrumblingEvents(UPrimitiveComponent* Component);
+	CHAOSSOLVERENGINE_API void RegisterForCrumblingEvents(UPrimitiveComponent* Component, FOnCrumblingEventCallback InFunc);
+	CHAOSSOLVERENGINE_API void UnRegisterForCrumblingEvents(UPrimitiveComponent* Component);
 	
 private:
 
diff --git a/Engine/Source/Runtime/Experimental/ChaosSolverEngine/Public/Chaos/ChaosNotifyHandlerInterface.h b/Engine/Source/Runtime/Experimental/ChaosSolverEngine/Public/Chaos/ChaosNotifyHandlerInterface.h
--- a/Engine/Source/Runtime/Experimental/ChaosSolverEngine/Public/Chaos/ChaosNotifyHandlerInterface.h
+++ b/Engine/Source/Runtime/Experimental/ChaosSolverEngine/Public/Chaos/ChaosNotifyHandlerInterface.h
@@ -11,12 +11,12 @@
 #include "ChaosNotifyHandlerInterface.generated.h"
 
 USTRUCT(BlueprintType)
-struct CHAOSSOLVERENGINE_API FChaosPhysicsCollisionInfo
+struct FChaosPhysicsCollisionInfo
 {
 	GENERATED_BODY()
 public:
 
-	FChaosPhysicsCollisionInfo();
+	CHAOSSOLVERENGINE_API FChaosPhysicsCollisionInfo();
 
 	UPROPERTY(EditAnywhere, BlueprintReadWrite, Category="Chaos")
 	TObjectPtr<UPrimitiveComponent> Component = nullptr;
@@ -57,13 +57,13 @@ public:
 DECLARE_DYNAMIC_MULTICAST_DELEGATE_OneParam(FOnChaosPhysicsCollision, const FChaosPhysicsCollisionInfo&, CollisionInfo);
 
 /** Interface for objects that want collision and trailing notifies from the Chaos solver */
-UINTERFACE(BlueprintType)
-class CHAOSSOLVERENGINE_API UChaosNotifyHandlerInterface : public UInterface
+UINTERFACE(BlueprintType, MinimalAPI)
+class UChaosNotifyHandlerInterface : public UInterface
 {
 	GENERATED_UINTERFACE_BODY()
 };
 
-class CHAOSSOLVERENGINE_API IChaosNotifyHandlerInterface : public IInterface
+class IChaosNotifyHandlerInterface : public IInterface
 {
 	GENERATED_IINTERFACE_BODY()
 
@@ -75,7 +75,7 @@ public:
 	virtual void DispatchChaosPhysicsCollisionBlueprintEvents(const FChaosPhysicsCollisionInfo& CollisionInfo) {};
 
 	/** Entry point for collision notifications, called by the underlying system. Not for overriding. */
-	void HandlePhysicsCollision(const FChaosPhysicsCollisionInfo& CollisionInfo);
+	CHAOSSOLVERENGINE_API void HandlePhysicsCollision(const FChaosPhysicsCollisionInfo& CollisionInfo);
 
 };
 
@@ -84,12 +84,12 @@ public:
 /**
  * 
  */
-UCLASS()
-class CHAOSSOLVERENGINE_API UChaosSolverEngineBlueprintLibrary : public UBlueprintFunctionLibrary
+UCLASS(MinimalAPI)
+class UChaosSolverEngineBlueprintLibrary : public UBlueprintFunctionLibrary
 {
 	GENERATED_BODY()
 public:
 
 	UFUNCTION(BlueprintPure, Category = "Chaos", meta = (WorldContext = "WorldContextObject"))
-	static FHitResult ConvertPhysicsCollisionToHitResult(const FChaosPhysicsCollisionInfo& PhysicsCollision);
+	static CHAOSSOLVERENGINE_API FHitResult ConvertPhysicsCollisionToHitResult(const FChaosPhysicsCollisionInfo& PhysicsCollision);
 };
diff --git a/Engine/Source/Runtime/Experimental/ChaosSolverEngine/Public/Chaos/ChaosSolver.h b/Engine/Source/Runtime/Experimental/ChaosSolverEngine/Public/Chaos/ChaosSolver.h
--- a/Engine/Source/Runtime/Experimental/ChaosSolverEngine/Public/Chaos/ChaosSolver.h
+++ b/Engine/Source/Runtime/Experimental/ChaosSolverEngine/Public/Chaos/ChaosSolver.h
@@ -13,13 +13,13 @@ class UChaosSolver;
 * UChaosSolver (UObject)
 *
 */
-UCLASS(customconstructor)
-class CHAOSSOLVERENGINE_API UChaosSolver : public UObject
+UCLASS(customconstructor, MinimalAPI)
+class UChaosSolver : public UObject
 {
 	GENERATED_UCLASS_BODY()
 
 public:
-	UChaosSolver(const FObjectInitializer& ObjectInitializer = FObjectInitializer::Get());
+	CHAOSSOLVERENGINE_API UChaosSolver(const FObjectInitializer& ObjectInitializer = FObjectInitializer::Get());
 	bool IsVisible() { return true; }
 };
 
diff --git a/Engine/Source/Runtime/Experimental/ChaosSolverEngine/Public/Chaos/ChaosSolverActor.h b/Engine/Source/Runtime/Experimental/ChaosSolverEngine/Public/Chaos/ChaosSolverActor.h
--- a/Engine/Source/Runtime/Experimental/ChaosSolverEngine/Public/Chaos/ChaosSolverActor.h
+++ b/Engine/Source/Runtime/Experimental/ChaosSolverEngine/Public/Chaos/ChaosSolverActor.h
@@ -71,8 +71,8 @@ struct FChaosDebugSubstepControl
 #endif
 };
 
-UCLASS()
-class CHAOSSOLVERENGINE_API AChaosSolverActor : public AActor
+UCLASS(MinimalAPI)
+class AChaosSolverActor : public AActor
 {
 	GENERATED_UCLASS_BODY()
 
@@ -124,11 +124,11 @@ public:
 
 	/** Makes this solver the current world solver. Dynamically spawned objects will have their physics state created in this solver. */
 	UFUNCTION(BlueprintCallable, Category = "ChaosPhysics")
-	void SetAsCurrentWorldSolver();
+	CHAOSSOLVERENGINE_API void SetAsCurrentWorldSolver();
 
 	/** Controls whether the solver is able to simulate particles it controls */
 	UFUNCTION(BlueprintCallable, Category = "ChaosPhysics")
-	virtual void SetSolverActive(bool bActive);
+	CHAOSSOLVERENGINE_API virtual void SetSolverActive(bool bActive);
 
 	/*
 	* Display icon in the editor
@@ -141,25 +141,25 @@ public:
 	TSharedPtr<FPhysScene_Chaos> GetPhysicsScene() const { return PhysScene; }
 	Chaos::FPhysicsSolver* GetSolver() const { return Solver; }
 
-	virtual void PostRegisterAllComponents() override;
-	virtual void PreInitializeComponents() override;
+	CHAOSSOLVERENGINE_API virtual void PostRegisterAllComponents() override;
+	CHAOSSOLVERENGINE_API virtual void PreInitializeComponents() override;
 	
-	virtual void BeginPlay() override;
-	virtual void EndPlay(const EEndPlayReason::Type ReasonEnd) override;
+	CHAOSSOLVERENGINE_API virtual void BeginPlay() override;
+	CHAOSSOLVERENGINE_API virtual void EndPlay(const EEndPlayReason::Type ReasonEnd) override;
 
 	/** UObject interface */
 #if WITH_EDITOR
-	virtual void PostEditChangeProperty(struct FPropertyChangedEvent& PropertyChangedEvent) override;
+	CHAOSSOLVERENGINE_API virtual void PostEditChangeProperty(struct FPropertyChangedEvent& PropertyChangedEvent) override;
 #endif
-	void PostLoad() override;
-	void Serialize(FArchive& Ar) override;
-	void PostDuplicate(EDuplicateMode::Type DuplicateMode) override;
+	CHAOSSOLVERENGINE_API void PostLoad() override;
+	CHAOSSOLVERENGINE_API void Serialize(FArchive& Ar) override;
+	CHAOSSOLVERENGINE_API void PostDuplicate(EDuplicateMode::Type DuplicateMode) override;
 	/** End UObject interface */
 
 private:
 
 	/** If floor is enabled, make a particle to represent it */
-	void MakeFloor();
+	CHAOSSOLVERENGINE_API void MakeFloor();
 
 	TSharedPtr<FPhysScene_Chaos> PhysScene;
 	Chaos::FPhysicsSolver* Solver;
@@ -170,4 +170,4 @@ private:
 
 	/** If floor is enabled - this will point to the solver particle for it */
 	Chaos::FSingleParticlePhysicsProxy* Proxy;
-};
\ No newline at end of file
+};
diff --git a/Engine/Source/Runtime/Experimental/ChaosSolverEngine/Public/Chaos/ChaosSolverSettings.h b/Engine/Source/Runtime/Experimental/ChaosSolverEngine/Public/Chaos/ChaosSolverSettings.h
--- a/Engine/Source/Runtime/Experimental/ChaosSolverEngine/Public/Chaos/ChaosSolverSettings.h
+++ b/Engine/Source/Runtime/Experimental/ChaosSolverEngine/Public/Chaos/ChaosSolverSettings.h
@@ -10,28 +10,28 @@
  * Settings class for the Chaos Solver
  */
 
-UCLASS(config = Engine, defaultconfig, meta = (DisplayName = "Chaos Solver"))
-class CHAOSSOLVERENGINE_API UChaosSolverSettings : public UDeveloperSettings, public IChaosSolverActorClassProvider
+UCLASS(config = Engine, defaultconfig, meta = (DisplayName = "Chaos Solver"), MinimalAPI)
+class UChaosSolverSettings : public UDeveloperSettings, public IChaosSolverActorClassProvider
 {
 	GENERATED_BODY()
 
 public:
 
-	UChaosSolverSettings();
+	CHAOSSOLVERENGINE_API UChaosSolverSettings();
 
 	// IChaosSolverActorClassProvider interface
-	virtual UClass* GetSolverActorClass() const;
+	CHAOSSOLVERENGINE_API virtual UClass* GetSolverActorClass() const;
 
 	/** The class to use when auto-creating a default chaos solver actor */
 	UPROPERTY(config, noclear, EditAnywhere, Category = GameInstance, meta = (MetaClass = "/Script/ChaosSolverEngine.ChaosSolverActor"))
 	FSoftClassPath DefaultChaosSolverActorClass;
 
 #if WITH_EDITOR
-	virtual void PostEditChangeProperty(FPropertyChangedEvent& PropertyChangedEvent) override;
+	CHAOSSOLVERENGINE_API virtual void PostEditChangeProperty(FPropertyChangedEvent& PropertyChangedEvent) override;
 #endif
 
-	virtual void PostInitProperties() override;
-	virtual void PostReloadConfig(class FProperty* PropertyThatWasLoaded) override;
+	CHAOSSOLVERENGINE_API virtual void PostInitProperties() override;
+	CHAOSSOLVERENGINE_API virtual void PostReloadConfig(class FProperty* PropertyThatWasLoaded) override;
 
 private: 
 
