Added a console variable that allows to set the name of the weight map to be visualized.
Added a method to ChaosCloth SimulationVisualization menu to draw the weight map whose name was set via the console variable.
Added a shortcut method for visualizing the TransferWeight node internal weight map specifically.

#rnx

diff --git a/Engine/Plugins/Experimental/ChaosCloth/Source/Chaos/Private/ChaosCloth/ChaosClothVisualization.cpp b/Engine/Plugins/Experimental/ChaosCloth/Source/Chaos/Private/ChaosCloth/ChaosClothVisualization.cpp
--- a/Engine/Plugins/Experimental/ChaosCloth/Source/Chaos/Private/ChaosCloth/ChaosClothVisualization.cpp
+++ b/Engine/Plugins/Experimental/ChaosCloth/Source/Chaos/Private/ChaosCloth/ChaosClothVisualization.cpp
@@ -74,6 +74,9 @@ static FAutoConsoleVariableRef CVarClothVizStretchBiasDrawRangeMax(TEXT("p.Chaos
 static bool bStretchBiasDrawOutOfRange = true;
 static FAutoConsoleVariableRef CVarClothVizStretchBiasDrawOutOfRange(TEXT("p.ChaosClothVisualization.StretchBiasDrawOutOfRange"), bStretchBiasDrawOutOfRange, TEXT("Draw out of range elements (When drawing warp/weft stretch)"));
 
+static FString WeightMapName = "";
+static FAutoConsoleVariableRef CVarClothVizWeightMapName(TEXT("p.ChaosClothVisualization.WeightMapName"), WeightMapName, TEXT("Weight map name to be visualized"));
+
 }// namespace Private
 
 	FClothVisualization::FClothVisualization(const ::Chaos::FClothingSimulationSolver* InSolver)
@@ -294,6 +297,76 @@ static FAutoConsoleVariableRef CVarClothVizStretchBiasDrawOutOfRange(TEXT("p.Cha
 			}
 		}
 	}
+
+	void FClothVisualization::DrawWeightMapWithName(FPrimitiveDrawInterface* PDI, const FString& Name) const
+	{
+		if (!Solver || !ClothMaterialColor)
+		{
+			return;
+		}
+
+		FDynamicMeshBuilder MeshBuilder(PDI->View->GetFeatureLevel());
+		int32 VertexIndex = 0;
+
+		for (const FClothingSimulationCloth* const Cloth : Solver->GetCloths())
+		{
+			const int32 Offset = Cloth->GetOffset(Solver);
+			if (Offset == INDEX_NONE)
+			{
+				continue;
+			}
+
+			const TConstArrayView<TVec3<int32>> Elements = Cloth->GetTriangleMesh(Solver).GetElements();
+			const TConstArrayView<Softs::FSolverVec3> Positions = Cloth->GetParticlePositions(Solver);
+			const TConstArrayView<FRealSingle>& WeightMap = Cloth->GetWeightMapByName(Solver, Name);
+
+			for (int32 ElementIndex = 0; ElementIndex < Elements.Num(); ++ElementIndex, VertexIndex += 3)
+			{
+				const TVec3<int32>& Element = Elements[ElementIndex];
+
+				const FVector3f Pos0(Positions[Element.X - Offset]);
+				const FVector3f Pos1(Positions[Element.Y - Offset]);
+				const FVector3f Pos2(Positions[Element.Z - Offset]);
+
+				const FVector3f Normal = FVector3f::CrossProduct(Pos2 - Pos0, Pos1 - Pos0).GetSafeNormal();
+				const FVector3f Tangent = ((Pos1 + Pos2) * 0.5f - Pos0).GetSafeNormal();
+
+				FLinearColor VertexColor1 = FLinearColor::Black;
+				FLinearColor VertexColor2 = FLinearColor::Black;
+				FLinearColor VertexColor3 = FLinearColor::Black;
+
+				if (!WeightMap.IsEmpty() && WeightMap.Num() == Positions.Num()) // if map with that name exists and not empty
+				{
+					const FRealSingle Value0(WeightMap[Element.X - Offset]);
+					const FRealSingle Value1(WeightMap[Element.Y - Offset]);
+					const FRealSingle Value2(WeightMap[Element.Z - Offset]);
+
+					VertexColor1 = FLinearColor::LerpUsingHSV(FLinearColor::Black, FLinearColor::White, (float)Value0);
+					VertexColor2 = FLinearColor::LerpUsingHSV(FLinearColor::Black, FLinearColor::White, (float)Value1);
+					VertexColor3 = FLinearColor::LerpUsingHSV(FLinearColor::Black, FLinearColor::White, (float)Value2);
+				}
+	
+				MeshBuilder.AddVertex(FDynamicMeshVertex(Pos0, Tangent, Normal, FVector2f(0.f, 0.f), VertexColor1.ToFColor(true)));
+				MeshBuilder.AddVertex(FDynamicMeshVertex(Pos1, Tangent, Normal, FVector2f(0.f, 1.f), VertexColor2.ToFColor(true)));
+				MeshBuilder.AddVertex(FDynamicMeshVertex(Pos2, Tangent, Normal, FVector2f(1.f, 1.f), VertexColor3.ToFColor(true)));
+				MeshBuilder.AddTriangle(VertexIndex, VertexIndex + 1, VertexIndex + 2);
+			}
+		}
+		
+		FMatrix LocalSimSpaceToWorld(FMatrix::Identity);
+		LocalSimSpaceToWorld.SetOrigin(Solver->GetLocalSpaceLocation());
+		MeshBuilder.Draw(PDI, LocalSimSpaceToWorld, ClothMaterialColor->GetRenderProxy(), SDPG_World, false, false);
+	}
+
+	void FClothVisualization::DrawWeightMap(FPrimitiveDrawInterface* PDI) const
+	{
+		DrawWeightMapWithName(PDI, Private::WeightMapName);
+	}
+
+	void FClothVisualization::DrawInpaintWeightsMatched(FPrimitiveDrawInterface* PDI) const
+	{
+		DrawWeightMapWithName(PDI, TEXT("_InpaintWeightMask"));
+	}
 #endif  // #if WITH_EDITOR
 
 	static void DrawPoint(FPrimitiveDrawInterface* PDI, const FVector& Pos, const FLinearColor& Color, const UMaterial* ClothMaterialVertex, const float Thickness = 1.f)  // Use color or material
diff --git a/Engine/Plugins/Experimental/ChaosCloth/Source/Chaos/Public/ChaosCloth/ChaosClothVisualization.h b/Engine/Plugins/Experimental/ChaosCloth/Source/Chaos/Public/ChaosCloth/ChaosClothVisualization.h
--- a/Engine/Plugins/Experimental/ChaosCloth/Source/Chaos/Public/ChaosCloth/ChaosClothVisualization.h
+++ b/Engine/Plugins/Experimental/ChaosCloth/Source/Chaos/Public/ChaosCloth/ChaosClothVisualization.h
@@ -84,7 +84,8 @@ namespace Chaos
 		CHAOSCLOTH_API void DrawParticleIndices(FCanvas* Canvas, const FSceneView* SceneView) const;
 		CHAOSCLOTH_API void DrawElementIndices(FCanvas* Canvas, const FSceneView* SceneView) const;
 		CHAOSCLOTH_API void DrawMaxDistanceValues(FCanvas* Canvas, const FSceneView* SceneView) const;
-
+		CHAOSCLOTH_API void DrawWeightMap(FPrimitiveDrawInterface* PDI = nullptr) const;
+		CHAOSCLOTH_API void DrawInpaintWeightsMatched(FPrimitiveDrawInterface* PDI = nullptr) const;
 	protected:
 		// FGCObject interface
 		CHAOSCLOTH_API virtual void AddReferencedObjects(FReferenceCollector& Collector) override;
@@ -108,6 +109,8 @@ namespace Chaos
 		TObjectPtr<const UMaterial> ClothMaterialColor = nullptr;
 		TObjectPtr<const UMaterial> ClothMaterialVertex = nullptr;
 		TObjectPtr<const UMaterial> CollisionMaterial = nullptr;
+		
+		void DrawWeightMapWithName(FPrimitiveDrawInterface* PDI, const FString& Name) const;
 #endif  // #if WITH_EDITOR
 #endif  // #if CHAOS_DEBUG_DRAW
 	};
diff --git a/Engine/Plugins/Experimental/ChaosClothAssetEditor/Source/ChaosClothAssetEditor/Private/ChaosClothAsset/ClothEditorSimulationVisualization.cpp b/Engine/Plugins/Experimental/ChaosClothAssetEditor/Source/ChaosClothAssetEditor/Private/ChaosClothAsset/ClothEditorSimulationVisualization.cpp
--- a/Engine/Plugins/Experimental/ChaosClothAssetEditor/Source/ChaosClothAssetEditor/Private/ChaosClothAsset/ClothEditorSimulationVisualization.cpp
+++ b/Engine/Plugins/Experimental/ChaosClothAssetEditor/Source/ChaosClothAssetEditor/Private/ChaosClothAsset/ClothEditorSimulationVisualization.cpp
@@ -147,6 +147,8 @@ const FVisualizationOption FVisualizationOption::OptionData[] =
 	FVisualizationOption(&::Chaos::FClothVisualization::DrawSelfCollision       , TEXT("p.ChaosClothAssetEditor.DebugDrawSelfCollision"), LOCTEXT("ChaosVisName_SelfCollision"       , "Self Collision"), LOCTEXT("ChaosVisName_SelfCollision_Tooltip"       , "Draws the self collision thickness/debugging information")),
 	FVisualizationOption(&::Chaos::FClothVisualization::DrawSelfIntersection    , TEXT("p.ChaosClothAssetEditor.DebugDrawSelfIntersection"), LOCTEXT("ChaosVisName_SelfIntersection"    , "Self Intersection"), LOCTEXT("ChaosVisName_SelfIntersection_Tooltip"    , "Draws the self intersection contour/region information")),
 	FVisualizationOption(&GetSimulationStatisticsString    , TEXT("p.ChaosClothAssetEditor.DebugDrawSimulationStatistics"), LOCTEXT("ChaosVisName_SimulationStatistics"    , "Simulation Statistics"), LOCTEXT("ChaosVisName_SimulationStatistics_Tooltip"    , "Displays simulation statistics")),
+	FVisualizationOption(&::Chaos::FClothVisualization::DrawWeightMap			, TEXT("p.ChaosClothAssetEditor.DrawWeightMap"), LOCTEXT("ChaosVisName_DrawWeightMap"        , "Weight Map"), LOCTEXT("ChaosVisName_DrawWeightMap_ToolTip"       , "Draw the weight map for the simulation mesh. You can control the name of the map to be visualized by setting the p.ChaosClothVisualization.WeightMapName console variable."), /*bDisablesSimulation =*/false, /*bHidesClothSections=*/true),
+	FVisualizationOption(&::Chaos::FClothVisualization::DrawInpaintWeightsMatched, TEXT("p.ChaosClothAssetEditor.DrawInpaintWeightsMatched"), LOCTEXT("ChaosVisName_DrawInpaintWeightsMatched"        , "Transfer Skin Weights Node: Matched Vertices"), LOCTEXT("ChaosVisName_DrawInpaintWeightsMatched_ToolTip"       , "When transferring weights using the InpaintWeights method, will highlight the vertices for which we copied the weights directly from the source mesh. For all other vertices, the weights were computed automatically."), /*bDisablesSimulation =*/false, /*bHidesClothSections=*/true),
 };
 const uint32 FVisualizationOption::Count = sizeof(OptionData) / sizeof(FVisualizationOption);
 
