Chaos Cloth - Disable similar property warnings when a default config had been added prior to a config node.

#rnx

diff --git a/Engine/Plugins/Experimental/ChaosCloth/Source/Chaos/Private/ChaosCloth/ChaosClothingSimulationConfig.cpp b/Engine/Plugins/Experimental/ChaosCloth/Source/Chaos/Private/ChaosCloth/ChaosClothingSimulationConfig.cpp
--- a/Engine/Plugins/Experimental/ChaosCloth/Source/Chaos/Private/ChaosCloth/ChaosClothingSimulationConfig.cpp
+++ b/Engine/Plugins/Experimental/ChaosCloth/Source/Chaos/Private/ChaosCloth/ChaosClothingSimulationConfig.cpp
@@ -24,8 +24,11 @@ namespace Chaos
 
 	void FClothingSimulationConfig::Initialize(const UChaosClothConfig* ClothConfig, const UChaosClothSharedSimConfig* ClothSharedConfig, bool bUseLegacyConfig)
 	{
-		constexpr bool bEnable = true;
-		constexpr bool bAnimatable = true;
+		constexpr ::Chaos::Softs::ECollectionPropertyFlags NonAnimatablePropertyFlags =
+			::Chaos::Softs::ECollectionPropertyFlags::Enabled |
+			::Chaos::Softs::ECollectionPropertyFlags::Legacy;  // Indicates a property set from a pre-property collection config (e.g. that can be overriden in Dataflow without warning)
+		constexpr ::Chaos::Softs::ECollectionPropertyFlags AnimatablePropertyFlags = NonAnimatablePropertyFlags |
+			::Chaos::Softs::ECollectionPropertyFlags::Animatable;
 
 		// Clear all properties
 		PropertyCollection->Reset();
@@ -34,9 +37,9 @@ namespace Chaos
 		// Solver properties
 		if (ClothSharedConfig)
 		{
-			Properties->AddValue(TEXT("NumIterations"), ClothSharedConfig->IterationCount, bEnable, bAnimatable);
-			Properties->AddValue(TEXT("MaxNumIterations"), ClothSharedConfig->MaxIterationCount, bEnable, bAnimatable);
-			Properties->AddValue(TEXT("NumSubsteps"), ClothSharedConfig->SubdivisionCount, bEnable, bAnimatable);
+			Properties->AddValue(TEXT("NumIterations"), ClothSharedConfig->IterationCount, AnimatablePropertyFlags);
+			Properties->AddValue(TEXT("MaxNumIterations"), ClothSharedConfig->MaxIterationCount, AnimatablePropertyFlags);
+			Properties->AddValue(TEXT("NumSubsteps"), ClothSharedConfig->SubdivisionCount, AnimatablePropertyFlags);
 		}
 
 		// Cloth properties
@@ -58,15 +61,15 @@ namespace Chaos
 					MassValue = ClothConfig->Density;
 					break;
 				}
-				Properties->AddValue(TEXT("MassMode"), (int32)ClothConfig->MassMode);
-				Properties->AddValue(TEXT("MassValue"), MassValue);
-				Properties->AddValue(TEXT("MinPerParticleMass"), ClothConfig->MinPerParticleMass);
+				Properties->AddValue(TEXT("MassMode"), (int32)ClothConfig->MassMode, NonAnimatablePropertyFlags);
+				Properties->AddValue(TEXT("MassValue"), MassValue, NonAnimatablePropertyFlags);
+				Properties->AddValue(TEXT("MinPerParticleMass"), ClothConfig->MinPerParticleMass, NonAnimatablePropertyFlags);
 			}
 
 			// Edge constraint
 			if (ClothConfig->EdgeStiffnessWeighted.Low > 0.f || ClothConfig->EdgeStiffnessWeighted.High > 0.f)
 			{
-				const int32 EdgeSpringStiffnessIndex = Properties->AddProperty(TEXT("EdgeSpringStiffness"), bEnable, bAnimatable);
+				const int32 EdgeSpringStiffnessIndex = Properties->AddProperty(TEXT("EdgeSpringStiffness"), AnimatablePropertyFlags);
 				Properties->SetWeightedValue(EdgeSpringStiffnessIndex, ClothConfig->EdgeStiffnessWeighted.Low, ClothConfig->EdgeStiffnessWeighted.High);
 				Properties->SetStringValue(EdgeSpringStiffnessIndex, TEXT("EdgeStiffness"));
 			}
@@ -77,22 +80,22 @@ namespace Chaos
 			{
 				if (ClothConfig->bUseBendingElements)
 				{
-					const int32 BendingElementStiffnessIndex = Properties->AddProperty(TEXT("BendingElementStiffness"), bEnable, bAnimatable);
+					const int32 BendingElementStiffnessIndex = Properties->AddProperty(TEXT("BendingElementStiffness"), AnimatablePropertyFlags);
 					Properties->SetWeightedValue(BendingElementStiffnessIndex, ClothConfig->BendingStiffnessWeighted.Low, ClothConfig->BendingStiffnessWeighted.High);
 					Properties->SetStringValue(BendingElementStiffnessIndex, TEXT("BendingStiffness"));
 
-					Properties->AddValue(TEXT("BucklingRatio"), ClothConfig->BucklingRatio);
+					Properties->AddValue(TEXT("BucklingRatio"), ClothConfig->BucklingRatio, NonAnimatablePropertyFlags);
 
 					if (ClothConfig->BucklingStiffnessWeighted.Low > 0.f || ClothConfig->BucklingStiffnessWeighted.High > 0.f)
 					{
-						const int32 BucklingStiffnessIndex = Properties->AddProperty(TEXT("BucklingStiffness"), bEnable, bAnimatable);
+						const int32 BucklingStiffnessIndex = Properties->AddProperty(TEXT("BucklingStiffness"), AnimatablePropertyFlags);
 						Properties->SetWeightedValue(BucklingStiffnessIndex, ClothConfig->BucklingStiffnessWeighted.Low, ClothConfig->BucklingStiffnessWeighted.High);
 						Properties->SetStringValue(BucklingStiffnessIndex, TEXT("BucklingStiffness"));
 					}
 				}
 				else  // Not using bending elements
 				{
-					const int32 BendingSpringStiffnessIndex = Properties->AddProperty(TEXT("BendingSpringStiffness"), bEnable, bAnimatable);
+					const int32 BendingSpringStiffnessIndex = Properties->AddProperty(TEXT("BendingSpringStiffness"), AnimatablePropertyFlags);
 					Properties->SetWeightedValue(BendingSpringStiffnessIndex, ClothConfig->BendingStiffnessWeighted.Low, ClothConfig->BendingStiffnessWeighted.High);
 					Properties->SetStringValue(BendingSpringStiffnessIndex, TEXT("BendingStiffness"));
 				}
@@ -101,7 +104,7 @@ namespace Chaos
 			// Area constraint
 			if (ClothConfig->AreaStiffnessWeighted.Low > 0.f || ClothConfig->AreaStiffnessWeighted.High > 0.f)
 			{
-				const int32 AreaSpringStiffnessIndex = Properties->AddProperty(TEXT("AreaSpringStiffness"), bEnable, bAnimatable);
+				const int32 AreaSpringStiffnessIndex = Properties->AddProperty(TEXT("AreaSpringStiffness"), AnimatablePropertyFlags);
 				Properties->SetWeightedValue(AreaSpringStiffnessIndex, ClothConfig->AreaStiffnessWeighted.Low, ClothConfig->AreaStiffnessWeighted.High);
 				Properties->SetStringValue(AreaSpringStiffnessIndex, TEXT("AreaStiffness"));
 			}
@@ -109,13 +112,13 @@ namespace Chaos
 			// Long range attachment
 			if (ClothConfig->TetherStiffness.Low > 0.f || ClothConfig->TetherStiffness.High > 0.f)
 			{
-				Properties->AddValue(TEXT("UseGeodesicTethers"), ClothConfig->bUseGeodesicDistance);
+				Properties->AddValue(TEXT("UseGeodesicTethers"), ClothConfig->bUseGeodesicDistance, NonAnimatablePropertyFlags);
 
-				const int32 TetherStiffnessIndex = Properties->AddProperty(TEXT("TetherStiffness"), bEnable, bAnimatable);
+				const int32 TetherStiffnessIndex = Properties->AddProperty(TEXT("TetherStiffness"), AnimatablePropertyFlags);
 				Properties->SetWeightedValue(TetherStiffnessIndex, ClothConfig->TetherStiffness.Low, ClothConfig->TetherStiffness.High);
 				Properties->SetStringValue(TetherStiffnessIndex, TEXT("TetherStiffness"));
 
-				const int32 TetherScaleIndex = Properties->AddProperty(TEXT("TetherScale"), bEnable, bAnimatable);
+				const int32 TetherScaleIndex = Properties->AddProperty(TEXT("TetherScale"), AnimatablePropertyFlags);
 				Properties->SetWeightedValue(TetherScaleIndex, ClothConfig->TetherScale.Low, ClothConfig->TetherScale.High);
 				Properties->SetStringValue(TetherScaleIndex, TEXT("TetherScale"));
 			}
@@ -123,89 +126,89 @@ namespace Chaos
 			// AnimDrive
 			if (ClothConfig->AnimDriveStiffness.Low > 0.f || ClothConfig->AnimDriveStiffness.High > 0.f)
 			{
-				const int32 AnimDriveStiffnessIndex = Properties->AddProperty(TEXT("AnimDriveStiffness"), bEnable, bAnimatable);
+				const int32 AnimDriveStiffnessIndex = Properties->AddProperty(TEXT("AnimDriveStiffness"), AnimatablePropertyFlags);
 				Properties->SetWeightedValue(AnimDriveStiffnessIndex, ClothConfig->AnimDriveStiffness.Low, ClothConfig->AnimDriveStiffness.High);
 				Properties->SetStringValue(AnimDriveStiffnessIndex, TEXT("AnimDriveStiffness"));
 
-				const int32 AnimDriveDampingIndex = Properties->AddProperty(TEXT("AnimDriveDamping"), bEnable, bAnimatable);
+				const int32 AnimDriveDampingIndex = Properties->AddProperty(TEXT("AnimDriveDamping"), AnimatablePropertyFlags);
 				Properties->SetWeightedValue(AnimDriveDampingIndex, ClothConfig->AnimDriveDamping.Low, ClothConfig->AnimDriveDamping.High);
 				Properties->SetStringValue(AnimDriveDampingIndex, TEXT("AnimDriveDamping"));
 			}
 
 			// Gravity
 			{
-				Properties->AddValue(TEXT("GravityScale"), ClothConfig->GravityScale, bEnable, bAnimatable);
-				Properties->AddValue(TEXT("UseGravityOverride"), ClothConfig->bUseGravityOverride, bEnable, bAnimatable);
-				Properties->AddValue(TEXT("GravityOverride"), FVector3f(ClothConfig->Gravity), bEnable, bAnimatable);
+				Properties->AddValue(TEXT("GravityScale"), ClothConfig->GravityScale, AnimatablePropertyFlags);
+				Properties->AddValue(TEXT("UseGravityOverride"), ClothConfig->bUseGravityOverride, AnimatablePropertyFlags);
+				Properties->AddValue(TEXT("GravityOverride"), FVector3f(ClothConfig->Gravity), AnimatablePropertyFlags);
 			}
 
 			// Velocity scale
 			{
-				Properties->AddValue(TEXT("LinearVelocityScale"), FVector3f(ClothConfig->LinearVelocityScale), bEnable, bAnimatable);
-				Properties->AddValue(TEXT("AngularVelocityScale"), ClothConfig->AngularVelocityScale, bEnable, bAnimatable);
-				Properties->AddValue(TEXT("FictitiousAngularScale"), ClothConfig->FictitiousAngularScale, bEnable, bAnimatable);
+				Properties->AddValue(TEXT("LinearVelocityScale"), FVector3f(ClothConfig->LinearVelocityScale), AnimatablePropertyFlags);
+				Properties->AddValue(TEXT("AngularVelocityScale"), ClothConfig->AngularVelocityScale, AnimatablePropertyFlags);
+				Properties->AddValue(TEXT("FictitiousAngularScale"), ClothConfig->FictitiousAngularScale, AnimatablePropertyFlags);
 			}
 
 			// Aerodynamics
-			Properties->AddValue(TEXT("UsePointBasedWindModel"), ClothConfig->bUsePointBasedWindModel);
+			Properties->AddValue(TEXT("UsePointBasedWindModel"), ClothConfig->bUsePointBasedWindModel, NonAnimatablePropertyFlags);
 			if (!ClothConfig->bUsePointBasedWindModel && (ClothConfig->Drag.Low > 0.f || ClothConfig->Drag.High > 0.f || ClothConfig->Lift.Low > 0.f || ClothConfig->Lift.High > 0.f))
 			{
-				const int32 DragIndex = Properties->AddProperty(TEXT("Drag"), bEnable, bAnimatable);
+				const int32 DragIndex = Properties->AddProperty(TEXT("Drag"), AnimatablePropertyFlags);
 				Properties->SetWeightedValue(DragIndex, ClothConfig->Drag.Low, ClothConfig->Drag.High);
 				Properties->SetStringValue(DragIndex, TEXT("Drag"));
 
-				const int32 LiftIndex = Properties->AddProperty(TEXT("Lift"), bEnable, bAnimatable);
+				const int32 LiftIndex = Properties->AddProperty(TEXT("Lift"), AnimatablePropertyFlags);
 				Properties->SetWeightedValue(LiftIndex, ClothConfig->Lift.Low, ClothConfig->Lift.High);
 				Properties->SetStringValue(LiftIndex, TEXT("Lift"));
 
 				constexpr float AirDensity = 1.225f;
-				Properties->AddValue(TEXT("FluidDensity"), AirDensity, bEnable, bAnimatable);
+				Properties->AddValue(TEXT("FluidDensity"), AirDensity, AnimatablePropertyFlags);
 			}
 
 			// Pressure
 			if (ClothConfig->Pressure.Low != 0.f || ClothConfig->Pressure.High != 0.f)
 			{
-				const int32 PressureIndex = Properties->AddProperty(TEXT("Pressure"), bEnable, bAnimatable);
+				const int32 PressureIndex = Properties->AddProperty(TEXT("Pressure"), AnimatablePropertyFlags);
 				Properties->SetWeightedValue(PressureIndex, ClothConfig->Pressure.Low, ClothConfig->Pressure.High);
 				Properties->SetStringValue(PressureIndex, TEXT("Pressure"));
 			}
 
 			// Damping
-			Properties->AddValue(TEXT("DampingCoefficient"), ClothConfig->DampingCoefficient, bEnable, bAnimatable);
-			Properties->AddValue(TEXT("LocalDampingCoefficient"), ClothConfig->LocalDampingCoefficient, bEnable, bAnimatable);
+			Properties->AddValue(TEXT("DampingCoefficient"), ClothConfig->DampingCoefficient, AnimatablePropertyFlags);
+			Properties->AddValue(TEXT("LocalDampingCoefficient"), ClothConfig->LocalDampingCoefficient, AnimatablePropertyFlags);
 
 			// Collision
-			Properties->AddValue(TEXT("CollisionThickness"), ClothConfig->CollisionThickness, bEnable, bAnimatable);
-			Properties->AddValue(TEXT("FrictionCoefficient"), ClothConfig->FrictionCoefficient, bEnable, bAnimatable);
-			Properties->AddValue(TEXT("UseCCD"), ClothConfig->bUseCCD, bEnable, bAnimatable);
-			Properties->AddValue(TEXT("UseSelfCollisions"), ClothConfig->bUseSelfCollisions);
-			Properties->AddValue(TEXT("SelfCollisionThickness"), ClothConfig->SelfCollisionThickness);
-			Properties->AddValue(TEXT("UseSelfIntersections"), ClothConfig->bUseSelfIntersections);
-			Properties->AddValue(TEXT("SelfCollisionFriction"), ClothConfig->SelfCollisionFriction);
+			Properties->AddValue(TEXT("CollisionThickness"), ClothConfig->CollisionThickness, AnimatablePropertyFlags);
+			Properties->AddValue(TEXT("FrictionCoefficient"), ClothConfig->FrictionCoefficient, AnimatablePropertyFlags);
+			Properties->AddValue(TEXT("UseCCD"), ClothConfig->bUseCCD, AnimatablePropertyFlags);
+			Properties->AddValue(TEXT("UseSelfCollisions"), ClothConfig->bUseSelfCollisions, NonAnimatablePropertyFlags);
+			Properties->AddValue(TEXT("SelfCollisionThickness"), ClothConfig->SelfCollisionThickness, NonAnimatablePropertyFlags);
+			Properties->AddValue(TEXT("UseSelfIntersections"), ClothConfig->bUseSelfIntersections, NonAnimatablePropertyFlags);
+			Properties->AddValue(TEXT("SelfCollisionFriction"), ClothConfig->SelfCollisionFriction, NonAnimatablePropertyFlags);
 
 			// Max distance
 			{
-				const int32 MaxDistanceIndex = Properties->AddProperty(TEXT("MaxDistance"));
+				const int32 MaxDistanceIndex = Properties->AddProperty(TEXT("MaxDistance"), AnimatablePropertyFlags);
 				Properties->SetWeightedValue(MaxDistanceIndex, 0.f, 1.f);  // Backward compatibility with legacy mask must use a unit range since the multiplier is in the mask
 				Properties->SetStringValue(MaxDistanceIndex, TEXT("MaxDistance"));
 			}
 
 			// Backstop
 			{
-				const int32 BackstopDistanceIndex = Properties->AddProperty(TEXT("BackstopDistance"));
+				const int32 BackstopDistanceIndex = Properties->AddProperty(TEXT("BackstopDistance"), AnimatablePropertyFlags);
 				Properties->SetWeightedValue(BackstopDistanceIndex, 0.f, 1.f);  // Backward compatibility with legacy mask must use a unit range since the multiplier is in the mask
 				Properties->SetStringValue(BackstopDistanceIndex, TEXT("BackstopDistance"));
 
-				const int32 BackstopRadiusIndex = Properties->AddProperty(TEXT("BackstopRadius"));
+				const int32 BackstopRadiusIndex = Properties->AddProperty(TEXT("BackstopRadius"), AnimatablePropertyFlags);
 				Properties->SetWeightedValue(BackstopRadiusIndex, 0.f, 1.f);  // Backward compatibility with legacy mask must use a unit range since the multiplier is in the mask
 				Properties->SetStringValue(BackstopRadiusIndex, TEXT("BackstopRadius"));
 
-				Properties->AddValue(TEXT("UseLegacyBackstop"), ClothConfig->bUseLegacyBackstop);
+				Properties->AddValue(TEXT("UseLegacyBackstop"), ClothConfig->bUseLegacyBackstop, NonAnimatablePropertyFlags);
 			}
 		}
 
-		// Mark this as a potential legacy config, but leave the behavior control to the client code
-		Properties->AddValue(TEXT("UseLegacyConfig"), bUseLegacyConfig);
+		// Mark this as a potential legacy config, but leave the behavior control to the client code (usually means constraint are removed with 0 stiffness, or missing weight maps)
+		Properties->AddValue(TEXT("UseLegacyConfig"), bUseLegacyConfig, NonAnimatablePropertyFlags);
 	}
 
 	void FClothingSimulationConfig::Initialize(const TSharedPtr<const FManagedArrayCollection>& InPropertyCollection)
diff --git a/Engine/Plugins/Experimental/ChaosClothAssetEditor/Source/ChaosClothAssetDataflowNodes/Private/ChaosClothAsset/SimulationBaseConfigNode.cpp b/Engine/Plugins/Experimental/ChaosClothAssetEditor/Source/ChaosClothAssetDataflowNodes/Private/ChaosClothAsset/SimulationBaseConfigNode.cpp
--- a/Engine/Plugins/Experimental/ChaosClothAssetEditor/Source/ChaosClothAssetDataflowNodes/Private/ChaosClothAsset/SimulationBaseConfigNode.cpp
+++ b/Engine/Plugins/Experimental/ChaosClothAssetEditor/Source/ChaosClothAssetDataflowNodes/Private/ChaosClothAsset/SimulationBaseConfigNode.cpp
@@ -86,22 +86,27 @@ int32 FChaosClothAssetSimulationBaseConfigNode::AddPropertyHelper(
 {
 	using namespace UE::Chaos::ClothAsset;
 
-	constexpr bool bIsEnabled = true;
-
 	int32 KeyIndex = Properties.GetKeyIndex(PropertyName.ToString());
 	if (KeyIndex == INDEX_NONE)
 	{
-		KeyIndex = Properties.AddProperty(PropertyName.ToString(), bIsEnabled, bIsAnimatable);
+		KeyIndex = Properties.AddProperty(PropertyName.ToString());
 	}
-	else
+	else if (!Properties.IsLegacy(KeyIndex))  // Only warns of duplicates when the property hasn't been set using a legacy Chaos config
 	{
-		Properties.SetAnimatable(KeyIndex, bIsAnimatable);
 		Private::LogAndToastDuplicateProperty(*this, PropertyName);
 	}
+	// else don't warn and hope people know what they are doing
+
+	// Set/clear flags
+	Properties.SetEnabled(KeyIndex, true);  // Always enabled when then node is active
+	Properties.SetLegacy(KeyIndex, false);  // No longer a legacy property after beind set by this node
+	Properties.SetAnimatable(KeyIndex, bIsAnimatable);
 
+	// Check for similar properties
 	for (const FName& SimilarPropertyName : SimilarPropertyNames)
 	{
-		if (Properties.GetKeyIndex(SimilarPropertyName.ToString()) != INDEX_NONE)
+		const int32 SimilarPropertyKeyIndex = Properties.GetKeyIndex(SimilarPropertyName.ToString());
+		if (SimilarPropertyKeyIndex != INDEX_NONE && !Properties.IsLegacy(SimilarPropertyKeyIndex))  // Only warns of overrides when the property hasn't been set using a legacy Chaos config
 		{
 			Private::LogAndToastSimilarProperty(*this, PropertyName, SimilarPropertyName);
 		}
diff --git a/Engine/Source/Runtime/Experimental/Chaos/Private/Chaos/CollectionPropertyFacade.cpp b/Engine/Source/Runtime/Experimental/Chaos/Private/Chaos/CollectionPropertyFacade.cpp
--- a/Engine/Source/Runtime/Experimental/Chaos/Private/Chaos/CollectionPropertyFacade.cpp
+++ b/Engine/Source/Runtime/Experimental/Chaos/Private/Chaos/CollectionPropertyFacade.cpp
@@ -157,10 +157,15 @@ namespace Chaos::Softs
 
 	int32 FCollectionPropertyMutableFacade::AddProperty(const FString& Key, bool bEnabled, bool bAnimatable)
 	{
-		const int32 Index = GetManagedArrayCollection()->AddElements(1, PropertyFacadeNames::PropertyGroup);
 		const ECollectionPropertyFlags Flags =
 			(bEnabled ? ECollectionPropertyFlags::Enabled : ECollectionPropertyFlags::None) |
 			(bAnimatable ? ECollectionPropertyFlags::Animatable : ECollectionPropertyFlags::None);
+		return AddProperty(Key, Flags);
+	}
+
+	int32 FCollectionPropertyMutableFacade::AddProperty(const FString& Key, ECollectionPropertyFlags Flags)
+	{
+		const int32 Index = GetManagedArrayCollection()->AddElements(1, PropertyFacadeNames::PropertyGroup);
 
 		// Update the arrayviews in case the new element triggered a reallocation 
 		UpdateArrays();
@@ -177,13 +182,19 @@ namespace Chaos::Softs
 	}
 
 	int32 FCollectionPropertyMutableFacade::AddProperties(const TArray<FString>& Keys, bool bEnabled, bool bAnimatable)
+	{
+		const ECollectionPropertyFlags Flags =
+			(bEnabled ? ECollectionPropertyFlags::Enabled : ECollectionPropertyFlags::None) |
+			(bAnimatable ? ECollectionPropertyFlags::Animatable : ECollectionPropertyFlags::None);
+
+		return AddProperties(Keys, Flags);
+	}
+
+	int32 FCollectionPropertyMutableFacade::AddProperties(const TArray<FString>& Keys, ECollectionPropertyFlags Flags)
 	{
 		if (const int32 NumProperties = Keys.Num())
 		{
 			const int32 StartIndex = GetManagedArrayCollection()->AddElements(NumProperties, PropertyFacadeNames::PropertyGroup);
-			const ECollectionPropertyFlags Flags =
-				(bEnabled ? ECollectionPropertyFlags::Enabled : ECollectionPropertyFlags::None) |
-				(bAnimatable ? ECollectionPropertyFlags::Animatable : ECollectionPropertyFlags::None);
 
 			// Update the arrayviews in case the new elements triggered a reallocation 
 			UpdateArrays();
@@ -228,10 +239,24 @@ namespace Chaos::Softs
 		return KeyIndex;
 	}
 
+	int32 FCollectionPropertyMutableFacade::AddWeightedFloatValue(const FString& Key, const FVector2f& Value, ECollectionPropertyFlags Flags)
+	{
+		const int32 KeyIndex = AddProperty(Key, Flags);
+		SetWeightedFloatValue(KeyIndex, Value);
+		return KeyIndex;
+	}
+
 	int32 FCollectionPropertyMutableFacade::AddStringValue(const FString& Key, const FString& Value, bool bEnabled, bool bAnimatable)
 	{
 		const int32 KeyIndex = AddProperty(Key, bEnabled, bAnimatable);
 		SetStringValue(Key, Value);
 		return KeyIndex;
 	}
+
+	int32 FCollectionPropertyMutableFacade::AddStringValue(const FString& Key, const FString& Value, ECollectionPropertyFlags Flags)
+	{
+		const int32 KeyIndex = AddProperty(Key, Flags);
+		SetStringValue(Key, Value);
+		return KeyIndex;
+	}
 }  // End namespace Chaos
diff --git a/Engine/Source/Runtime/Experimental/Chaos/Public/Chaos/CollectionPropertyFacade.h b/Engine/Source/Runtime/Experimental/Chaos/Public/Chaos/CollectionPropertyFacade.h
--- a/Engine/Source/Runtime/Experimental/Chaos/Public/Chaos/CollectionPropertyFacade.h
+++ b/Engine/Source/Runtime/Experimental/Chaos/Public/Chaos/CollectionPropertyFacade.h
@@ -27,6 +27,7 @@ namespace Chaos::Softs
 		None = 0,
 		Enabled = 1 << 0,  /** Whether this property is enabled(so that it doesn't have to be removed from the collection when not needed). */
 		Animatable = 1 << 1,  /** Whether this property needs to be set at every frame. */
+		Legacy = 1 << 2,  /** Whether this property has been set by a legacy system predating the property collection. Can be useful for overriding/upgrading some properties post conversion. */
 		//~ Add new flags above this line
 		StringDirty = 1 << 6,  /** Whether this property's string has changed and needs to be updated at the next frame. */
 		Dirty = 1 << 7  /** Whether this property's value has changed and needs to be updated at the next frame. */
@@ -67,6 +68,9 @@ namespace Chaos::Softs
 		/** Return the property index for the specified key if it exists, or INDEX_NONE otherwise. */
 		int32 GetKeyIndex(const FString& Key) const { const int32* const Index = KeyIndices.Find(Key); return Index ? *Index : INDEX_NONE; }
 
+		/** Return the property name (key) for the specified index. */
+		const FString& GetKey(int32 KeyIndex) const { return GetValue<const FString&>(KeyIndex, KeyArray); }
+
 		//~ Values access per index, fast, no check, index must be valid (0 <= KeyIndex < Num())
 		template<typename T, TEMPLATE_REQUIRES(TIsWeightedType<T>::Value)>
 		T GetLowValue(int32 KeyIndex) const { return GetValue<T>(KeyIndex, LowValueArray); }
@@ -91,6 +95,7 @@ namespace Chaos::Softs
 
 		bool IsEnabled(int32 KeyIndex) const { return HasAnyFlags(KeyIndex, ECollectionPropertyFlags::Enabled); }
 		bool IsAnimatable(int32 KeyIndex) const { return HasAnyFlags(KeyIndex, ECollectionPropertyFlags::Animatable); }
+		bool IsLegacy(int32 KeyIndex) const { return HasAnyFlags(KeyIndex, ECollectionPropertyFlags::Legacy); }
 		bool IsStringDirty(int32 KeyIndex) const { return HasAnyFlags(KeyIndex, ECollectionPropertyFlags::StringDirty); }
 		bool IsDirty(int32 KeyIndex) const { return HasAnyFlags(KeyIndex, ECollectionPropertyFlags::Dirty); }
 
@@ -147,6 +152,11 @@ PRAGMA_ENABLE_DEPRECATION_WARNINGS
 			return SafeGet(Key, [this](int32 KeyIndex)->bool { return IsAnimatable(KeyIndex); }, bDefault, OutKeyIndex);
 		}
 
+		bool IsLegacy(const FString& Key, bool bDefault = false, int32* OutKeyIndex = nullptr) const
+		{
+			return SafeGet(Key, [this](int32 KeyIndex)->bool { return IsLegacy(KeyIndex); }, bDefault, OutKeyIndex);
+		}
+
 		bool IsStringDirty(const FString& Key, bool bDefault = false, int32* OutKeyIndex = nullptr) const
 		{
 			return SafeGet(Key, [this](int32 KeyIndex)->bool { return IsStringDirty(KeyIndex); }, bDefault, OutKeyIndex);
@@ -250,6 +260,7 @@ PRAGMA_ENABLE_DEPRECATION_WARNINGS
 
 		void SetEnabled(int32 KeyIndex, bool bEnabled) { EnableFlags(KeyIndex, ECollectionPropertyFlags::Enabled, bEnabled); }
 		void SetAnimatable(int32 KeyIndex, bool bAnimatable) { EnableFlags(KeyIndex, ECollectionPropertyFlags::Animatable, bAnimatable); }
+		void SetLegacy(int32 KeyIndex, bool bLegacy) { EnableFlags(KeyIndex, ECollectionPropertyFlags::Legacy, bLegacy); }
 		void SetDirty(int32 KeyIndex) { EnableFlags(KeyIndex, ECollectionPropertyFlags::Dirty, true); }
 		UE_DEPRECATED(5.3, "SetDirty can only be set, to unset use ClearDirtyFlags instead.")
 		void SetDirty(int32 KeyIndex, bool bDirty) { EnableFlags(KeyIndex, ECollectionPropertyFlags::Dirty, bDirty); }
@@ -308,6 +319,11 @@ PRAGMA_ENABLE_DEPRECATION_WARNINGS
 			return SafeSet(Key, [this, bAnimatable](int32 KeyIndex) { SetAnimatable(KeyIndex, bAnimatable); });
 		}
 
+		int32 SetLegacy(const FString& Key, bool bLegacy)
+		{
+			return SafeSet(Key, [this, bLegacy](int32 KeyIndex) { SetLegacy(KeyIndex, bLegacy); });
+		}
+
 		UE_DEPRECATED(5.3, "SetDirty can only be set, to unset use ClearDirtyFlags instead.")
 		int32 SetDirty(const FString& Key, bool bDirty)
 		{
@@ -381,10 +397,16 @@ PRAGMA_ENABLE_DEPRECATION_WARNINGS
 		CHAOS_API void DefineSchema();
 
 		/** Add a single property, and return its index. */
-		CHAOS_API int32 AddProperty(const FString& Key, bool bEnabled = true, bool bAnimatable = false);
+		CHAOS_API int32 AddProperty(const FString& Key, ECollectionPropertyFlags Flags = ECollectionPropertyFlags::Enabled);
+
+		/** Add a single property, and return its index. */
+		CHAOS_API int32 AddProperty(const FString& Key, bool bEnabled, bool bAnimatable = false);
 
 		/** Add new properties, and return the index of the first added property. */
-		CHAOS_API int32 AddProperties(const TArray<FString>& Keys, bool bEnabled = true, bool bAnimatable = false);
+		CHAOS_API int32 AddProperties(const TArray<FString>& Keys, ECollectionPropertyFlags Flags = ECollectionPropertyFlags::Enabled);
+
+		/** Add new properties, and return the index of the first added property. */
+		CHAOS_API int32 AddProperties(const TArray<FString>& Keys, bool bEnabled, bool bAnimatable = false);
 
 		/**
 		 * Append all properties and values from an existing collection to this property collection.
@@ -402,14 +424,20 @@ PRAGMA_ENABLE_DEPRECATION_WARNINGS
 
 		//~ Add values
 		template<typename T, TEMPLATE_REQUIRES(TIsWeightedType<T>::Value)>
-		inline int32 AddWeightedValue(const FString& Key, const T& LowValue, const T& HighValue, bool bEnabled = true, bool bAnimatable = false);
+		inline int32 AddWeightedValue(const FString& Key, const T& LowValue, const T& HighValue, ECollectionPropertyFlags Flags = ECollectionPropertyFlags::Enabled);
+		template<typename T, TEMPLATE_REQUIRES(TIsWeightedType<T>::Value)>
+		inline int32 AddWeightedValue(const FString& Key, const T& LowValue, const T& HighValue, bool bEnabled, bool bAnimatable = false);
 
+		CHAOS_API int32 AddWeightedFloatValue(const FString& Key, const FVector2f& Value, ECollectionPropertyFlags Flags = ECollectionPropertyFlags::Enabled);
 		CHAOS_API int32 AddWeightedFloatValue(const FString& Key, const FVector2f& Value, bool bEnabled, bool bAnimatable);
 
 		template<typename T, TEMPLATE_REQUIRES(TIsWeightedType<T>::Value)>
-		int32 AddValue(const FString& Key, const T& Value, bool bEnabled = true, bool bAnimatable = false) { return AddWeightedValue(Key, Value, Value, bEnabled, bAnimatable); }
+		int32 AddValue(const FString& Key, const T& Value, ECollectionPropertyFlags Flags = ECollectionPropertyFlags::Enabled) { return AddWeightedValue(Key, Value, Value, Flags); }
+		template<typename T, TEMPLATE_REQUIRES(TIsWeightedType<T>::Value)>
+		int32 AddValue(const FString& Key, const T& Value, bool bEnabled, bool bAnimatable = false) { return AddWeightedValue(Key, Value, Value, bEnabled, bAnimatable); }
 
-		CHAOS_API int32 AddStringValue(const FString& Key, const FString& Value, bool bEnabled = true, bool bAnimatable = false);
+		CHAOS_API int32 AddStringValue(const FString& Key, const FString& Value, ECollectionPropertyFlags Flags = ECollectionPropertyFlags::Enabled);
+		CHAOS_API int32 AddStringValue(const FString& Key, const FString& Value, bool bEnabled, bool bAnimatable = false);
 	};
 
 	template<typename T>
@@ -422,6 +450,14 @@ PRAGMA_ENABLE_DEPRECATION_WARNINGS
 		}
 	}
 
+	template<typename T, typename TEnableIf<TIsWeightedType<T>::Value, int>::type>
+	inline int32 FCollectionPropertyMutableFacade::AddWeightedValue(const FString& Key, const T& LowValue, const T& HighValue, ECollectionPropertyFlags Flags)
+	{
+		const int32 KeyIndex = AddProperty(Key, Flags);
+		SetWeightedValue(KeyIndex, LowValue, HighValue);
+		return KeyIndex;
+	}
+
 	template<typename T, typename TEnableIf<TIsWeightedType<T>::Value, int>::type>
 	inline int32 FCollectionPropertyMutableFacade::AddWeightedValue(const FString& Key, const T& LowValue, const T& HighValue, bool bEnabled, bool bAnimatable)
 	{
