Chaos Cloth - Fixed MaxDistance behavior when the weight map is empty.

Before the fix, and without providing a weight map, MaxDistanceBase was ignored and the particles would all be set as dynamic.
Now MaxDistanceBase (Low value) correctly sets a blanket value for the entire cloth without requiring a weight map to be specified.

#rnx

diff --git a/Engine/Plugins/Experimental/ChaosCloth/Source/Chaos/Private/ChaosCloth/ChaosClothingSimulationCloth.cpp b/Engine/Plugins/Experimental/ChaosCloth/Source/Chaos/Private/ChaosCloth/ChaosClothingSimulationCloth.cpp
--- a/Engine/Plugins/Experimental/ChaosCloth/Source/Chaos/Private/ChaosCloth/ChaosClothingSimulationCloth.cpp
+++ b/Engine/Plugins/Experimental/ChaosCloth/Source/Chaos/Private/ChaosCloth/ChaosClothingSimulationCloth.cpp
@@ -180,7 +180,7 @@ PRAGMA_ENABLE_DEPRECATION_WARNINGS
 	auto KinematicPredicate =
 		[MaxDistances, MaxDistanceBase, MaxDistanceRange](int32 Index)
 		{
-			return MaxDistances.IsValidIndex(Index) && MaxDistanceBase + MaxDistanceRange * MaxDistances[Index] < KinematicDistanceThreshold;
+			return (MaxDistances.IsValidIndex(Index) ? MaxDistanceBase + MaxDistanceRange * MaxDistances[Index] : MaxDistanceBase) < KinematicDistanceThreshold;
 		};
 
 	const int32 MassMode = ConfigProperties.GetValue<int32>(TEXT("MassMode"), ClothingSimulationClothDefault::MassMode);
