Chaos Cloth - Fixed air density and wind velocity properties.

#rnx

diff --git a/Engine/Plugins/Experimental/ChaosCloth/Source/Chaos/Private/ChaosCloth/ChaosClothingSimulationCloth.cpp b/Engine/Plugins/Experimental/ChaosCloth/Source/Chaos/Private/ChaosCloth/ChaosClothingSimulationCloth.cpp
--- a/Engine/Plugins/Experimental/ChaosCloth/Source/Chaos/Private/ChaosCloth/ChaosClothingSimulationCloth.cpp
+++ b/Engine/Plugins/Experimental/ChaosCloth/Source/Chaos/Private/ChaosCloth/ChaosClothingSimulationCloth.cpp
@@ -29,7 +29,7 @@ namespace ClothingSimulationClothDefault
 	constexpr float Drag = 0.035f;
 	constexpr float Lift = 0.035f;
 	constexpr float Pressure = 0.f;
-	constexpr float AirDensity = 1.225f;
+	constexpr float AirDensity = 1.225f;  // Air density in kg/m^3
 	constexpr float GravityScale = 1.f;
 	constexpr float GravityZOverride = -980.665f;
 	constexpr float VelocityScale = 0.75f;
@@ -473,6 +473,8 @@ PRAGMA_ENABLE_DEPRECATION_WARNINGS
 		Properties.SetStringValue(LiftIndex, TEXT("Lift"));
 
 		Properties.AddValue(TEXT("FluidDensity"), ClothingSimulationClothDefault::AirDensity, bEnable, bAnimatable);
+
+		Properties.AddValue(TEXT("WindVelocity"), FVector3f(0.f), bEnable, bAnimatable);  // Wind velocity must exist to be animatable
 	}
 
 	// Pressure
@@ -1002,7 +1004,8 @@ PRAGMA_ENABLE_DEPRECATION_WARNINGS
 		const bool bEnableAerodynamics = !(bUsePointBasedWindModel && bPointBasedWindDisablesAccurateWind);
 		Solver->SetWindAndPressureProperties(GroupId, ConfigProperties, LODData[LODIndex]->WeightMaps, bEnableAerodynamics);
 
-		const FVec3f WindVelocity = ConfigProperties.GetValue<FVector3f>(TEXT("WindVelocity"));
+		constexpr float WorldScale = 100.f;  // VelocityField wind is in m/s in the config (same as the wind unit), but cm/s in the solver  TODO: Cleanup the Solver SetWindVelocity functions to be consistent with the unit
+		const FVec3f WindVelocity = ConfigProperties.GetValue<FVector3f>(TEXT("WindVelocity")) * WorldScale;
 		Solver->SetWindVelocity(GroupId, WindVelocity + Solver->GetWindVelocity());
 
 		// Update general solver properties
@@ -1129,11 +1132,11 @@ void FClothingSimulationCloth::SetDampingProperties(FRealSingle InDampingCoeffic
 
 void FClothingSimulationCloth::SetAerodynamicsProperties(const TVec2<FRealSingle>& InDrag, const TVec2<FRealSingle>& InLift, FRealSingle InAirDensity, const FVec3& InWindVelocity)
 {
-	constexpr float WordScale = 100.f;  // Unreal's world unit is the cm
+	constexpr float WordScale = 100.f;  // Interactor values are setup in engine scale, kg/cm^3 for air density and cm/s for wind velocity, but the properties are stored in kg/m^3 and m/s in the UI.
 	Config->GetProperties().SetWeightedFloatValue(TEXT("Drag"), FVector2f(InDrag));
 	Config->GetProperties().SetWeightedFloatValue(TEXT("Lift"), FVector2f(InLift));
-	Config->GetProperties().SetValue(TEXT("FluidDensity"), (float)InAirDensity * WordScale);  // AirDensity is here in kg/cm^3 for legacy reason but muust be in kg/m^3 in the config UI
-	Config->GetProperties().SetValue(TEXT("WindVelocity"), FVector3f(InWindVelocity));
+	Config->GetProperties().SetValue(TEXT("FluidDensity"), (float)InAirDensity * FMath::Cube(WordScale));
+	Config->GetProperties().SetValue(TEXT("WindVelocity"), FVector3f(InWindVelocity) / WordScale);
 }
 
 void FClothingSimulationCloth::SetPressureProperties(const TVec2<FRealSingle>& InPressure)
diff --git a/Engine/Plugins/Experimental/ChaosCloth/Source/Chaos/Private/ChaosCloth/ChaosClothingSimulationConfig.cpp b/Engine/Plugins/Experimental/ChaosCloth/Source/Chaos/Private/ChaosCloth/ChaosClothingSimulationConfig.cpp
--- a/Engine/Plugins/Experimental/ChaosCloth/Source/Chaos/Private/ChaosCloth/ChaosClothingSimulationConfig.cpp
+++ b/Engine/Plugins/Experimental/ChaosCloth/Source/Chaos/Private/ChaosCloth/ChaosClothingSimulationConfig.cpp
@@ -161,8 +161,10 @@ namespace Chaos
 				Properties->SetWeightedValue(LiftIndex, ClothConfig->Lift.Low, ClothConfig->Lift.High);
 				Properties->SetStringValue(LiftIndex, TEXT("Lift"));
 
-				constexpr float AirDensity = 1.225f;
+				constexpr float AirDensity = 1.225f;  // Air density in kg/m^3
 				Properties->AddValue(TEXT("FluidDensity"), AirDensity, AnimatablePropertyFlags);
+
+				Properties->AddValue(TEXT("WindVelocity"), FVector3f(0.f), AnimatablePropertyFlags);  // Wind velocity must exist to be animatable
 			}
 
 			// Pressure
diff --git a/Engine/Plugins/Experimental/ChaosClothAssetEditor/Source/ChaosClothAssetDataflowNodes/Private/ChaosClothAsset/SimulationAerodynamicsConfigNode.cpp b/Engine/Plugins/Experimental/ChaosClothAssetEditor/Source/ChaosClothAssetDataflowNodes/Private/ChaosClothAsset/SimulationAerodynamicsConfigNode.cpp
--- a/Engine/Plugins/Experimental/ChaosClothAssetEditor/Source/ChaosClothAssetDataflowNodes/Private/ChaosClothAsset/SimulationAerodynamicsConfigNode.cpp
+++ b/Engine/Plugins/Experimental/ChaosClothAssetEditor/Source/ChaosClothAssetDataflowNodes/Private/ChaosClothAsset/SimulationAerodynamicsConfigNode.cpp
@@ -19,4 +19,5 @@ void FChaosClothAssetSimulationAerodynamicsConfigNode::AddProperties(Dataflow::F
 	UE_CHAOS_CLOTHASSET_SIMULATIONCONFIG_SETPROPERTY(FluidDensity);
 	UE_CHAOS_CLOTHASSET_SIMULATIONCONFIG_SETPROPERTYWEIGHTED(Drag);
 	UE_CHAOS_CLOTHASSET_SIMULATIONCONFIG_SETPROPERTYWEIGHTED(Lift);
+	UE_CHAOS_CLOTHASSET_SIMULATIONCONFIG_SETPROPERTY(WindVelocity);
 }
diff --git a/Engine/Plugins/Experimental/ChaosClothAssetEditor/Source/ChaosClothAssetDataflowNodes/Public/ChaosClothAsset/SimulationAerodynamicsConfigNode.h b/Engine/Plugins/Experimental/ChaosClothAssetEditor/Source/ChaosClothAssetDataflowNodes/Public/ChaosClothAsset/SimulationAerodynamicsConfigNode.h
--- a/Engine/Plugins/Experimental/ChaosClothAssetEditor/Source/ChaosClothAssetDataflowNodes/Public/ChaosClothAsset/SimulationAerodynamicsConfigNode.h
+++ b/Engine/Plugins/Experimental/ChaosClothAssetEditor/Source/ChaosClothAssetDataflowNodes/Public/ChaosClothAsset/SimulationAerodynamicsConfigNode.h
@@ -40,6 +40,13 @@ public:
 	UPROPERTY(EditAnywhere, Category = "Aerodynamics", Meta = (UIMin = "0", UIMax = "1", ClampMin = "0", ClampMax = "10"))
 	FChaosClothAssetWeightedValue Lift = { true, 0.035f, 1.f, TEXT("Lift") };
 
+	/**
+	 * The fixed wind velocity [m/s] for this asset.
+	 * For reference a wind gust is above 8m/s (18mph).
+	 */
+	UPROPERTY(EditAnywhere, Category = "Aerodynamics", Meta = (UIMin = "0", UIMax = "10"))
+	FVector3f WindVelocity = { 0.f, 0.f, 0.f };
+
 	FChaosClothAssetSimulationAerodynamicsConfigNode(const Dataflow::FNodeParameters& InParam, FGuid InGuid = FGuid::NewGuid());
 
 private:
