Chaos Cloth: Support multiple property LODs.  Also added forced lod support to the cloth editor.

Now that each LOD has its own ManagedArrayCollection, they support authoring different properties. This threads through support for per-lod collections at the ChaosClothingSimulationConfig level. Legacy ClothAsset (i.e., those on SKM) continue to have a single collection. No change in behavior is expected.

This has support for multiple solver-level LODs as well. In practice, we only have multiple LODs with the new Chaos Cloth Assets, and there is currently one cloth per solver, so managing this is trivial.

I did deprecate the methods to set properties on the FClothingSimulationCloth objects (where properties used to live) in favor or getting the config and setting the properties there. I did the same for the FClothingSimulationSolver (including making it so the solver always has a config). This only affected the blueprint interactors.

diff --git a/Engine/Plugins/Experimental/ChaosCloth/Source/Chaos/Private/ChaosCloth/ChaosClothingSimulation.cpp b/Engine/Plugins/Experimental/ChaosCloth/Source/Chaos/Private/ChaosCloth/ChaosClothingSimulation.cpp
--- a/Engine/Plugins/Experimental/ChaosCloth/Source/Chaos/Private/ChaosCloth/ChaosClothingSimulation.cpp
+++ b/Engine/Plugins/Experimental/ChaosCloth/Source/Chaos/Private/ChaosCloth/ChaosClothingSimulation.cpp
@@ -440,21 +440,26 @@ void FClothingSimulation::UpdateSimulationFromSharedSimConfig()
 		constexpr bool bUseLegacyConfig = true;  // Make the config a legacy cloth config, so that the constraints disable themselves with missing masks, ...etc.
 		SolverConfig->Initialize(nullptr, ClothSharedSimConfig, bUseLegacyConfig);
 	}
+	else
+	{
+		// This will cause the solver to create a default config
+		Solver->SetConfig(nullptr);
+	}
 }
 
 void FClothingSimulation::SetNumIterations(int32 InNumIterations)
 {
-	Solver->SetNumIterations(InNumIterations);
+	Solver->GetConfig()->GetProperties(Solver->GetSolverLOD()).SetValue(TEXT("NumIterations"), InNumIterations);
 }
 
 void FClothingSimulation::SetMaxNumIterations(int32 MaxNumIterations)
 {
-	Solver->SetMaxNumIterations(MaxNumIterations);
+	Solver->GetConfig()->GetProperties(Solver->GetSolverLOD()).SetValue(TEXT("MaxNumIterations"), MaxNumIterations);
 }
 
 void FClothingSimulation::SetNumSubsteps(int32 InNumSubsteps)
 {
-	Solver->SetNumSubsteps(InNumSubsteps);
+	Solver->GetConfig()->GetProperties(Solver->GetSolverLOD()).SetValue(TEXT("NumSubsteps"), InNumSubsteps);
 }
 
 bool FClothingSimulation::ShouldSimulate() const
diff --git a/Engine/Plugins/Experimental/ChaosCloth/Source/Chaos/Private/ChaosCloth/ChaosClothingSimulationCloth.cpp b/Engine/Plugins/Experimental/ChaosCloth/Source/Chaos/Private/ChaosCloth/ChaosClothingSimulationCloth.cpp
--- a/Engine/Plugins/Experimental/ChaosCloth/Source/Chaos/Private/ChaosCloth/ChaosClothingSimulationCloth.cpp
+++ b/Engine/Plugins/Experimental/ChaosCloth/Source/Chaos/Private/ChaosCloth/ChaosClothingSimulationCloth.cpp
@@ -54,6 +54,7 @@ struct FClothingSimulationCloth::FLODData
 	// Per Solver data
 	struct FSolverData
 	{
+		int32 LODIndex;
 		int32 Offset;
 		FTriangleMesh TriangleMesh;  // TODO: Triangle Mesh shouldn't really be solver dependent (ie not use an offset)
 	};
@@ -109,6 +110,7 @@ void FClothingSimulationCloth::FLODData::Add(FClothingSimulationSolver* Solver,
 	// Add a new solver data chunk
 	check(!SolverData.Find(Solver));
 	FSolverData& SolverDatum = SolverData.Add(Solver);
+	SolverDatum.LODIndex = InLODIndex;
 	int32& Offset = SolverDatum.Offset;
 	FTriangleMesh& TriangleMesh = SolverDatum.TriangleMesh;
 
@@ -153,7 +155,7 @@ PRAGMA_ENABLE_DEPRECATION_WARNINGS
 
 	// Retrieve config properties
 	check(Cloth->Config);
-	const Softs::FCollectionPropertyFacade& ConfigProperties = Cloth->Config->GetProperties();
+	const Softs::FCollectionPropertyFacade& ConfigProperties = Cloth->Config->GetProperties(InLODIndex);
 
 	// Retrieve MaxDistance information (weight map and Low/High values)
 	int32 MaxDistancePropertyKeyIndex;
@@ -234,9 +236,10 @@ void FClothingSimulationCloth::FLODData::Remove(FClothingSimulationSolver* Solve
 void FClothingSimulationCloth::FLODData::Update(FClothingSimulationSolver* Solver, FClothingSimulationCloth* Cloth)
 {
 	check(Solver);
-	check(Cloth);
+	check(Cloth); 
+	const FSolverData& SolverDatum = SolverData.FindChecked(Solver);
 
-	const int32 Offset = SolverData.FindChecked(Solver).Offset;
+	const int32 Offset = SolverDatum.Offset;
 	check(Offset != INDEX_NONE);
 
 	// Update the animatable constraint parameters
@@ -247,7 +250,7 @@ PRAGMA_DISABLE_DEPRECATION_WARNINGS  // TODO: CHAOS_IS_CLOTHINGSIMULATIONMESH_AB
 	const Softs::FSolverReal MeshScale = Cloth->Mesh->GetScale();
 PRAGMA_ENABLE_DEPRECATION_WARNINGS
 	const Softs::FSolverReal MaxDistancesScale = (Softs::FSolverReal)Cloth->MaxDistancesMultiplier;
-	ClothConstraints.Update(Cloth->Config->GetProperties(), WeightMaps, MeshScale, MaxDistancesScale);
+	ClothConstraints.Update(Cloth->Config->GetProperties(SolverDatum.LODIndex), WeightMaps, MeshScale, MaxDistancesScale);
 }
 
 void FClothingSimulationCloth::FLODData::Enable(FClothingSimulationSolver* Solver, bool bEnable) const
@@ -518,7 +521,9 @@ PRAGMA_ENABLE_DEPRECATION_WARNINGS
 		Properties.AddValue(TEXT("UseLegacyBackstop"), bInUseLegacyBackstop);
 	}
 
+	PRAGMA_DISABLE_DEPRECATION_WARNINGS
 	Config = new FClothingSimulationConfig(PropertyCollection);
+	PRAGMA_ENABLE_DEPRECATION_WARNINGS
 
 	// Set mesh and colliders
 	SetMesh(InMesh);
@@ -545,8 +550,6 @@ void FClothingSimulationCloth::SetMesh(FClothingSimulationMesh* InMesh)
 	// Reset LODs
 	const int32 NumLODs = Mesh ? Mesh->GetNumLODs() : 0;
 	LODData.Reset(NumLODs);
-
-	const bool bUseGeodesicTethers = Config->GetProperties().GetValue<bool>(TEXT("UseGeodesicTethers"), ClothingSimulationClothDefault::bUseGeodesicTethers);
 	for (int32 LODIndex = 0; LODIndex < NumLODs; ++LODIndex)
 	{
 		// Regenerate LOD weight maps lookup map
@@ -563,6 +566,8 @@ void FClothingSimulationCloth::SetMesh(FClothingSimulationMesh* InMesh)
 					TConstArrayView<FRealSingle>());
 		}
 
+		const bool bUseGeodesicTethers = Config->GetProperties(LODIndex).GetValue<bool>(TEXT("UseGeodesicTethers"), ClothingSimulationClothDefault::bUseGeodesicTethers);
+
 		// Add LOD data
 		LODData.Add(MakeUnique<FLODData>(
 			Mesh->GetNumPoints(LODIndex),
@@ -602,7 +607,9 @@ void FClothingSimulationCloth::SetConfig(FClothingSimulationConfig* InConfig)
 	{
 		// Create a default empty config object for coherence
 		PropertyCollection = MakeShared<FManagedArrayCollection>();
-		Config = new FClothingSimulationConfig(PropertyCollection);
+		PRAGMA_DISABLE_DEPRECATION_WARNINGS
+		Config = new FClothingSimulationConfig(PropertyCollection); 
+		PRAGMA_ENABLE_DEPRECATION_WARNINGS
 	}
 }
 
@@ -750,7 +757,7 @@ TVec3<FRealSingle> FClothingSimulationCloth::GetGravity(const FClothingSimulatio
 {
 	check(Solver);
 	check(Config);
-	const Softs::FCollectionPropertyFacade& ConfigProperties = Config->GetProperties();
+	const Softs::FCollectionPropertyFacade& ConfigProperties = Config->GetProperties(GetLODIndex(Solver));
 
 	const bool bUseGravityOverride = ConfigProperties.GetValue<bool>(TEXT("UseGravityOverride"));
 	const TVec3<FRealSingle> GravityOverride = (TVec3<FRealSingle>)ConfigProperties.GetValue<FVector3f>(TEXT("GravityOverride"), FVector3f(0.f, 0.f, ClothingSimulationClothDefault::GravityZOverride));
@@ -811,7 +818,7 @@ TConstArrayView<FRealSingle> FClothingSimulationCloth::GetWeightMapByName(const
 TConstArrayView<FRealSingle> FClothingSimulationCloth::GetWeightMapByProperty(const FClothingSimulationSolver* Solver, const FString& Property) const
 {
 	check(Config);
-	const FString PropertyString = Config->GetProperties().GetStringValue(Property);
+	const FString PropertyString = Config->GetProperties(GetLODIndex(Solver)).GetStringValue(Property);
 	return GetWeightMapByName(Solver, PropertyString);
 }
 
@@ -896,14 +903,17 @@ PRAGMA_DISABLE_DEPRECATION_WARNINGS  // TODO: CHAOS_IS_CLOTHINGSIMULATIONMESH_AB
 PRAGMA_ENABLE_DEPRECATION_WARNINGS
 
 	// Update all LODs dirty properties, since it is easier done than re-updating all properties when switching LODs
-	for (TUniquePtr<FLODData>& LODDatum : LODData)
+	if (Config->IsLegacySingleLOD())
 	{
-		LODDatum->Update(Solver, this);
+		for (TUniquePtr<FLODData>& LODDatum : LODData)
+		{
+			LODDatum->Update(Solver, this);
+		}
 	}
 
 	// Retrieve config
 	check(Config);
-	const Softs::FCollectionPropertyFacade& ConfigProperties = Config->GetProperties();
+	const Softs::FCollectionPropertyFacade& ConfigProperties = Config->GetProperties(LODIndex);
 
 	// LOD Switching
 	if (LODIndex != PrevLODIndex)
@@ -946,6 +956,13 @@ PRAGMA_ENABLE_DEPRECATION_WARNINGS
 	// Update Cloth group parameters
 	if (LODIndex != INDEX_NONE)
 	{
+		// Update Cloth group parameters  TODO: Cloth groups should exist as their own node object so that they can be used by several cloth objects
+		// This was updated above for all LODs if single lod config.
+		if (!Config->IsLegacySingleLOD())
+		{
+			LODData[LODIndex]->Update(Solver, this);
+		}
+
 		// TODO: Move all groupID updates out of the cloth update to allow to use of the same GroupId with different cloths
 
 		// Set the reference input velocity and deal with teleport & reset; external forces depends on these values, so they must be initialized before then
@@ -1115,13 +1132,13 @@ void FClothingSimulationCloth::SetCollisionProperties(FRealSingle InCollisionThi
 {
 	Config->GetProperties().SetValue(TEXT("CollisionThickness"), (float)InCollisionThickness);
 	Config->GetProperties().SetValue(TEXT("FrictionCoefficient"), (float)InFrictionCoefficient);
-	Config->GetProperties().SetValue(TEXT("bUseCCD"), bInUseCCD);
+	Config->GetProperties().SetValue(TEXT("UseCCD"), bInUseCCD);
 	Config->GetProperties().SetValue(TEXT("SelfCollisionThickness"), (float)InSelfCollisionThickness);
 }
 
 void FClothingSimulationCloth::SetBackstopProperties(bool bInEnableBackstop)
 {
-	Config->GetProperties().SetEnabled(TEXT("BackstopScale"), bInEnableBackstop);  // BackstopScale controls whether the backstop is enabled or not
+	Config->GetProperties().SetEnabled(TEXT("BackstopDistance"), bInEnableBackstop);  // BackstopScale controls whether the backstop is enabled or not
 }
 
 void FClothingSimulationCloth::SetDampingProperties(FRealSingle InDampingCoefficient, FRealSingle InLocalDampingCoefficient)
@@ -1132,11 +1149,11 @@ void FClothingSimulationCloth::SetDampingProperties(FRealSingle InDampingCoeffic
 
 void FClothingSimulationCloth::SetAerodynamicsProperties(const TVec2<FRealSingle>& InDrag, const TVec2<FRealSingle>& InLift, FRealSingle InAirDensity, const FVec3& InWindVelocity)
 {
-	constexpr float WordScale = 100.f;  // Interactor values are setup in engine scale, kg/cm^3 for air density and cm/s for wind velocity, but the properties are stored in kg/m^3 and m/s in the UI.
+	constexpr float WorldScale = 100.f;  // Interactor values are setup in engine scale, kg/cm^3 for air density and cm/s for wind velocity, but the properties are stored in kg/m^3 and m/s in the UI.
 	Config->GetProperties().SetWeightedFloatValue(TEXT("Drag"), FVector2f(InDrag));
 	Config->GetProperties().SetWeightedFloatValue(TEXT("Lift"), FVector2f(InLift));
-	Config->GetProperties().SetValue(TEXT("FluidDensity"), (float)InAirDensity * FMath::Cube(WordScale));
-	Config->GetProperties().SetValue(TEXT("WindVelocity"), FVector3f(InWindVelocity) / WordScale);
+	Config->GetProperties().SetValue(TEXT("FluidDensity"), (float)InAirDensity * FMath::Cube(WorldScale));
+	Config->GetProperties().SetValue(TEXT("WindVelocity"), FVector3f(InWindVelocity) / WorldScale);
 }
 
 void FClothingSimulationCloth::SetPressureProperties(const TVec2<FRealSingle>& InPressure)
diff --git a/Engine/Plugins/Experimental/ChaosCloth/Source/Chaos/Private/ChaosCloth/ChaosClothingSimulationConfig.cpp b/Engine/Plugins/Experimental/ChaosCloth/Source/Chaos/Private/ChaosCloth/ChaosClothingSimulationConfig.cpp
--- a/Engine/Plugins/Experimental/ChaosCloth/Source/Chaos/Private/ChaosCloth/ChaosClothingSimulationConfig.cpp
+++ b/Engine/Plugins/Experimental/ChaosCloth/Source/Chaos/Private/ChaosCloth/ChaosClothingSimulationConfig.cpp
@@ -8,16 +8,19 @@
 namespace Chaos
 {
 	FClothingSimulationConfig::FClothingSimulationConfig()
-		: PropertyCollection(MakeShared<FManagedArrayCollection>())
-		, Properties(MakeUnique<Softs::FCollectionPropertyMutableFacade>(PropertyCollection))
 	{
 	}
 
 	FClothingSimulationConfig::FClothingSimulationConfig(const TSharedPtr<const FManagedArrayCollection>& InPropertyCollection)
-		: PropertyCollection(MakeShared<FManagedArrayCollection>())
-		, Properties(MakeUnique<Softs::FCollectionPropertyMutableFacade>(PropertyCollection))
 	{
+		PRAGMA_DISABLE_DEPRECATION_WARNINGS
 		Initialize(InPropertyCollection);
+		PRAGMA_ENABLE_DEPRECATION_WARNINGS
+	}
+
+	FClothingSimulationConfig::FClothingSimulationConfig(const TArray<TSharedPtr<const FManagedArrayCollection>>& InPropertyCollections)
+	{
+		Initialize(InPropertyCollections);
 	}
 
 	FClothingSimulationConfig::~FClothingSimulationConfig() = default;
@@ -31,15 +34,20 @@ namespace Chaos
 			::Chaos::Softs::ECollectionPropertyFlags::Animatable;
 
 		// Clear all properties
-		PropertyCollection->Reset();
-		Properties->DefineSchema();
+		PropertyCollections.Reset(1);
+		Properties.Reset(1);
+
+		bIsLegacySingleLOD = true;
+		TSharedPtr<FManagedArrayCollection>& PropertyCollection = PropertyCollections.Add_GetRef(MakeShared<FManagedArrayCollection>());
+		TUniquePtr<Softs::FCollectionPropertyMutableFacade>& Property = Properties.Add_GetRef(MakeUnique<Softs::FCollectionPropertyMutableFacade>(PropertyCollection));
+		Property->DefineSchema();
 
 		// Solver properties
 		if (ClothSharedConfig)
 		{
-			Properties->AddValue(TEXT("NumIterations"), ClothSharedConfig->IterationCount, AnimatablePropertyFlags);
-			Properties->AddValue(TEXT("MaxNumIterations"), ClothSharedConfig->MaxIterationCount, AnimatablePropertyFlags);
-			Properties->AddValue(TEXT("NumSubsteps"), ClothSharedConfig->SubdivisionCount, AnimatablePropertyFlags);
+			Property->AddValue(TEXT("NumIterations"), ClothSharedConfig->IterationCount, AnimatablePropertyFlags);
+			Property->AddValue(TEXT("MaxNumIterations"), ClothSharedConfig->MaxIterationCount, AnimatablePropertyFlags);
+			Property->AddValue(TEXT("NumSubsteps"), ClothSharedConfig->SubdivisionCount, AnimatablePropertyFlags);
 		}
 
 		// Cloth properties
@@ -61,17 +69,17 @@ namespace Chaos
 					MassValue = ClothConfig->Density;
 					break;
 				}
-				Properties->AddValue(TEXT("MassMode"), (int32)ClothConfig->MassMode, NonAnimatablePropertyFlags);
-				Properties->AddValue(TEXT("MassValue"), MassValue, NonAnimatablePropertyFlags);
-				Properties->AddValue(TEXT("MinPerParticleMass"), ClothConfig->MinPerParticleMass, NonAnimatablePropertyFlags);
+				Property->AddValue(TEXT("MassMode"), (int32)ClothConfig->MassMode, NonAnimatablePropertyFlags);
+				Property->AddValue(TEXT("MassValue"), MassValue, NonAnimatablePropertyFlags);
+				Property->AddValue(TEXT("MinPerParticleMass"), ClothConfig->MinPerParticleMass, NonAnimatablePropertyFlags);
 			}
 
 			// Edge constraint
 			if (ClothConfig->EdgeStiffnessWeighted.Low > 0.f || ClothConfig->EdgeStiffnessWeighted.High > 0.f)
 			{
-				const int32 EdgeSpringStiffnessIndex = Properties->AddProperty(TEXT("EdgeSpringStiffness"), AnimatablePropertyFlags);
-				Properties->SetWeightedValue(EdgeSpringStiffnessIndex, ClothConfig->EdgeStiffnessWeighted.Low, ClothConfig->EdgeStiffnessWeighted.High);
-				Properties->SetStringValue(EdgeSpringStiffnessIndex, TEXT("EdgeStiffness"));
+				const int32 EdgeSpringStiffnessIndex = Property->AddProperty(TEXT("EdgeSpringStiffness"), AnimatablePropertyFlags);
+				Property->SetWeightedValue(EdgeSpringStiffnessIndex, ClothConfig->EdgeStiffnessWeighted.Low, ClothConfig->EdgeStiffnessWeighted.High);
+				Property->SetStringValue(EdgeSpringStiffnessIndex, TEXT("EdgeStiffness"));
 			}
 
 			// Bending constraint
@@ -80,151 +88,155 @@ namespace Chaos
 			{
 				if (ClothConfig->bUseBendingElements)
 				{
-					const int32 BendingElementStiffnessIndex = Properties->AddProperty(TEXT("BendingElementStiffness"), AnimatablePropertyFlags);
-					Properties->SetWeightedValue(BendingElementStiffnessIndex, ClothConfig->BendingStiffnessWeighted.Low, ClothConfig->BendingStiffnessWeighted.High);
-					Properties->SetStringValue(BendingElementStiffnessIndex, TEXT("BendingStiffness"));
+					const int32 BendingElementStiffnessIndex = Property->AddProperty(TEXT("BendingElementStiffness"), AnimatablePropertyFlags);
+					Property->SetWeightedValue(BendingElementStiffnessIndex, ClothConfig->BendingStiffnessWeighted.Low, ClothConfig->BendingStiffnessWeighted.High);
+					Property->SetStringValue(BendingElementStiffnessIndex, TEXT("BendingStiffness"));
 
-					Properties->AddValue(TEXT("BucklingRatio"), ClothConfig->BucklingRatio, NonAnimatablePropertyFlags);
+					Property->AddValue(TEXT("BucklingRatio"), ClothConfig->BucklingRatio, NonAnimatablePropertyFlags);
 
 					if (ClothConfig->BucklingStiffnessWeighted.Low > 0.f || ClothConfig->BucklingStiffnessWeighted.High > 0.f)
 					{
-						const int32 BucklingStiffnessIndex = Properties->AddProperty(TEXT("BucklingStiffness"), AnimatablePropertyFlags);
-						Properties->SetWeightedValue(BucklingStiffnessIndex, ClothConfig->BucklingStiffnessWeighted.Low, ClothConfig->BucklingStiffnessWeighted.High);
-						Properties->SetStringValue(BucklingStiffnessIndex, TEXT("BucklingStiffness"));
+						const int32 BucklingStiffnessIndex = Property->AddProperty(TEXT("BucklingStiffness"), AnimatablePropertyFlags);
+						Property->SetWeightedValue(BucklingStiffnessIndex, ClothConfig->BucklingStiffnessWeighted.Low, ClothConfig->BucklingStiffnessWeighted.High);
+						Property->SetStringValue(BucklingStiffnessIndex, TEXT("BucklingStiffness"));
 					}
 				}
 				else  // Not using bending elements
 				{
-					const int32 BendingSpringStiffnessIndex = Properties->AddProperty(TEXT("BendingSpringStiffness"), AnimatablePropertyFlags);
-					Properties->SetWeightedValue(BendingSpringStiffnessIndex, ClothConfig->BendingStiffnessWeighted.Low, ClothConfig->BendingStiffnessWeighted.High);
-					Properties->SetStringValue(BendingSpringStiffnessIndex, TEXT("BendingStiffness"));
+					const int32 BendingSpringStiffnessIndex = Property->AddProperty(TEXT("BendingSpringStiffness"), AnimatablePropertyFlags);
+					Property->SetWeightedValue(BendingSpringStiffnessIndex, ClothConfig->BendingStiffnessWeighted.Low, ClothConfig->BendingStiffnessWeighted.High);
+					Property->SetStringValue(BendingSpringStiffnessIndex, TEXT("BendingStiffness"));
 				}
 			}
 
 			// Area constraint
 			if (ClothConfig->AreaStiffnessWeighted.Low > 0.f || ClothConfig->AreaStiffnessWeighted.High > 0.f)
 			{
-				const int32 AreaSpringStiffnessIndex = Properties->AddProperty(TEXT("AreaSpringStiffness"), AnimatablePropertyFlags);
-				Properties->SetWeightedValue(AreaSpringStiffnessIndex, ClothConfig->AreaStiffnessWeighted.Low, ClothConfig->AreaStiffnessWeighted.High);
-				Properties->SetStringValue(AreaSpringStiffnessIndex, TEXT("AreaStiffness"));
+				const int32 AreaSpringStiffnessIndex = Property->AddProperty(TEXT("AreaSpringStiffness"), AnimatablePropertyFlags);
+				Property->SetWeightedValue(AreaSpringStiffnessIndex, ClothConfig->AreaStiffnessWeighted.Low, ClothConfig->AreaStiffnessWeighted.High);
+				Property->SetStringValue(AreaSpringStiffnessIndex, TEXT("AreaStiffness"));
 			}
 
 			// Long range attachment
 			if (ClothConfig->TetherStiffness.Low > 0.f || ClothConfig->TetherStiffness.High > 0.f)
 			{
-				Properties->AddValue(TEXT("UseGeodesicTethers"), ClothConfig->bUseGeodesicDistance, NonAnimatablePropertyFlags);
+				Property->AddValue(TEXT("UseGeodesicTethers"), ClothConfig->bUseGeodesicDistance, NonAnimatablePropertyFlags);
 
-				const int32 TetherStiffnessIndex = Properties->AddProperty(TEXT("TetherStiffness"), AnimatablePropertyFlags);
-				Properties->SetWeightedValue(TetherStiffnessIndex, ClothConfig->TetherStiffness.Low, ClothConfig->TetherStiffness.High);
-				Properties->SetStringValue(TetherStiffnessIndex, TEXT("TetherStiffness"));
+				const int32 TetherStiffnessIndex = Property->AddProperty(TEXT("TetherStiffness"), AnimatablePropertyFlags);
+				Property->SetWeightedValue(TetherStiffnessIndex, ClothConfig->TetherStiffness.Low, ClothConfig->TetherStiffness.High);
+				Property->SetStringValue(TetherStiffnessIndex, TEXT("TetherStiffness"));
 
-				const int32 TetherScaleIndex = Properties->AddProperty(TEXT("TetherScale"), AnimatablePropertyFlags);
-				Properties->SetWeightedValue(TetherScaleIndex, ClothConfig->TetherScale.Low, ClothConfig->TetherScale.High);
-				Properties->SetStringValue(TetherScaleIndex, TEXT("TetherScale"));
+				const int32 TetherScaleIndex = Property->AddProperty(TEXT("TetherScale"), AnimatablePropertyFlags);
+				Property->SetWeightedValue(TetherScaleIndex, ClothConfig->TetherScale.Low, ClothConfig->TetherScale.High);
+				Property->SetStringValue(TetherScaleIndex, TEXT("TetherScale"));
 			}
 
 			// AnimDrive
 			if (ClothConfig->AnimDriveStiffness.Low > 0.f || ClothConfig->AnimDriveStiffness.High > 0.f)
 			{
-				const int32 AnimDriveStiffnessIndex = Properties->AddProperty(TEXT("AnimDriveStiffness"), AnimatablePropertyFlags);
-				Properties->SetWeightedValue(AnimDriveStiffnessIndex, ClothConfig->AnimDriveStiffness.Low, ClothConfig->AnimDriveStiffness.High);
-				Properties->SetStringValue(AnimDriveStiffnessIndex, TEXT("AnimDriveStiffness"));
+				const int32 AnimDriveStiffnessIndex = Property->AddProperty(TEXT("AnimDriveStiffness"), AnimatablePropertyFlags);
+				Property->SetWeightedValue(AnimDriveStiffnessIndex, ClothConfig->AnimDriveStiffness.Low, ClothConfig->AnimDriveStiffness.High);
+				Property->SetStringValue(AnimDriveStiffnessIndex, TEXT("AnimDriveStiffness"));
 
-				const int32 AnimDriveDampingIndex = Properties->AddProperty(TEXT("AnimDriveDamping"), AnimatablePropertyFlags);
-				Properties->SetWeightedValue(AnimDriveDampingIndex, ClothConfig->AnimDriveDamping.Low, ClothConfig->AnimDriveDamping.High);
-				Properties->SetStringValue(AnimDriveDampingIndex, TEXT("AnimDriveDamping"));
+				const int32 AnimDriveDampingIndex = Property->AddProperty(TEXT("AnimDriveDamping"), AnimatablePropertyFlags);
+				Property->SetWeightedValue(AnimDriveDampingIndex, ClothConfig->AnimDriveDamping.Low, ClothConfig->AnimDriveDamping.High);
+				Property->SetStringValue(AnimDriveDampingIndex, TEXT("AnimDriveDamping"));
 			}
 
 			// Gravity
 			{
-				Properties->AddValue(TEXT("GravityScale"), ClothConfig->GravityScale, AnimatablePropertyFlags);
-				Properties->AddValue(TEXT("UseGravityOverride"), ClothConfig->bUseGravityOverride, AnimatablePropertyFlags);
-				Properties->AddValue(TEXT("GravityOverride"), FVector3f(ClothConfig->Gravity), AnimatablePropertyFlags);
+				Property->AddValue(TEXT("GravityScale"), ClothConfig->GravityScale, AnimatablePropertyFlags);
+				Property->AddValue(TEXT("UseGravityOverride"), ClothConfig->bUseGravityOverride, AnimatablePropertyFlags);
+				Property->AddValue(TEXT("GravityOverride"), FVector3f(ClothConfig->Gravity), AnimatablePropertyFlags);
 			}
 
 			// Velocity scale
 			{
-				Properties->AddValue(TEXT("LinearVelocityScale"), FVector3f(ClothConfig->LinearVelocityScale), AnimatablePropertyFlags);
-				Properties->AddValue(TEXT("AngularVelocityScale"), ClothConfig->AngularVelocityScale, AnimatablePropertyFlags);
-				Properties->AddValue(TEXT("FictitiousAngularScale"), ClothConfig->FictitiousAngularScale, AnimatablePropertyFlags);
+				Property->AddValue(TEXT("LinearVelocityScale"), FVector3f(ClothConfig->LinearVelocityScale), AnimatablePropertyFlags);
+				Property->AddValue(TEXT("AngularVelocityScale"), ClothConfig->AngularVelocityScale, AnimatablePropertyFlags);
+				Property->AddValue(TEXT("FictitiousAngularScale"), ClothConfig->FictitiousAngularScale, AnimatablePropertyFlags);
 			}
 
 			// Aerodynamics
-			Properties->AddValue(TEXT("UsePointBasedWindModel"), ClothConfig->bUsePointBasedWindModel, NonAnimatablePropertyFlags);
-			if (!ClothConfig->bUsePointBasedWindModel && (ClothConfig->Drag.Low > 0.f || ClothConfig->Drag.High > 0.f || ClothConfig->Lift.Low > 0.f || ClothConfig->Lift.High > 0.f))
+			Property->AddValue(TEXT("UsePointBasedWindModel"), ClothConfig->bUsePointBasedWindModel, NonAnimatablePropertyFlags);
 			{
-				const int32 DragIndex = Properties->AddProperty(TEXT("Drag"), AnimatablePropertyFlags);
-				Properties->SetWeightedValue(DragIndex, ClothConfig->Drag.Low, ClothConfig->Drag.High);
-				Properties->SetStringValue(DragIndex, TEXT("Drag"));
+				const int32 DragIndex = Property->AddProperty(TEXT("Drag"), AnimatablePropertyFlags);
+				Property->SetWeightedValue(DragIndex, ClothConfig->Drag.Low, ClothConfig->Drag.High);
+				Property->SetStringValue(DragIndex, TEXT("Drag"));
 
-				const int32 LiftIndex = Properties->AddProperty(TEXT("Lift"), AnimatablePropertyFlags);
-				Properties->SetWeightedValue(LiftIndex, ClothConfig->Lift.Low, ClothConfig->Lift.High);
-				Properties->SetStringValue(LiftIndex, TEXT("Lift"));
+				const int32 LiftIndex = Property->AddProperty(TEXT("Lift"), AnimatablePropertyFlags);
+				Property->SetWeightedValue(LiftIndex, ClothConfig->Lift.Low, ClothConfig->Lift.High);
+				Property->SetStringValue(LiftIndex, TEXT("Lift"));
 
 				constexpr float AirDensity = 1.225f;  // Air density in kg/m^3
-				Properties->AddValue(TEXT("FluidDensity"), AirDensity, AnimatablePropertyFlags);
+				Property->AddValue(TEXT("FluidDensity"), AirDensity, AnimatablePropertyFlags);
 
-				Properties->AddValue(TEXT("WindVelocity"), FVector3f(0.f), AnimatablePropertyFlags);  // Wind velocity must exist to be animatable
+				Property->AddValue(TEXT("WindVelocity"), FVector3f(0.f), AnimatablePropertyFlags);  // Wind velocity must exist to be animatable
 			}
 
 			// Pressure
-			if (ClothConfig->Pressure.Low != 0.f || ClothConfig->Pressure.High != 0.f)
 			{
-				const int32 PressureIndex = Properties->AddProperty(TEXT("Pressure"), AnimatablePropertyFlags);
-				Properties->SetWeightedValue(PressureIndex, ClothConfig->Pressure.Low, ClothConfig->Pressure.High);
-				Properties->SetStringValue(PressureIndex, TEXT("Pressure"));
+				const int32 PressureIndex = Property->AddProperty(TEXT("Pressure"), AnimatablePropertyFlags);
+				Property->SetWeightedValue(PressureIndex, ClothConfig->Pressure.Low, ClothConfig->Pressure.High);
+				Property->SetStringValue(PressureIndex, TEXT("Pressure"));
 			}
 
 			// Damping
-			Properties->AddValue(TEXT("DampingCoefficient"), ClothConfig->DampingCoefficient, AnimatablePropertyFlags);
-			Properties->AddValue(TEXT("LocalDampingCoefficient"), ClothConfig->LocalDampingCoefficient, AnimatablePropertyFlags);
+			Property->AddValue(TEXT("DampingCoefficient"), ClothConfig->DampingCoefficient, AnimatablePropertyFlags);
+			Property->AddValue(TEXT("LocalDampingCoefficient"), ClothConfig->LocalDampingCoefficient, AnimatablePropertyFlags);
 
 			// Collision
-			Properties->AddValue(TEXT("CollisionThickness"), ClothConfig->CollisionThickness, AnimatablePropertyFlags);
-			Properties->AddValue(TEXT("FrictionCoefficient"), ClothConfig->FrictionCoefficient, AnimatablePropertyFlags);
-			Properties->AddValue(TEXT("UseCCD"), ClothConfig->bUseCCD, AnimatablePropertyFlags);
-			Properties->AddValue(TEXT("UseSelfCollisions"), ClothConfig->bUseSelfCollisions, NonAnimatablePropertyFlags);
-			Properties->AddValue(TEXT("SelfCollisionThickness"), ClothConfig->SelfCollisionThickness, NonAnimatablePropertyFlags);
-			Properties->AddValue(TEXT("UseSelfIntersections"), ClothConfig->bUseSelfIntersections, NonAnimatablePropertyFlags);
-			Properties->AddValue(TEXT("SelfCollisionFriction"), ClothConfig->SelfCollisionFriction, NonAnimatablePropertyFlags);
+			Property->AddValue(TEXT("CollisionThickness"), ClothConfig->CollisionThickness, AnimatablePropertyFlags);
+			Property->AddValue(TEXT("FrictionCoefficient"), ClothConfig->FrictionCoefficient, AnimatablePropertyFlags);
+			Property->AddValue(TEXT("UseCCD"), ClothConfig->bUseCCD, AnimatablePropertyFlags);
+			Property->AddValue(TEXT("UseSelfCollisions"), ClothConfig->bUseSelfCollisions, NonAnimatablePropertyFlags);
+			Property->AddValue(TEXT("SelfCollisionThickness"), ClothConfig->SelfCollisionThickness, NonAnimatablePropertyFlags);
+			Property->AddValue(TEXT("UseSelfIntersections"), ClothConfig->bUseSelfIntersections, NonAnimatablePropertyFlags);
+			Property->AddValue(TEXT("SelfCollisionFriction"), ClothConfig->SelfCollisionFriction, NonAnimatablePropertyFlags);
 
 			// Max distance
 			{
-				const int32 MaxDistanceIndex = Properties->AddProperty(TEXT("MaxDistance"), AnimatablePropertyFlags);
-				Properties->SetWeightedValue(MaxDistanceIndex, 0.f, 1.f);  // Backward compatibility with legacy mask must use a unit range since the multiplier is in the mask
-				Properties->SetStringValue(MaxDistanceIndex, TEXT("MaxDistance"));
+				const int32 MaxDistanceIndex = Property->AddProperty(TEXT("MaxDistance"), AnimatablePropertyFlags);
+				Property->SetWeightedValue(MaxDistanceIndex, 0.f, 1.f);  // Backward compatibility with legacy mask must use a unit range since the multiplier is in the mask
+				Property->SetStringValue(MaxDistanceIndex, TEXT("MaxDistance"));
 			}
 
 			// Backstop
 			{
-				const int32 BackstopDistanceIndex = Properties->AddProperty(TEXT("BackstopDistance"), AnimatablePropertyFlags);
-				Properties->SetWeightedValue(BackstopDistanceIndex, 0.f, 1.f);  // Backward compatibility with legacy mask must use a unit range since the multiplier is in the mask
-				Properties->SetStringValue(BackstopDistanceIndex, TEXT("BackstopDistance"));
+				const int32 BackstopDistanceIndex = Property->AddProperty(TEXT("BackstopDistance"), AnimatablePropertyFlags);
+				Property->SetWeightedValue(BackstopDistanceIndex, 0.f, 1.f);  // Backward compatibility with legacy mask must use a unit range since the multiplier is in the mask
+				Property->SetStringValue(BackstopDistanceIndex, TEXT("BackstopDistance"));
 
-				const int32 BackstopRadiusIndex = Properties->AddProperty(TEXT("BackstopRadius"), AnimatablePropertyFlags);
-				Properties->SetWeightedValue(BackstopRadiusIndex, 0.f, 1.f);  // Backward compatibility with legacy mask must use a unit range since the multiplier is in the mask
-				Properties->SetStringValue(BackstopRadiusIndex, TEXT("BackstopRadius"));
+				const int32 BackstopRadiusIndex = Property->AddProperty(TEXT("BackstopRadius"), AnimatablePropertyFlags);
+				Property->SetWeightedValue(BackstopRadiusIndex, 0.f, 1.f);  // Backward compatibility with legacy mask must use a unit range since the multiplier is in the mask
+				Property->SetStringValue(BackstopRadiusIndex, TEXT("BackstopRadius"));
 
-				Properties->AddValue(TEXT("UseLegacyBackstop"), ClothConfig->bUseLegacyBackstop, NonAnimatablePropertyFlags);
+				Property->AddValue(TEXT("UseLegacyBackstop"), ClothConfig->bUseLegacyBackstop, NonAnimatablePropertyFlags);
 			}
 		}
 
 		// Mark this as a potential legacy config, but leave the behavior control to the client code (usually means constraint are removed with 0 stiffness, or missing weight maps)
-		Properties->AddValue(TEXT("UseLegacyConfig"), bUseLegacyConfig, NonAnimatablePropertyFlags);
+		Property->AddValue(TEXT("UseLegacyConfig"), bUseLegacyConfig, NonAnimatablePropertyFlags);
 	}
 
-	void FClothingSimulationConfig::Initialize(const TSharedPtr<const FManagedArrayCollection>& InPropertyCollection)
+	void FClothingSimulationConfig::Initialize(const TArray<TSharedPtr<const FManagedArrayCollection>>& InPropertyCollections)
 	{
-		Properties->Copy(*InPropertyCollection);
-	}
+		PropertyCollections.Reset(InPropertyCollections.Num());
+		Properties.Reset(InPropertyCollections.Num());
 
-	const Softs::FCollectionPropertyConstFacade& FClothingSimulationConfig::GetProperties() const
+		for (const TSharedPtr<const FManagedArrayCollection>& InPropertyCollection : InPropertyCollections)
+		{
+			TSharedPtr<FManagedArrayCollection>& PropertyCollection = PropertyCollections.Add_GetRef(MakeShared<FManagedArrayCollection>());
+			TUniquePtr<Softs::FCollectionPropertyMutableFacade>& Property = Properties.Add_GetRef(MakeUnique<Softs::FCollectionPropertyMutableFacade>(PropertyCollection));
+			Property->Copy(*InPropertyCollection);
+		}
+	}
+	const Softs::FCollectionPropertyConstFacade& FClothingSimulationConfig::GetProperties(int32 LODIndex) const
 	{
-		return *Properties;
+		return bIsLegacySingleLOD ? *Properties[0] : *Properties[LODIndex];
 	}
-
-	Softs::FCollectionPropertyFacade& FClothingSimulationConfig::GetProperties()
+	Softs::FCollectionPropertyFacade& FClothingSimulationConfig::GetProperties(int32 LODIndex)
 	{
-		return *Properties;
+		return bIsLegacySingleLOD ? *Properties[0] : *Properties[LODIndex];
 	}
 }  // End namespace Chaos
diff --git a/Engine/Plugins/Experimental/ChaosCloth/Source/Chaos/Private/ChaosCloth/ChaosClothingSimulationInteractor.cpp b/Engine/Plugins/Experimental/ChaosCloth/Source/Chaos/Private/ChaosCloth/ChaosClothingSimulationInteractor.cpp
--- a/Engine/Plugins/Experimental/ChaosCloth/Source/Chaos/Private/ChaosCloth/ChaosClothingSimulationInteractor.cpp
+++ b/Engine/Plugins/Experimental/ChaosCloth/Source/Chaos/Private/ChaosCloth/ChaosClothingSimulationInteractor.cpp
@@ -3,6 +3,8 @@
 #include "ChaosCloth/ChaosClothingSimulationInteractor.h"
 #include "ChaosCloth/ChaosClothingSimulationCloth.h"
 #include "ChaosCloth/ChaosClothingSimulation.h"
+#include "ChaosCloth/ChaosClothingSimulationSolver.h"
+#include "ChaosCloth/ChaosClothingSimulationConfig.h"
 
 #include UE_INLINE_GENERATED_CPP_BY_NAME(ChaosClothingSimulationInteractor)
 
@@ -22,6 +24,12 @@ void UChaosClothingInteractor::Sync(IClothingSimulation* Simulation)
 			Command.Execute(Cloth);
 		}
 		Commands.Reset();
+		for (FChaosClothingInteractorConfigCommand& ConfigCommand : ConfigCommands)
+		{
+			check(Cloth->GetConfig()->IsLegacySingleLOD());
+			ConfigCommand.Execute(Cloth->GetConfig(), 0);
+		}
+		ConfigCommands.Reset();
 	}
 
 	// Call to base class' sync
@@ -30,125 +38,147 @@ void UChaosClothingInteractor::Sync(IClothingSimulation* Simulation)
 
 void UChaosClothingInteractor::SetMaterialLinear(float EdgeStiffnessLinear, float BendingStiffnessLinear, float AreaStiffnessLinear)
 {
-	const Chaos::TVec2<Chaos::FRealSingle> EdgeStiffness((FMath::Clamp(FMath::Loge(EdgeStiffnessLinear) * ChaosClothingInteractor::InvStiffnessLogBase + 1.f, 0.f, 1.f)), 1.f);
-	const Chaos::TVec2<Chaos::FRealSingle> BendingStiffness((FMath::Clamp(FMath::Loge(BendingStiffnessLinear) * ChaosClothingInteractor::InvStiffnessLogBase + 1.f, 0.f, 1.f)), 1.f);
-	const Chaos::TVec2<Chaos::FRealSingle> AreaStiffness((FMath::Clamp(FMath::Loge(AreaStiffnessLinear) * ChaosClothingInteractor::InvStiffnessLogBase + 1.f, 0.f, 1.f)), 1.f);
+	const FVector2f EdgeStiffness((FMath::Clamp(FMath::Loge(EdgeStiffnessLinear) * ChaosClothingInteractor::InvStiffnessLogBase + 1.f, 0.f, 1.f)), 1.f);
+	const FVector2f BendingStiffness((FMath::Clamp(FMath::Loge(BendingStiffnessLinear) * ChaosClothingInteractor::InvStiffnessLogBase + 1.f, 0.f, 1.f)), 1.f);
+	const FVector2f AreaStiffness((FMath::Clamp(FMath::Loge(AreaStiffnessLinear) * ChaosClothingInteractor::InvStiffnessLogBase + 1.f, 0.f, 1.f)), 1.f);
 
-	Commands.Add(FChaosClothingInteractorCommand::CreateLambda([EdgeStiffness, BendingStiffness, AreaStiffness](Chaos::FClothingSimulationCloth* Cloth)
+	ConfigCommands.Add(FChaosClothingInteractorConfigCommand::CreateLambda([EdgeStiffness, BendingStiffness, AreaStiffness](Chaos::FClothingSimulationConfig* Config, int32 LODIndex)
 	{
-		Cloth->SetMaterialProperties(EdgeStiffness, BendingStiffness, AreaStiffness);
+		Config->GetProperties(LODIndex).SetWeightedFloatValue(TEXT("EdgeSpringStiffness"), EdgeStiffness);
+		Config->GetProperties(LODIndex).SetWeightedFloatValue(TEXT("BendingSpringStiffness"), BendingStiffness);
+		Config->GetProperties(LODIndex).SetWeightedFloatValue(TEXT("AreaSpringStiffness"), AreaStiffness);
 	}));
 }
 
 void UChaosClothingInteractor::SetMaterial(FVector2D EdgeStiffness, FVector2D BendingStiffness, FVector2D AreaStiffness)
 {
-	Commands.Add(FChaosClothingInteractorCommand::CreateLambda([EdgeStiffness, BendingStiffness, AreaStiffness](Chaos::FClothingSimulationCloth* Cloth)
+	ConfigCommands.Add(FChaosClothingInteractorConfigCommand::CreateLambda([EdgeStiffness, BendingStiffness, AreaStiffness](Chaos::FClothingSimulationConfig* Config, int32 LODIndex)
 	{
-		Cloth->SetMaterialProperties(Chaos::TVec2<Chaos::FRealSingle>(EdgeStiffness[0], EdgeStiffness[1]), Chaos::TVec2<Chaos::FRealSingle>(BendingStiffness[0], BendingStiffness[1]), Chaos::TVec2<Chaos::FRealSingle>(AreaStiffness[0], AreaStiffness[1]));
+		Config->GetProperties(LODIndex).SetWeightedFloatValue(TEXT("EdgeSpringStiffness"), FVector2f(EdgeStiffness));
+		Config->GetProperties(LODIndex).SetWeightedFloatValue(TEXT("BendingSpringStiffness"), FVector2f(BendingStiffness));
+		Config->GetProperties(LODIndex).SetWeightedFloatValue(TEXT("AreaSpringStiffness"), FVector2f(AreaStiffness));
 	}));
 }
 
 void UChaosClothingInteractor::SetLongRangeAttachmentLinear(float TetherStiffnessLinear, float TetherScale)
 {
 	// Deprecated
-	const Chaos::TVec2<Chaos::FRealSingle> TetherStiffness((FMath::Clamp(FMath::Loge(TetherStiffnessLinear) * ChaosClothingInteractor::InvStiffnessLogBase + 1.f, 0.f, 1.f)), 1.f);
-	Commands.Add(FChaosClothingInteractorCommand::CreateLambda([TetherStiffness, TetherScale](Chaos::FClothingSimulationCloth* Cloth)
+	const FVector2f TetherStiffness((FMath::Clamp(FMath::Loge(TetherStiffnessLinear) * ChaosClothingInteractor::InvStiffnessLogBase + 1.f, 0.f, 1.f)), 1.f);
+	ConfigCommands.Add(FChaosClothingInteractorConfigCommand::CreateLambda([TetherStiffness, TetherScale](Chaos::FClothingSimulationConfig* Config, int32 LODIndex)
 	{
-		Cloth->SetLongRangeAttachmentProperties(TetherStiffness, Chaos::TVec2<Chaos::FRealSingle>(TetherScale, TetherScale));
+		Config->GetProperties(LODIndex).SetWeightedFloatValue(TEXT("TetherStiffness"), TetherStiffness);
+		Config->GetProperties(LODIndex).SetWeightedFloatValue(TEXT("TetherScale"), FVector2f(TetherScale));
 	}));
 }
 
 void UChaosClothingInteractor::SetLongRangeAttachment(FVector2D TetherStiffness, FVector2D TetherScale)
 {
-	Commands.Add(FChaosClothingInteractorCommand::CreateLambda([TetherStiffness, TetherScale](Chaos::FClothingSimulationCloth* Cloth)
+	ConfigCommands.Add(FChaosClothingInteractorConfigCommand::CreateLambda([TetherStiffness, TetherScale](Chaos::FClothingSimulationConfig* Config, int32 LODIndex)
 	{
-		Cloth->SetLongRangeAttachmentProperties(
-			Chaos::TVec2<Chaos::FRealSingle>(TetherStiffness[0], TetherStiffness[1]),
-			Chaos::TVec2<Chaos::FRealSingle>(TetherScale[0], TetherScale[1]));
+		Config->GetProperties(LODIndex).SetWeightedFloatValue(TEXT("TetherStiffness"), FVector2f(TetherStiffness));
+		Config->GetProperties(LODIndex).SetWeightedFloatValue(TEXT("TetherScale"), FVector2f(TetherScale));
 	}));
 }
 
 void UChaosClothingInteractor::SetCollision(float CollisionThickness, float FrictionCoefficient, bool bUseCCD, float SelfCollisionThickness)
 {
-	Commands.Add(FChaosClothingInteractorCommand::CreateLambda([CollisionThickness, FrictionCoefficient, bUseCCD, SelfCollisionThickness](Chaos::FClothingSimulationCloth* Cloth)
+	ConfigCommands.Add(FChaosClothingInteractorConfigCommand::CreateLambda([CollisionThickness, FrictionCoefficient, bUseCCD, SelfCollisionThickness](Chaos::FClothingSimulationConfig* Config, int32 LODIndex)
 	{
-		Cloth->SetCollisionProperties(CollisionThickness, FrictionCoefficient, bUseCCD, SelfCollisionThickness);
+		Config->GetProperties(LODIndex).SetValue(TEXT("CollisionThickness"), CollisionThickness);
+		Config->GetProperties(LODIndex).SetValue(TEXT("FrictionCoefficient"), FrictionCoefficient);
+		Config->GetProperties(LODIndex).SetValue(TEXT("UseCCD"), bUseCCD);
+		Config->GetProperties(LODIndex).SetValue(TEXT("SelfCollisionThickness"), SelfCollisionThickness);
 	}));
 }
 
 void UChaosClothingInteractor::SetBackstop(bool bEnabled)
 {
-	Commands.Add(FChaosClothingInteractorCommand::CreateLambda([bEnabled](Chaos::FClothingSimulationCloth* Cloth)
+	ConfigCommands.Add(FChaosClothingInteractorConfigCommand::CreateLambda([bEnabled](Chaos::FClothingSimulationConfig* Config, int32 LODIndex)
 	{
-		Cloth->SetBackstopProperties(bEnabled);
+		Config->GetProperties(LODIndex).SetEnabled(TEXT("BackstopScale"), bEnabled);  // BackstopScale controls whether the backstop is enabled or not
 	}));
 }
 void UChaosClothingInteractor::SetDamping(float DampingCoefficient, float LocalDampingCoefficient)
 {
-	Commands.Add(FChaosClothingInteractorCommand::CreateLambda([DampingCoefficient, LocalDampingCoefficient](Chaos::FClothingSimulationCloth* Cloth)
+	ConfigCommands.Add(FChaosClothingInteractorConfigCommand::CreateLambda([DampingCoefficient, LocalDampingCoefficient](Chaos::FClothingSimulationConfig* Config, int32 LODIndex)
 	{
-		Cloth->SetDampingProperties(DampingCoefficient, LocalDampingCoefficient);
+		Config->GetProperties(LODIndex).SetValue(TEXT("DampingCoefficient"), DampingCoefficient);
+		Config->GetProperties(LODIndex).SetValue(TEXT("LocalDampingCoefficient"), LocalDampingCoefficient);
 	}));
 }
 
 void UChaosClothingInteractor::SetAerodynamics(float DragCoefficient, float LiftCoefficient, FVector WindVelocity)
 {
-	Commands.Add(FChaosClothingInteractorCommand::CreateLambda([DragCoefficient, LiftCoefficient, WindVelocity](Chaos::FClothingSimulationCloth* Cloth)
+	ConfigCommands.Add(FChaosClothingInteractorConfigCommand::CreateLambda([DragCoefficient, LiftCoefficient, WindVelocity](Chaos::FClothingSimulationConfig* Config, int32 LODIndex)
 	{
-		constexpr Chaos::FRealSingle AirDensity = 1.225e-6f;
-		Cloth->SetAerodynamicsProperties(Chaos::TVec2<Chaos::FRealSingle>(DragCoefficient, DragCoefficient), Chaos::TVec2<Chaos::FRealSingle>(LiftCoefficient, LiftCoefficient), AirDensity, WindVelocity);
+		constexpr float AirDensity = 1.225f;
+		constexpr float WorldScale = 100.f;  // Unreal's world unit is the cm
+		Config->GetProperties(LODIndex).SetWeightedFloatValue(TEXT("Drag"), FVector2f(DragCoefficient));
+		Config->GetProperties(LODIndex).SetWeightedFloatValue(TEXT("Lift"), FVector2f(LiftCoefficient));
+		Config->GetProperties(LODIndex).SetValue(TEXT("FluidDensity"), AirDensity);
+		Config->GetProperties(LODIndex).SetValue(TEXT("WindVelocity"), FVector3f(WindVelocity) / WorldScale);
 	}));
 }
 
 void UChaosClothingInteractor::SetWind(FVector2D Drag, FVector2D Lift, float AirDensity, FVector WindVelocity)
 {
-	Commands.Add(FChaosClothingInteractorCommand::CreateLambda([Drag, Lift, AirDensity, WindVelocity](Chaos::FClothingSimulationCloth* Cloth)
+	ConfigCommands.Add(FChaosClothingInteractorConfigCommand::CreateLambda([Drag, Lift, AirDensity, WindVelocity](Chaos::FClothingSimulationConfig* Config, int32 LODIndex)
 	{
-		Cloth->SetAerodynamicsProperties(Chaos::TVec2<Chaos::FRealSingle>(Drag[0], Drag[1]), Chaos::TVec2<Chaos::FRealSingle>(Lift[0], Lift[1]), AirDensity, WindVelocity);
+		constexpr float WorldScale = 100.f;  // Unreal's world unit is the cm
+		Config->GetProperties(LODIndex).SetWeightedFloatValue(TEXT("Drag"), FVector2f(Drag));
+		Config->GetProperties(LODIndex).SetWeightedFloatValue(TEXT("Lift"), FVector2f(Lift));
+		Config->GetProperties(LODIndex).SetValue(TEXT("FluidDensity"), (float)AirDensity * FMath::Cube(WorldScale));  // AirDensity is here in kg/cm^3 for legacy reason but must be in kg/m^3 in the config UI
+		Config->GetProperties(LODIndex).SetValue(TEXT("WindVelocity"), FVector3f(WindVelocity) / WorldScale);
 	}));
 }
 
 void UChaosClothingInteractor::SetPressure(FVector2D Pressure)
 {
-	Commands.Add(FChaosClothingInteractorCommand::CreateLambda([Pressure](Chaos::FClothingSimulationCloth* Cloth)
+	ConfigCommands.Add(FChaosClothingInteractorConfigCommand::CreateLambda([Pressure](Chaos::FClothingSimulationConfig* Config, int32 LODIndex)
 	{
-		Cloth->SetPressureProperties(Chaos::TVec2<Chaos::FRealSingle>(Pressure[0], Pressure[1]));
+		Config->GetProperties(LODIndex).SetWeightedFloatValue(TEXT("Pressure"), FVector2f(Pressure));
 	}));
 }
 
 void UChaosClothingInteractor::SetGravity(float GravityScale, bool bIsGravityOverridden, FVector GravityOverride)
 {
-	Commands.Add(FChaosClothingInteractorCommand::CreateLambda([GravityScale, bIsGravityOverridden, GravityOverride](Chaos::FClothingSimulationCloth* Cloth)
+	ConfigCommands.Add(FChaosClothingInteractorConfigCommand::CreateLambda([GravityScale, bIsGravityOverridden, GravityOverride](Chaos::FClothingSimulationConfig* Config, int32 LODIndex)
 	{
-		Cloth->SetGravityProperties(GravityScale, bIsGravityOverridden, GravityOverride);
+		Config->GetProperties(LODIndex).SetValue(TEXT("GravityScale"), GravityScale);
+		Config->GetProperties(LODIndex).SetValue(TEXT("UseGravityOverride"), bIsGravityOverridden);
+		Config->GetProperties(LODIndex).SetValue(TEXT("GravityOverride"), FVector3f(GravityOverride));
 	}));
 }
 
 void UChaosClothingInteractor::SetAnimDriveLinear(float AnimDriveStiffnessLinear)
 {
 	// Deprecated
-	const Chaos::TVec2<Chaos::FRealSingle> AnimDriveStiffness(0.f, FMath::Clamp(FMath::Loge(AnimDriveStiffnessLinear) * ChaosClothingInteractor::InvStiffnessLogBase + 1.f, 0.f, 1.f));
-	Commands.Add(FChaosClothingInteractorCommand::CreateLambda([AnimDriveStiffness](Chaos::FClothingSimulationCloth* Cloth)
+	const FVector2f AnimDriveStiffness(0.f, FMath::Clamp(FMath::Loge(AnimDriveStiffnessLinear) * ChaosClothingInteractor::InvStiffnessLogBase + 1.f, 0.f, 1.f));
+	ConfigCommands.Add(FChaosClothingInteractorConfigCommand::CreateLambda([AnimDriveStiffness](Chaos::FClothingSimulationConfig* Config, int32 LODIndex)
 	{
 		// The Anim Drive stiffness Low value needs to be 0 in order to keep backward compatibility with existing mask (this wouldn't be an issue if this property had no legacy mask)
-		static const Chaos::TVec2<Chaos::FRealSingle> AnimDriveDamping(0.f, 1.f);
-		Cloth->SetAnimDriveProperties(AnimDriveStiffness, AnimDriveDamping);
+		static const FVector2f AnimDriveDamping(0.f, 1.f);
+		Config->GetProperties(LODIndex).SetWeightedFloatValue(TEXT("AnimDriveStiffness"), AnimDriveStiffness);
+		Config->GetProperties(LODIndex).SetWeightedFloatValue(TEXT("AnimDriveDamping"), AnimDriveDamping);
 	}));
 }
 
 void UChaosClothingInteractor::SetAnimDrive(FVector2D AnimDriveStiffness, FVector2D AnimDriveDamping)
 {
-	Commands.Add(FChaosClothingInteractorCommand::CreateLambda([AnimDriveStiffness, AnimDriveDamping](Chaos::FClothingSimulationCloth* Cloth)
+	ConfigCommands.Add(FChaosClothingInteractorConfigCommand::CreateLambda([AnimDriveStiffness, AnimDriveDamping](Chaos::FClothingSimulationConfig* Config, int32 LODIndex)
 	{
-		Cloth->SetAnimDriveProperties(Chaos::TVec2<Chaos::FRealSingle>(AnimDriveStiffness.X, AnimDriveStiffness.Y), Chaos::TVec2<Chaos::FRealSingle>(AnimDriveDamping.X, AnimDriveDamping.Y));
+		Config->GetProperties(LODIndex).SetWeightedFloatValue(TEXT("AnimDriveStiffness"), FVector2f(AnimDriveStiffness));
+		Config->GetProperties(LODIndex).SetWeightedFloatValue(TEXT("AnimDriveDamping"), FVector2f(AnimDriveDamping));
 	}));
 }
 
 void UChaosClothingInteractor::SetVelocityScale(FVector LinearVelocityScale, float AngularVelocityScale, float FictitiousAngularScale)
 {
-	Commands.Add(FChaosClothingInteractorCommand::CreateLambda([LinearVelocityScale, AngularVelocityScale, FictitiousAngularScale](Chaos::FClothingSimulationCloth* Cloth)
+	ConfigCommands.Add(FChaosClothingInteractorConfigCommand::CreateLambda([LinearVelocityScale, AngularVelocityScale, FictitiousAngularScale](Chaos::FClothingSimulationConfig* Config, int32 LODIndex)
 	{
-		Cloth->SetVelocityScaleProperties(LinearVelocityScale, AngularVelocityScale, FictitiousAngularScale);
+		Config->GetProperties(LODIndex).SetValue(TEXT("LinearVelocityScale"), FVector3f(LinearVelocityScale));
+		Config->GetProperties(LODIndex).SetValue(TEXT("AngularVelocityScale"), AngularVelocityScale);
+		Config->GetProperties(LODIndex).SetValue(TEXT("FictitiousAngularScale"), FictitiousAngularScale);
 	}));
 }
 
@@ -180,6 +210,12 @@ void UChaosClothingSimulationInteractor::Sync(IClothingSimulation* Simulation, I
 		Command.Execute(static_cast<Chaos::FClothingSimulation*>(Simulation), static_cast<Chaos::FClothingSimulationContext*>(Context));
 	}
 	Commands.Reset();
+	for (FChaosClothingInteractorConfigCommand& ConfigCommand : ConfigCommands)
+	{
+		ensure(static_cast<Chaos::FClothingSimulation*>(Simulation)->GetSolver()->GetConfig()->IsLegacySingleLOD());
+		ConfigCommand.Execute(static_cast<Chaos::FClothingSimulation*>(Simulation)->GetSolver()->GetConfig(), 0);
+	}
+	ConfigCommands.Reset();
 
 	// Call base class' sync 
 	UClothingSimulationInteractor::Sync(Simulation, Context);
@@ -231,25 +267,25 @@ void UChaosClothingSimulationInteractor::DisableGravityOverride()
 
 void UChaosClothingSimulationInteractor::SetNumIterations(int32 NumIterations)
 {
-	Commands.Add(FChaosClothingSimulationInteractorCommand::CreateLambda([NumIterations](Chaos::FClothingSimulation* Simulation, Chaos::FClothingSimulationContext* /*Context*/)
+	ConfigCommands.Add(FChaosClothingInteractorConfigCommand::CreateLambda([NumIterations](Chaos::FClothingSimulationConfig* Config, int32 LODIndex)
 	{
-		Simulation->SetNumIterations(NumIterations);
+		Config->GetProperties(LODIndex).SetValue(TEXT("NumIterations"), NumIterations);
 	}));
 }
 
 void UChaosClothingSimulationInteractor::SetMaxNumIterations(int32 MaxNumIterations)
 {
-	Commands.Add(FChaosClothingSimulationInteractorCommand::CreateLambda([MaxNumIterations](Chaos::FClothingSimulation* Simulation, Chaos::FClothingSimulationContext* /*Context*/)
+	ConfigCommands.Add(FChaosClothingInteractorConfigCommand::CreateLambda([MaxNumIterations](Chaos::FClothingSimulationConfig* Config, int32 LODIndex)
 	{
-		Simulation->SetMaxNumIterations(MaxNumIterations);
+		Config->GetProperties(LODIndex).SetValue(TEXT("MaxNumIterations"), MaxNumIterations);
 	}));
 }
 
 void UChaosClothingSimulationInteractor::SetNumSubsteps(int32 NumSubsteps)
 {
-	Commands.Add(FChaosClothingSimulationInteractorCommand::CreateLambda([NumSubsteps](Chaos::FClothingSimulation* Simulation, Chaos::FClothingSimulationContext* /*Context*/)
+	ConfigCommands.Add(FChaosClothingInteractorConfigCommand::CreateLambda([NumSubsteps](Chaos::FClothingSimulationConfig* Config, int32 LODIndex)
 	{
-		Simulation->SetNumSubsteps(NumSubsteps);
+		Config->GetProperties(LODIndex).SetValue(TEXT("NumSubsteps"), NumSubsteps);
 	}));
 }
 
diff --git a/Engine/Plugins/Experimental/ChaosCloth/Source/Chaos/Private/ChaosCloth/ChaosClothingSimulationSolver.cpp b/Engine/Plugins/Experimental/ChaosCloth/Source/Chaos/Private/ChaosCloth/ChaosClothingSimulationSolver.cpp
--- a/Engine/Plugins/Experimental/ChaosCloth/Source/Chaos/Private/ChaosCloth/ChaosClothingSimulationSolver.cpp
+++ b/Engine/Plugins/Experimental/ChaosCloth/Source/Chaos/Private/ChaosCloth/ChaosClothingSimulationSolver.cpp
@@ -9,6 +9,7 @@
 #include "Chaos/PBDEvolution.h"
 #include "Chaos/WeightedLatticeImplicitObject.h"
 #include "Chaos/Levelset.h"
+#include "GeometryCollection/ManagedArrayCollection.h"
 #if INTEL_ISPC
 #include "ChaosClothingSimulationSolver.ispc.generated.h"
 #endif
@@ -100,9 +101,6 @@ FClothingSimulationSolver::FClothingSimulationSolver()
 	, VelocityScale(1.)
 	, Time(0.)
 	, DeltaTime(ClothingSimulationSolverConstant::StartDeltaTime)
-	, NumIterations(ClothingSimulationSolverDefault::NumIterations)
-	, MaxNumIterations(ClothingSimulationSolverDefault::MaxNumIterations)
-	, NumSubsteps(ClothingSimulationSolverDefault::NumSubsteps)
 	, CollisionParticlesOffset(0)
 	, CollisionParticlesSize(0)
 	, Gravity(ClothingSimulationSolverDefault::Gravity)
@@ -186,6 +184,11 @@ FClothingSimulationSolver::FClothingSimulationSolver()
 
 FClothingSimulationSolver::~FClothingSimulationSolver()
 {
+	// If the PropertyCollection is owned by this object, so does the current config object
+	if (PropertyCollection.IsValid())
+	{
+		delete Config;
+	}
 }
 
 void FClothingSimulationSolver::SetLocalSpaceLocation(const FVec3& InLocalSpaceLocation, bool bReset)
@@ -339,6 +342,29 @@ void FClothingSimulationSolver::RefreshCloths()
 	CollisionParticlesOffset = Evolution->CollisionParticles().Size();
 }
 
+void FClothingSimulationSolver::SetConfig(FClothingSimulationConfig* InConfig)
+{
+	// If the PropertyCollection is owned by this object, so does the current config object
+	if (PropertyCollection.IsValid())
+	{
+		delete Config;
+		PropertyCollection.Reset();
+	}
+
+	if (InConfig)
+	{
+		Config = InConfig;
+	}
+	else
+	{
+		// Create a default empty config object for coherence
+		PropertyCollection = MakeShared<FManagedArrayCollection>();
+		PRAGMA_DISABLE_DEPRECATION_WARNINGS
+		Config = new FClothingSimulationConfig(PropertyCollection);
+		PRAGMA_ENABLE_DEPRECATION_WARNINGS
+	}
+}
+
 void FClothingSimulationSolver::ResetParticles()
 {
 	Evolution->ResetParticles();
@@ -725,59 +751,47 @@ void FClothingSimulationSolver::SetWindVelocity(const TVec3<FRealSingle>& InWind
 
 void FClothingSimulationSolver::SetNumIterations(int32 InNumIterations)
 {
-	if (Config)
-	{
-		Config->GetProperties().SetValue(TEXT("NumIterations"), InNumIterations);
-	}
-	else
+	if (ensure(Config))
 	{
-		NumIterations = InNumIterations;
+		Config->GetProperties(SolverLOD).SetValue(TEXT("NumIterations"), InNumIterations);
 	}
 }
 
 int32 FClothingSimulationSolver::GetNumIterations() const
 {
 	return Config ?
-		Config->GetProperties().GetValue<int32>(TEXT("NumIterations"), ClothingSimulationSolverDefault::NumIterations) :
-		NumIterations;
+		Config->GetProperties(SolverLOD).GetValue<int32>(TEXT("NumIterations"), ClothingSimulationSolverDefault::NumIterations) :
+		ClothingSimulationSolverDefault::NumIterations;
 }
 
 void FClothingSimulationSolver::SetMaxNumIterations(int32 InMaxNumIterations)
 {
-	if (Config)
-	{
-		Config->GetProperties().SetValue(TEXT("MaxNumIterations"), InMaxNumIterations);
-	}
-	else
+	if (ensure(Config))
 	{
-		MaxNumIterations = InMaxNumIterations;
+		Config->GetProperties(SolverLOD).SetValue(TEXT("MaxNumIterations"), InMaxNumIterations);
 	}
 }
 
 int32 FClothingSimulationSolver::GetMaxNumIterations() const
 {
-	return Config?
-		Config->GetProperties().GetValue<int32>(TEXT("MaxNumIterations"), ClothingSimulationSolverDefault::MaxNumIterations) :
-		MaxNumIterations;
+	return Config ?
+		Config->GetProperties(SolverLOD).GetValue<int32>(TEXT("MaxNumIterations"), ClothingSimulationSolverDefault::MaxNumIterations) :
+		ClothingSimulationSolverDefault::MaxNumIterations;
 }
 
 void FClothingSimulationSolver::SetNumSubsteps(int32 InNumSubsteps)
 {
-	if (Config)
-	{
-		Config->GetProperties().SetValue(TEXT("NumSubsteps"), InNumSubsteps);
-	}
-	else
+	if (ensure(Config))
 	{
-		NumSubsteps = InNumSubsteps;
+		Config->GetProperties(SolverLOD).SetValue(TEXT("NumSubsteps"), InNumSubsteps);
 	}
 }
 
 int32 FClothingSimulationSolver::GetNumSubsteps() const
 {
-	return Config?
-		Config->GetProperties().GetValue<int32>(TEXT("NumSubsteps"), ClothingSimulationSolverDefault::NumSubsteps) :
-		NumSubsteps;
+	return Config ?
+		Config->GetProperties(SolverLOD).GetValue<int32>(TEXT("NumSubsteps"), ClothingSimulationSolverDefault::NumSubsteps) :
+		ClothingSimulationSolverDefault::NumSubsteps;
 }
 
 void FClothingSimulationSolver::SetWindVelocity(uint32 GroupId, const TVec3<FRealSingle>& InWindVelocity)
@@ -789,11 +803,11 @@ void FClothingSimulationSolver::SetWindVelocity(uint32 GroupId, const TVec3<FRea
 void FClothingSimulationSolver::SetWindAndPressureGeometry(
 	uint32 GroupId,
 	const FTriangleMesh& TriangleMesh,
-	const Softs::FCollectionPropertyConstFacade& PropertyCollection,
+	const Softs::FCollectionPropertyConstFacade& InPropertyCollection,
 	const TMap<FString, TConstArrayView<FRealSingle>>& WeightMaps)
 {
 	Softs::FVelocityAndPressureField& VelocityAndPressureField = Evolution->GetVelocityAndPressureField(GroupId);
-	VelocityAndPressureField.SetGeometry(&TriangleMesh, PropertyCollection, WeightMaps, ClothingSimulationSolverConstant::WorldScale);
+	VelocityAndPressureField.SetGeometry(&TriangleMesh, InPropertyCollection, WeightMaps, ClothingSimulationSolverConstant::WorldScale);
 }
 
 void FClothingSimulationSolver::SetWindAndPressureGeometry(
@@ -809,12 +823,12 @@ void FClothingSimulationSolver::SetWindAndPressureGeometry(
 
 void FClothingSimulationSolver::SetWindAndPressureProperties(
 	uint32 GroupId,
-	const Softs::FCollectionPropertyConstFacade& PropertyCollection,
+	const Softs::FCollectionPropertyConstFacade& InPropertyCollection,
 	const TMap<FString, TConstArrayView<FRealSingle>>& WeightMaps,
 	bool bEnableAerodynamics)
 {
 	Softs::FVelocityAndPressureField& VelocityAndPressureField = Evolution->GetVelocityAndPressureField(GroupId);
-	VelocityAndPressureField.SetProperties(PropertyCollection, WeightMaps, ClothingSimulationSolverConstant::WorldScale, bEnableAerodynamics);
+	VelocityAndPressureField.SetProperties(InPropertyCollection, WeightMaps, ClothingSimulationSolverConstant::WorldScale, bEnableAerodynamics);
 }
 
 void FClothingSimulationSolver::SetWindAndPressureProperties(
@@ -1143,33 +1157,22 @@ void FClothingSimulationSolver::Update(Softs::FSolverReal InDeltaTime)
 			SCOPE_CYCLE_COUNTER(STAT_ChaosClothSolverUpdateSolverStep);
 			SCOPE_CYCLE_COUNTER(STAT_ClothInternalSolve);
 
-			// Read config properties if any, otherwise go back to using the property setters for this object
-			int32 ConfigNumIterations;
-			int32 ConfigMaxNumIterations;
-			int32 ConfigNumSubsteps;
-			if (Config)
-			{
-				const Softs::FCollectionPropertyFacade& Properties = Config->GetProperties();
+			check(Config);
+			const Softs::FCollectionPropertyFacade& Properties = Config->GetProperties(SolverLOD);
 
-				ConfigMaxNumIterations = FMath::Max(
-					Properties.GetValue<int32>(TEXT("MaxNumIterations"), ClothingSimulationSolverDefault::MaxNumIterations),
-					ClothingSimulationSolverDefault::MinNumIterations);
+			const int32 ConfigMaxNumIterations = FMath::Max(
+				Properties.GetValue<int32>(TEXT("MaxNumIterations"), ClothingSimulationSolverDefault::MaxNumIterations),
+				ClothingSimulationSolverDefault::MinNumIterations);
 
-				ConfigNumIterations = FMath::Clamp(
-					Properties.GetValue<int32>(TEXT("NumIterations"), ClothingSimulationSolverDefault::NumIterations),
-					ClothingSimulationSolverDefault::MinNumIterations,
-					ConfigMaxNumIterations);
+			const int32 ConfigNumIterations = FMath::Clamp(
+				Properties.GetValue<int32>(TEXT("NumIterations"), ClothingSimulationSolverDefault::NumIterations),
+				ClothingSimulationSolverDefault::MinNumIterations,
+				ConfigMaxNumIterations);
+
+			const int32 ConfigNumSubsteps = FMath::Max(
+				Properties.GetValue<int32>(TEXT("NumSubsteps"), ClothingSimulationSolverDefault::NumSubsteps),
+				ClothingSimulationSolverDefault::MinNumSubsteps);
 
-				ConfigNumSubsteps = FMath::Max(
-					Properties.GetValue<int32>(TEXT("NumSubsteps"), ClothingSimulationSolverDefault::NumSubsteps),
-					ClothingSimulationSolverDefault::MinNumSubsteps);
-			}
-			else
-			{
-				ConfigMaxNumIterations = MaxNumIterations;
-				ConfigNumIterations = NumIterations;
-				ConfigNumSubsteps = NumSubsteps;
-			}
 			// Update solver time dependent parameters
 			const Softs::FSolverReal SolverFrequency = (Softs::FSolverReal)ClothingSimulationSolverDefault::SolverFrequency;  // 60Hz default TODO: Should this become a solver property?
 
diff --git a/Engine/Plugins/Experimental/ChaosCloth/Source/Chaos/Public/ChaosCloth/ChaosClothingSimulation.h b/Engine/Plugins/Experimental/ChaosCloth/Source/Chaos/Public/ChaosCloth/ChaosClothingSimulation.h
--- a/Engine/Plugins/Experimental/ChaosCloth/Source/Chaos/Public/ChaosCloth/ChaosClothingSimulation.h
+++ b/Engine/Plugins/Experimental/ChaosCloth/Source/Chaos/Public/ChaosCloth/ChaosClothingSimulation.h
@@ -82,6 +82,8 @@ namespace Chaos
 
 		FClothingSimulationCloth* GetCloth(int32 ClothId);
 
+		FClothingSimulationSolver* GetSolver() { return Solver.Get(); }
+
 #if WITH_EDITOR
 		// Editor only debug draw function
 		void DebugDrawPhysMeshShaded(FPrimitiveDrawInterface* PDI) const { Visualization.DrawPhysMeshShaded(PDI); }
diff --git a/Engine/Plugins/Experimental/ChaosCloth/Source/Chaos/Public/ChaosCloth/ChaosClothingSimulationCloth.h b/Engine/Plugins/Experimental/ChaosCloth/Source/Chaos/Public/ChaosCloth/ChaosClothingSimulationCloth.h
--- a/Engine/Plugins/Experimental/ChaosCloth/Source/Chaos/Public/ChaosCloth/ChaosClothingSimulationCloth.h
+++ b/Engine/Plugins/Experimental/ChaosCloth/Source/Chaos/Public/ChaosCloth/ChaosClothingSimulationCloth.h
@@ -110,17 +110,27 @@ PRAGMA_ENABLE_DEPRECATION_WARNINGS
 
 		// ---- Animatable property setters ----
 		void SetMaxDistancesMultiplier(FRealSingle InMaxDistancesMultiplier) { MaxDistancesMultiplier = InMaxDistancesMultiplier; }
-
+		UE_DEPRECATED(5.3, "Set properties directly through FClothingSimulationConfig")
 		CHAOSCLOTH_API void SetMaterialProperties(const TVec2<FRealSingle>& InEdgeStiffness, const TVec2<FRealSingle>& InBendingStiffness, const TVec2<FRealSingle>& InAreaStiffness);
+		UE_DEPRECATED(5.3, "Set properties directly through FClothingSimulationConfig")
 		CHAOSCLOTH_API void SetLongRangeAttachmentProperties(const TVec2<FRealSingle>& InTetherStiffness, const TVec2<FRealSingle>& InTetherScale);
+		UE_DEPRECATED(5.3, "Set properties directly through FClothingSimulationConfig")
 		CHAOSCLOTH_API void SetCollisionProperties(FRealSingle InCollisionThickness, FRealSingle InFrictionCoefficient, bool bInUseCCD, FRealSingle InSelfCollisionThickness);
+		UE_DEPRECATED(5.3, "Set properties directly through FClothingSimulationConfig")
 		CHAOSCLOTH_API void SetBackstopProperties(bool bInEnableBackstop);
+		UE_DEPRECATED(5.3, "Set properties directly through FClothingSimulationConfig")
 		CHAOSCLOTH_API void SetDampingProperties(FRealSingle InDampingCoefficient, FRealSingle InLocalDampingCoefficient = 0.f);
+		UE_DEPRECATED(5.3, "Set properties directly through FClothingSimulationConfig")
 		CHAOSCLOTH_API void SetAerodynamicsProperties(const TVec2<FRealSingle>& InDrag, const TVec2<FRealSingle>& InLift, FRealSingle InAirDensity, const FVec3& InWindVelocity);  // AirDensity is here in kg/cm^3 for legacy reason (kg/m^3 in UI)
+		UE_DEPRECATED(5.3, "Set properties directly through FClothingSimulationConfig")
 		CHAOSCLOTH_API void SetPressureProperties(const TVec2<FRealSingle>& InPressure);
+		UE_DEPRECATED(5.3, "Set properties directly through FClothingSimulationConfig")
 		CHAOSCLOTH_API void SetGravityProperties(FRealSingle InGravityScale, bool bInUseGravityOverride, const FVec3& InGravityOverride);
+		UE_DEPRECATED(5.3, "Set properties directly through FClothingSimulationConfig")
 		CHAOSCLOTH_API void SetAnimDriveProperties(const TVec2<FRealSingle>& InAnimDriveStiffness, const TVec2<FRealSingle>& InAnimDriveDamping);
+		UE_DEPRECATED(5.3, "Get properties directly through FClothingSimulationConfig")
 		CHAOSCLOTH_API void GetAnimDriveProperties(TVec2<FRealSingle>& OutAnimDriveStiffness, TVec2<FRealSingle>& OutAnimDriveDamping);
+		UE_DEPRECATED(5.3, "Set properties directly through FClothingSimulationConfig")
 		CHAOSCLOTH_API void SetVelocityScaleProperties(const FVec3& InLinearVelocityScale, FRealSingle InAngularVelocityScale, FRealSingle InFictitiousAngularScale);
 
 		void Reset() { bNeedsReset = true; }
diff --git a/Engine/Plugins/Experimental/ChaosCloth/Source/Chaos/Public/ChaosCloth/ChaosClothingSimulationConfig.h b/Engine/Plugins/Experimental/ChaosCloth/Source/Chaos/Public/ChaosCloth/ChaosClothingSimulationConfig.h
--- a/Engine/Plugins/Experimental/ChaosCloth/Source/Chaos/Public/ChaosCloth/ChaosClothingSimulationConfig.h
+++ b/Engine/Plugins/Experimental/ChaosCloth/Source/Chaos/Public/ChaosCloth/ChaosClothingSimulationConfig.h
@@ -22,7 +22,9 @@ namespace Chaos
 	{
 	public:
 		CHAOSCLOTH_API FClothingSimulationConfig();
+		UE_DEPRECATED(5.3, "Use TArray version of this constructor.")
 		CHAOSCLOTH_API FClothingSimulationConfig(const TSharedPtr<const FManagedArrayCollection>& InPropertyCollection);
+		CHAOSCLOTH_API FClothingSimulationConfig(const TArray<TSharedPtr<const FManagedArrayCollection>>& InPropertyCollections);
 
 		CHAOSCLOTH_API ~FClothingSimulationConfig();
 
@@ -40,19 +42,68 @@ namespace Chaos
 		CHAOSCLOTH_API void Initialize(const UChaosClothConfig* ClothConfig, const UChaosClothSharedSimConfig* ClothSharedConfig, bool bUseLegacyConfig = false);
 
 		/** Initialize config from a property collection. */
-		CHAOSCLOTH_API void Initialize(const TSharedPtr<const FManagedArrayCollection>& InPropertyCollection);
+		UE_DEPRECATED(5.3, "Use TArray version of Initialize.")
+		CHAOSCLOTH_API void Initialize(const TSharedPtr<const FManagedArrayCollection>& InPropertyCollection)
+		{
+			Initialize(TArray<TSharedPtr<const FManagedArrayCollection>>({ InPropertyCollection }));
+			bIsLegacySingleLOD = true;
+		}
+
+		/** Initialize config from an array of property collections (one per LOD). */
+		CHAOSCLOTH_API void Initialize(const TArray<TSharedPtr<const FManagedArrayCollection>>& InPropertyCollections);
 
 		/** Return a property collection facade for reading properties from this configuration. */
-		CHAOSCLOTH_API const Softs::FCollectionPropertyConstFacade& GetProperties() const;
+		CHAOSCLOTH_API const Softs::FCollectionPropertyConstFacade& GetProperties(int32 LODIndex) const;
+		CHAOSCLOTH_API const Softs::FCollectionPropertyConstFacade& GetProperties() const 
+		{
+			ensureMsgf(bIsLegacySingleLOD, TEXT("LODIndex-less version should only be used with legacy single LOD configs."));
+			return GetProperties(0); 
+		}
 
 		/** Return a property collection facade for setting properties for this configuration. */
-		CHAOSCLOTH_API Softs::FCollectionPropertyFacade& GetProperties();
+		CHAOSCLOTH_API Softs::FCollectionPropertyFacade& GetProperties(int32 LODIndex);
+		CHAOSCLOTH_API Softs::FCollectionPropertyFacade& GetProperties() 
+		{
+			ensureMsgf(bIsLegacySingleLOD, TEXT("LODIndex-less version should only be used with legacy single LOD configs."));
+			return GetProperties(0); 
+		}
 
 		/** Return this configuration's internal property collection. */
-		TSharedPtr<const FManagedArrayCollection> GetPropertyCollection() const { return TSharedPtr<const FManagedArrayCollection>(PropertyCollection); }
+		TSharedPtr<const FManagedArrayCollection> GetPropertyCollection(int32 LODIndex) const 
+		{
+			return bIsLegacySingleLOD ? TSharedPtr<const FManagedArrayCollection>(PropertyCollections[0]) : 
+				TSharedPtr<const FManagedArrayCollection>(PropertyCollections[LODIndex]);
+		}
+		TSharedPtr<const FManagedArrayCollection> GetPropertyCollection() const 
+		{
+			ensureMsgf(bIsLegacySingleLOD, TEXT("LODIndex-less version should only be used with legacy single LOD configs."));
+			return GetPropertyCollection(0); 
+		}
+
+		bool IsLegacySingleLOD() const { return bIsLegacySingleLOD; }
+		int32 GetNumLODs() const { return PropertyCollections.Num(); }
+
+		template<typename FunctionType>
+		void ForAllProperties(FunctionType Func)
+		{
+			for (TUniquePtr<Softs::FCollectionPropertyMutableFacade>& Property : Properties)
+			{
+				Func((Softs::FCollectionPropertyFacade&)(*Property));
+			}
+		}
+
+		template<typename FunctionType>
+		void ForAllProperties(FunctionType Func) const
+		{
+			for (const TUniquePtr<Softs::FCollectionPropertyMutableFacade>& Property : Properties)
+			{
+				Func((Softs::FCollectionPropertyConstFacade&)(*Property));
+			}
+		}
 
 	private:
-		TSharedPtr<FManagedArrayCollection> PropertyCollection;
-		TUniquePtr<Softs::FCollectionPropertyMutableFacade> Properties;
+		bool bIsLegacySingleLOD = false;
+		TArray<TSharedPtr<FManagedArrayCollection>> PropertyCollections;
+		TArray<TUniquePtr<Softs::FCollectionPropertyMutableFacade>> Properties;
 	};
 } // namespace Chaos
diff --git a/Engine/Plugins/Experimental/ChaosCloth/Source/Chaos/Public/ChaosCloth/ChaosClothingSimulationInteractor.h b/Engine/Plugins/Experimental/ChaosCloth/Source/Chaos/Public/ChaosCloth/ChaosClothingSimulationInteractor.h
--- a/Engine/Plugins/Experimental/ChaosCloth/Source/Chaos/Public/ChaosCloth/ChaosClothingSimulationInteractor.h
+++ b/Engine/Plugins/Experimental/ChaosCloth/Source/Chaos/Public/ChaosCloth/ChaosClothingSimulationInteractor.h
@@ -9,10 +9,12 @@ namespace Chaos
 {
 	class FClothingSimulation;
 	class FClothingSimulationCloth;
+	class FClothingSimulationConfig;
 }
 class FClothingSimulationContextCommon;
 
 // Command signature for handling synced command buffer
+DECLARE_DELEGATE_TwoParams(FChaosClothingInteractorConfigCommand, Chaos::FClothingSimulationConfig*, int32)
 DECLARE_DELEGATE_OneParam(FChaosClothingInteractorCommand, Chaos::FClothingSimulationCloth*)
 DECLARE_DELEGATE_TwoParams(FChaosClothingSimulationInteractorCommand, Chaos::FClothingSimulation*, FClothingSimulationContextCommon*)
 
@@ -73,6 +75,7 @@ public:
 private:
 	// Cloth command queue processed when we hit a sync
 	TArray<FChaosClothingInteractorCommand> Commands;
+	TArray<FChaosClothingInteractorConfigCommand> ConfigCommands;
 };
 
 UCLASS(BlueprintType, MinimalAPI)
@@ -98,4 +101,5 @@ protected:
 private:
 	// Simulation command queue processed when we hit a sync
 	TArray<FChaosClothingSimulationInteractorCommand> Commands;
+	TArray<FChaosClothingInteractorConfigCommand> ConfigCommands;
 };
diff --git a/Engine/Plugins/Experimental/ChaosCloth/Source/Chaos/Public/ChaosCloth/ChaosClothingSimulationSolver.h b/Engine/Plugins/Experimental/ChaosCloth/Source/Chaos/Public/ChaosCloth/ChaosClothingSimulationSolver.h
--- a/Engine/Plugins/Experimental/ChaosCloth/Source/Chaos/Public/ChaosCloth/ChaosClothingSimulationSolver.h
+++ b/Engine/Plugins/Experimental/ChaosCloth/Source/Chaos/Public/ChaosCloth/ChaosClothingSimulationSolver.h
@@ -56,10 +56,13 @@ namespace Chaos
 		CHAOSCLOTH_API void SetWindVelocity(const TVec3<FRealSingle>& InWindVelocity, FRealSingle InLegacyWindAdaption = (FRealSingle)0.);
 		const TVec3<FRealSingle>& GetWindVelocity() const { return WindVelocity; }
 
+		UE_DEPRECATED(5.3, "Set properties directly through FClothingSimulationConfig")
 		CHAOSCLOTH_API void SetNumIterations(int32 InNumIterations);
 		CHAOSCLOTH_API int32 GetNumIterations() const;
+		UE_DEPRECATED(5.3, "Set properties directly through FClothingSimulationConfig")
 		CHAOSCLOTH_API void SetMaxNumIterations(int32 InMaxNumIterations);
 		CHAOSCLOTH_API int32 GetMaxNumIterations() const;
+		UE_DEPRECATED(5.3, "Set properties directly through FClothingSimulationConfig")
 		CHAOSCLOTH_API void SetNumSubsteps(int32 InNumSubsteps);
 		CHAOSCLOTH_API int32 GetNumSubsteps() const;
 
@@ -84,9 +87,11 @@ namespace Chaos
 		/**
 		 * Set the solver configuration.
 		 * Can use a cloth config if a single cloth is being simulated.
-		 * When a solver configuration is set, the solver shared properties (e.g. NumIterations, MaxNumIterations, NumSubsteps) are overriden.
 		 */
-		void SetConfig(FClothingSimulationConfig* InConfig) { Config = InConfig; }
+		CHAOSCLOTH_API void SetConfig(FClothingSimulationConfig* InConfig);
+
+		void SetSolverLOD(int32 LODIndex) { SolverLOD = LODIndex; }
+		int32 GetSolverLOD() const { return SolverLOD; }
 
 		/** Advance the simulation. */
 		CHAOSCLOTH_API void Update(FSolverReal InDeltaTime);
@@ -260,7 +265,9 @@ namespace Chaos
 		TArray<FClothingSimulationCloth*> Cloths;
 		
 		// Solver config
+		int32 SolverLOD = 0; // Config may contain multiple LODs.
 		FClothingSimulationConfig* Config = nullptr;
+		TSharedPtr<FManagedArrayCollection> PropertyCollection;  // Used for backward compatibility only, otherwise the properties are owned by the Config
 
 		// Simulation group attributes
 		TArrayCollectionArray<Softs::FSolverRigidTransform3> PreSimulationTransforms;  // Allow a different frame of reference for each cloth groups
@@ -296,11 +303,6 @@ namespace Chaos
 		FSolverReal Time;
 		FSolverReal DeltaTime;
 
-		// Solver config properties, superseded by the config pointer
-		int32 NumIterations;
-		int32 MaxNumIterations;
-		int32 NumSubsteps;
-
 		// Solver colliders offset
 		int32 CollisionParticlesOffset;  // Collision particle offset on the first solver/non cloth collider
 		int32 CollisionParticlesSize;  // Number of solver only colliders
diff --git a/Engine/Plugins/Experimental/ChaosClothAsset/Source/ChaosClothAssetEngine/Private/ChaosClothAsset/ClothComponent.cpp b/Engine/Plugins/Experimental/ChaosClothAsset/Source/ChaosClothAssetEngine/Private/ChaosClothAsset/ClothComponent.cpp
--- a/Engine/Plugins/Experimental/ChaosClothAsset/Source/ChaosClothAssetEngine/Private/ChaosClothAsset/ClothComponent.cpp
+++ b/Engine/Plugins/Experimental/ChaosClothAsset/Source/ChaosClothAssetEngine/Private/ChaosClothAsset/ClothComponent.cpp
@@ -67,24 +67,29 @@ void UChaosClothComponent::ResetConfigProperties()
 {
 	if (IsRegistered())
 	{
-		check(PropertyCollection);
+		if (GetClothAsset())
+		{
+			const TArray<TSharedRef<FManagedArrayCollection>>& ClothCollections = GetClothAsset()->GetClothCollections();
+			PropertyCollections.Reset();
+			PropertyCollections.Reserve(ClothCollections.Num());
+			CollectionPropertyFacades.Reset();
+			CollectionPropertyFacades.Reserve(ClothCollections.Num());
+			for (const TSharedRef<FManagedArrayCollection>& ClothCollection : ClothCollections)
+			{
+				TSharedPtr<FManagedArrayCollection>& PropertyCollection = PropertyCollections.Add_GetRef(MakeShared<FManagedArrayCollection>());
 
-		// Create a mutable facade for our component's property collection, and use it to copy the properties from the cloth collection
-		::Chaos::Softs::FCollectionPropertyMutableFacade CollectionPropertyMutableFacade(PropertyCollection);
+				::Chaos::Softs::FCollectionPropertyMutableFacade CollectionPropertyMutableFacade(PropertyCollection);
+				CollectionPropertyMutableFacade.Copy(*ClothCollection);
 
-		// TODO: Make PropertyCollection also one per LOD.
-		if (GetClothAsset() && GetClothAsset()->GetClothCollections().Num() > 0)
-		{
-			CollectionPropertyMutableFacade.Copy(*GetClothAsset()->GetClothCollections()[0]);
+				CollectionPropertyFacades.Add(MakeUnique<::Chaos::Softs::FCollectionPropertyFacade>(PropertyCollection));
+
+			}
 		}
 		else
 		{
-			PropertyCollection->Reset();
-			CollectionPropertyMutableFacade.DefineSchema();
+			PropertyCollections.Reset();
+			CollectionPropertyFacades.Reset();
 		}
-
-		// Now the component property facade must be updated
-		*CollectionPropertyFacade = MoveTemp(CollectionPropertyMutableFacade);
 	}
 	else
 	{
@@ -153,10 +158,6 @@ void UChaosClothComponent::OnRegister()
 	// Update the component bone transforms (for colliders) from the cloth asset until these are animated from a leader component
 	UpdateComponentSpaceTransforms();
 
-	// Create the runtime simulation property 
-	PropertyCollection = MakeShared<FManagedArrayCollection>();
-	CollectionPropertyFacade = MakeUnique<::Chaos::Softs::FCollectionPropertyFacade>(PropertyCollection);
-
 	// Fill up the property collection with the original cloth asset properties
 	ResetConfigProperties();
 
@@ -175,8 +176,8 @@ void UChaosClothComponent::OnUnregister()
 	ClothSimulationProxy.Reset();
 
 	// Release the runtime simulation collection and facade
-	PropertyCollection.Reset();
-	CollectionPropertyFacade.Reset();
+	CollectionPropertyFacades.Empty();
+	PropertyCollections.Empty();
 }
 
 bool UChaosClothComponent::IsComponentTickEnabled() const
@@ -188,6 +189,9 @@ void UChaosClothComponent::TickComponent(float DeltaTime, enum ELevelTick TickTy
 {
 	CSV_SCOPED_TIMING_STAT_EXCLUSIVE(Physics);
 	QUICK_SCOPE_CYCLE_COUNTER(STAT_ClothComponentTick);
+	
+	// Tick USkinnedMeshComponent first so it will update the predicted lod
+	Super::TickComponent(DeltaTime, TickType, ThisTickFunction);
 
 	// TODO: Fields
 	//if (ClothingSimulation)
@@ -208,8 +212,6 @@ void UChaosClothComponent::TickComponent(float DeltaTime, enum ELevelTick TickTy
 	{
 		HandleExistingParallelSimulation();
 	}
-
-	Super::TickComponent(DeltaTime, TickType, ThisTickFunction);
 }
 
 bool UChaosClothComponent::RequiresPreEndOfFrameSync() const
@@ -356,10 +358,11 @@ void UChaosClothComponent::StartNewParallelSimulation(float DeltaTime)
 	{
 		CSV_SCOPED_TIMING_STAT(Animation, Cloth);
 		const bool bIsSimulating = ClothSimulationProxy->Tick_GameThread(DeltaTime);
+		const int32 CurrentLOD = GetPredictedLODLevel();
 
-		if (bIsSimulating && CollectionPropertyFacade.IsValid())
+		if (bIsSimulating && CollectionPropertyFacades.IsValidIndex(CurrentLOD) && CollectionPropertyFacades[CurrentLOD].IsValid())
 		{
-			CollectionPropertyFacade->ClearDirtyFlags();
+			CollectionPropertyFacades[CurrentLOD]->ClearDirtyFlags();
 		}
 	}
 }
diff --git a/Engine/Plugins/Experimental/ChaosClothAsset/Source/ChaosClothAssetEngine/Private/ChaosClothAsset/ClothSimulationProxy.cpp b/Engine/Plugins/Experimental/ChaosClothAsset/Source/ChaosClothAssetEngine/Private/ChaosClothAsset/ClothSimulationProxy.cpp
--- a/Engine/Plugins/Experimental/ChaosClothAsset/Source/ChaosClothAssetEngine/Private/ChaosClothAsset/ClothSimulationProxy.cpp
+++ b/Engine/Plugins/Experimental/ChaosClothAsset/Source/ChaosClothAssetEngine/Private/ChaosClothAsset/ClothSimulationProxy.cpp
@@ -132,7 +132,7 @@ PRAGMA_ENABLE_DEPRECATION_WARNINGS
 		//Colliders[ColliderIndex]->SetCollisionData(&ExternalCollisionData);  // TODO: External collision data
 
 		// Create cloth config simulation thread object
-		const int32 ClothConfigIndex = Configs.Emplace(MakeUnique<FClothingSimulationConfig>(ClothComponent.GetPropertyCollection()));
+		const int32 ClothConfigIndex = Configs.Emplace(MakeUnique<FClothingSimulationConfig>(ClothComponent.GetPropertyCollections()));
 
 		// Create cloth simulation thread object
 		constexpr uint32 GroupId = 0;
@@ -154,7 +154,7 @@ PRAGMA_ENABLE_DEPRECATION_WARNINGS
 		NumDynamicParticles = Cloths[ClothIndex]->GetNumActiveDynamicParticles();
 
 		// Create solver config simulation thread object
-		const int32 SolverConfigIndex = Configs.Emplace(MakeUnique<FClothingSimulationConfig>(ClothComponent.GetPropertyCollection()));  // TODO: Use a separate solver config for outfits
+		const int32 SolverConfigIndex = Configs.Emplace(MakeUnique<FClothingSimulationConfig>(ClothComponent.GetPropertyCollections()));  // TODO: Use a separate solver config for outfits
 		Solver->SetConfig(Configs[SolverConfigIndex].Get());
 
 		// Set start pose (update the context, then the solver without advancing the simulation)
@@ -256,6 +256,15 @@ PRAGMA_ENABLE_DEPRECATION_WARNINGS
 		static const float SimulationTimeDecay = 0.03f; // 0.03 seems to provide a good rate of update for the instant average
 		SimulationTime = PrevSimulationTime ? PrevSimulationTime + (CurrSimulationTime - PrevSimulationTime) * SimulationTimeDecay : CurrSimulationTime;
 
+		// Update particle counts (could have changed if lod changed)
+		NumKinematicParticles = 0;
+		NumDynamicParticles = 0;
+		for (const TUniquePtr<FClothingSimulationCloth>& Cloth : Cloths)
+		{
+			NumKinematicParticles += Cloth->GetNumActiveKinematicParticles();
+			NumDynamicParticles += Cloth->GetNumActiveDynamicParticles();
+		}
+
 		// Visualization
 		Visualization->DrawPhysMeshWired();
 		Visualization->DrawCollision();
@@ -446,8 +455,10 @@ PRAGMA_ENABLE_DEPRECATION_WARNINGS
 		// Replace physics thread's configs with the game thread's configs
 		for (const TUniquePtr<::Chaos::FClothingSimulationConfig>& Config : Configs)
 		{
-			Config->Initialize(ClothComponent.GetPropertyCollection());  // TODO: Outfit and multi-cloths/solver config update
+			Config->Initialize(ClothComponent.GetPropertyCollections());  // TODO: Outfit and multi-cloths/solver config update
 		}
+		// TODO: separate solver config lod from cloth component
+		Solver->SetSolverLOD(ClothSimulationContext->LodIndex);
 	}
 
 	void FClothSimulationProxy::FillSimulationContext(float DeltaTime, bool bIsInitialization, FClothingSimulationCacheData* CacheData)
diff --git a/Engine/Plugins/Experimental/ChaosClothAsset/Source/ChaosClothAssetEngine/Public/ChaosClothAsset/ClothComponent.h b/Engine/Plugins/Experimental/ChaosClothAsset/Source/ChaosClothAssetEngine/Public/ChaosClothAsset/ClothComponent.h
--- a/Engine/Plugins/Experimental/ChaosClothAsset/Source/ChaosClothAssetEngine/Public/ChaosClothAsset/ClothComponent.h
+++ b/Engine/Plugins/Experimental/ChaosClothAsset/Source/ChaosClothAssetEngine/Public/ChaosClothAsset/ClothComponent.h
@@ -91,11 +91,14 @@ public:
 	void RecreateClothSimulationProxy();
 
 	/**
-	 * Return the property collection holding the runtime properties for this cloth component.
+	 * Return the property collections holding the runtime properties for this cloth component (one per LOD).
 	 * This might be different from the cloth asset's since the component's properties can be modified in code or in blueprints.
 	 * This could also be different from the cloth simulation object until the cloth simulation thread synchronise the properties.
 	 */
-	TSharedPtr<const FManagedArrayCollection> GetPropertyCollection() const { return TSharedPtr<const FManagedArrayCollection>(PropertyCollection); }
+	const TArray<TSharedPtr<const FManagedArrayCollection>>& GetPropertyCollections() const 
+	{
+		return reinterpret_cast<const TArray<TSharedPtr<const FManagedArrayCollection>>&>(PropertyCollections);
+	}
 
 	const UE::Chaos::ClothAsset::FClothSimulationProxy* GetClothSimulationProxy() const { return ClothSimulationProxy.Get(); }
 
@@ -178,8 +181,8 @@ private:
 	UPROPERTY()
 	float BlendWeight = 1.f;
 
-	TSharedPtr<FManagedArrayCollection> PropertyCollection;
-	TUniquePtr<::Chaos::Softs::FCollectionPropertyFacade> CollectionPropertyFacade;
+	TArray<TSharedPtr<FManagedArrayCollection>> PropertyCollections;
+	TArray<TUniquePtr<::Chaos::Softs::FCollectionPropertyFacade>> CollectionPropertyFacades;
 
 	TSharedPtr<UE::Chaos::ClothAsset::FClothSimulationProxy> ClothSimulationProxy;
 };
diff --git a/Engine/Plugins/Experimental/ChaosClothAssetEditor/Source/ChaosClothAssetDataflowNodes/Private/ChaosClothAsset/MergeClothCollectionsNode.cpp b/Engine/Plugins/Experimental/ChaosClothAssetEditor/Source/ChaosClothAssetDataflowNodes/Private/ChaosClothAsset/MergeClothCollectionsNode.cpp
--- a/Engine/Plugins/Experimental/ChaosClothAssetEditor/Source/ChaosClothAssetDataflowNodes/Private/ChaosClothAsset/MergeClothCollectionsNode.cpp
+++ b/Engine/Plugins/Experimental/ChaosClothAssetEditor/Source/ChaosClothAssetDataflowNodes/Private/ChaosClothAsset/MergeClothCollectionsNode.cpp
@@ -2,6 +2,7 @@
 
 #include "ChaosClothAsset/MergeClothCollectionsNode.h"
 #include "ChaosClothAsset/CollectionClothFacade.h"
+#include "Chaos/CollectionPropertyFacade.h"
 #include "Dataflow/DataflowInputOutput.h"
 
 #include UE_INLINE_GENERATED_CPP_BY_NAME(MergeClothCollectionsNode)
@@ -20,6 +21,7 @@ void FChaosClothAssetMergeClothCollectionsNode::Evaluate(Dataflow::FContext& Con
 	if (Out->IsA<FManagedArrayCollection>(&Collection))
 	{
 		using namespace UE::Chaos::ClothAsset;
+		using namespace Chaos::Softs;
 
 		// Evaluate in collection
 		FManagedArrayCollection InCollection = GetValue<FManagedArrayCollection>(Context, &Collection);
@@ -35,6 +37,9 @@ void FChaosClothAssetMergeClothCollectionsNode::Evaluate(Dataflow::FContext& Con
 			ClothFacade.DefineSchema();
 		}
 
+		FCollectionPropertyMutableFacade PropertyFacade(ClothCollection);
+		bAreAnyValid |= PropertyFacade.IsValid();
+
 		// Iterate through the inputs and append them to LOD 0
 		const TArray<const FManagedArrayCollection*> Collections = GetCollections();
 		for (int32 InputIndex = 1; InputIndex < Collections.Num(); ++InputIndex)
@@ -47,6 +52,14 @@ void FChaosClothAssetMergeClothCollectionsNode::Evaluate(Dataflow::FContext& Con
 				ClothFacade.Append(OtherClothFacade);
 				bAreAnyValid = true;
 			}
+			// Copy properties
+			const FCollectionPropertyConstFacade OtherPropertyFacade(OtherClothCollection);
+			if (OtherPropertyFacade.IsValid())
+			{
+				constexpr bool bUpdateExistingProperties = true; // Want last one wins.
+				PropertyFacade.Append(OtherClothCollection.ToSharedPtr(), bUpdateExistingProperties);
+				bAreAnyValid = true;
+			}
 		}
 
 		// Set the output
diff --git a/Engine/Plugins/Experimental/ChaosClothAssetEditor/Source/ChaosClothAssetDataflowNodes/Private/ChaosClothAsset/SimulationDefaultConfigNode.cpp b/Engine/Plugins/Experimental/ChaosClothAssetEditor/Source/ChaosClothAssetDataflowNodes/Private/ChaosClothAsset/SimulationDefaultConfigNode.cpp
--- a/Engine/Plugins/Experimental/ChaosClothAssetEditor/Source/ChaosClothAssetDataflowNodes/Private/ChaosClothAsset/SimulationDefaultConfigNode.cpp
+++ b/Engine/Plugins/Experimental/ChaosClothAssetEditor/Source/ChaosClothAssetDataflowNodes/Private/ChaosClothAsset/SimulationDefaultConfigNode.cpp
@@ -56,10 +56,10 @@ void FChaosClothAssetSimulationDefaultConfigNode::Evaluate(Dataflow::FContext& C
 				FClothingSimulationConfig ClothingSimulationConfig;
 				ClothingSimulationConfig.Initialize(SimulationConfig.Get(), SharedSimulationConfig.Get());
 
-				ClothingSimulationConfig.GetPropertyCollection()->CopyTo(&ClothCollection.Get());
+				ClothingSimulationConfig.GetPropertyCollection(0)->CopyTo(&ClothCollection.Get());
 
 				// Generate tethers
-				const FCollectionPropertyConstFacade& Properties = ClothingSimulationConfig.GetProperties();
+				const FCollectionPropertyConstFacade& Properties = ClothingSimulationConfig.GetProperties(0);
 				constexpr bool bUseGeodesicTethersDefault = true;
 				const bool bUseGeodesicTethers = Properties.GetValue<bool>(TEXT("UseGeodesicTethers"), bUseGeodesicTethersDefault);
 				// Use the "MaxDistance" weight map to generate tethers. This follows legacy behavior.
diff --git a/Engine/Plugins/Experimental/ChaosClothAssetEditor/Source/ChaosClothAssetDataflowNodes/Private/ChaosClothAsset/TerminalNode.cpp b/Engine/Plugins/Experimental/ChaosClothAssetEditor/Source/ChaosClothAssetDataflowNodes/Private/ChaosClothAsset/TerminalNode.cpp
--- a/Engine/Plugins/Experimental/ChaosClothAssetEditor/Source/ChaosClothAssetDataflowNodes/Private/ChaosClothAsset/TerminalNode.cpp
+++ b/Engine/Plugins/Experimental/ChaosClothAssetEditor/Source/ChaosClothAssetDataflowNodes/Private/ChaosClothAsset/TerminalNode.cpp
@@ -93,7 +93,8 @@ void FChaosClothAssetTerminalNode::SetAssetValue(TObjectPtr<UObject> Asset, Data
 				}
 
 				// Set properties
-				Chaos::Softs::FCollectionPropertyMutableFacade(ClothCollection).Append(*InClothCollection);
+				constexpr bool bUpdateExistingProperties = false;
+				Chaos::Softs::FCollectionPropertyMutableFacade(ClothCollection).Append(InClothCollection.ToSharedPtr(), bUpdateExistingProperties);
 
 				// Set physics asset only with LOD 0 at the moment
 				if (LodIndex == 0)
diff --git a/Engine/Plugins/Experimental/ChaosClothAssetEditor/Source/ChaosClothAssetEditor/Private/ChaosClothAsset/ClothEditor3DViewportClient.cpp b/Engine/Plugins/Experimental/ChaosClothAssetEditor/Source/ChaosClothAssetEditor/Private/ChaosClothAsset/ClothEditor3DViewportClient.cpp
--- a/Engine/Plugins/Experimental/ChaosClothAssetEditor/Source/ChaosClothAssetEditor/Private/ChaosClothAsset/ClothEditor3DViewportClient.cpp
+++ b/Engine/Plugins/Experimental/ChaosClothAssetEditor/Source/ChaosClothAssetEditor/Private/ChaosClothAsset/ClothEditor3DViewportClient.cpp
@@ -242,6 +242,41 @@ bool FChaosClothAssetEditor3DViewportClient::IsSimulationEnabled() const
 	return false;
 }
 
+void FChaosClothAssetEditor3DViewportClient::SetLODModel(int32 LODIndex)
+{
+	if (ClothEdMode)
+	{
+		ClothEdMode->SetLODModel(LODIndex);
+	}
+}
+
+bool FChaosClothAssetEditor3DViewportClient::IsLODModelSelected(int32 LODIndex) const
+{
+	if (ClothEdMode)
+	{
+		return ClothEdMode->IsLODModelSelected(LODIndex);
+	}
+	return false;
+}
+
+int32 FChaosClothAssetEditor3DViewportClient::GetLODModel() const
+{
+	if (ClothEdMode)
+	{
+		return ClothEdMode->GetLODModel();
+	}
+	return INDEX_NONE;
+}
+
+int32 FChaosClothAssetEditor3DViewportClient::GetNumLODs() const
+{
+	if (ClothEdMode)
+	{
+		return ClothEdMode->GetNumLODs();
+	}
+	return 0;
+}
+
 void FChaosClothAssetEditor3DViewportClient::Draw(const FSceneView* View, FPrimitiveDrawInterface* PDI)
 {
 	FEditorViewportClient::Draw(View, PDI);
diff --git a/Engine/Plugins/Experimental/ChaosClothAssetEditor/Source/ChaosClothAssetEditor/Private/ChaosClothAsset/ClothEditor3DViewportClient.h b/Engine/Plugins/Experimental/ChaosClothAssetEditor/Source/ChaosClothAssetEditor/Private/ChaosClothAsset/ClothEditor3DViewportClient.h
--- a/Engine/Plugins/Experimental/ChaosClothAssetEditor/Source/ChaosClothAssetEditor/Private/ChaosClothAsset/ClothEditor3DViewportClient.h
+++ b/Engine/Plugins/Experimental/ChaosClothAssetEditor/Source/ChaosClothAssetEditor/Private/ChaosClothAsset/ClothEditor3DViewportClient.h
@@ -58,6 +58,12 @@ public:
 	void SetEnableSimulation(bool bEnable);
 	bool IsSimulationEnabled() const;
 
+	// LODIndex == INDEX_NONE is LOD Auto
+	void SetLODModel(int32 LODIndex);
+	bool IsLODModelSelected(int32 LODIndex) const;
+	int32 GetLODModel() const;
+	int32 GetNumLODs() const;
+
 	FBox PreviewBoundingBox() const;
 
 	TWeakPtr<FChaosClothPreviewScene> GetClothPreviewScene();
diff --git a/Engine/Plugins/Experimental/ChaosClothAssetEditor/Source/ChaosClothAssetEditor/Private/ChaosClothAsset/ClothEditorCommands.cpp b/Engine/Plugins/Experimental/ChaosClothAssetEditor/Source/ChaosClothAssetEditor/Private/ChaosClothAsset/ClothEditorCommands.cpp
--- a/Engine/Plugins/Experimental/ChaosClothAssetEditor/Source/ChaosClothAssetEditor/Private/ChaosClothAsset/ClothEditorCommands.cpp
+++ b/Engine/Plugins/Experimental/ChaosClothAssetEditor/Source/ChaosClothAssetEditor/Private/ChaosClothAsset/ClothEditorCommands.cpp
@@ -61,6 +61,8 @@ void FChaosClothAssetEditorCommands::RegisterCommands()
 	UI_COMMAND(HardResetSimulation, "HardResetClothSimulation", "Hard reset simulation", EUserInterfaceActionType::ToggleButton, FInputChord(EModifierKey::Control | EModifierKey::Alt, EKeys::C));
 	UI_COMMAND(ToggleSimulationSuspended, "ToggleSimulationSuspended", "Toggle simulation suspended", EUserInterfaceActionType::ToggleButton, FInputChord());
 
+	UI_COMMAND(LODAuto, "LOD Auto", "Automatically select LOD", EUserInterfaceActionType::RadioButton, FInputChord());
+	UI_COMMAND(LOD0, "LOD 0", "Force select LOD 0", EUserInterfaceActionType::RadioButton, FInputChord());
 }
 
 void FChaosClothAssetEditorCommands::GetToolDefaultObjectList(TArray<UInteractiveTool*>& ToolCDOs)
diff --git a/Engine/Plugins/Experimental/ChaosClothAssetEditor/Source/ChaosClothAssetEditor/Private/ChaosClothAsset/ClothEditorCommands.h b/Engine/Plugins/Experimental/ChaosClothAssetEditor/Source/ChaosClothAssetEditor/Private/ChaosClothAsset/ClothEditorCommands.h
--- a/Engine/Plugins/Experimental/ChaosClothAssetEditor/Source/ChaosClothAssetEditor/Private/ChaosClothAsset/ClothEditorCommands.h
+++ b/Engine/Plugins/Experimental/ChaosClothAssetEditor/Source/ChaosClothAssetEditor/Private/ChaosClothAsset/ClothEditorCommands.h
@@ -72,5 +72,7 @@ public:
 	const static FString ToggleSimulationSuspendedIdentifier;
 	TSharedPtr<FUICommandInfo> ToggleSimulationSuspended;
 
+	TSharedPtr<FUICommandInfo> LODAuto;
+	TSharedPtr<FUICommandInfo> LOD0;
 };
 } // namespace UE::Chaos::ClothAsset
diff --git a/Engine/Plugins/Experimental/ChaosClothAssetEditor/Source/ChaosClothAssetEditor/Private/ChaosClothAsset/ClothEditorMode.cpp b/Engine/Plugins/Experimental/ChaosClothAssetEditor/Source/ChaosClothAssetEditor/Private/ChaosClothAsset/ClothEditorMode.cpp
--- a/Engine/Plugins/Experimental/ChaosClothAssetEditor/Source/ChaosClothAssetEditor/Private/ChaosClothAsset/ClothEditorMode.cpp
+++ b/Engine/Plugins/Experimental/ChaosClothAssetEditor/Source/ChaosClothAssetEditor/Private/ChaosClothAsset/ClothEditorMode.cpp
@@ -714,6 +714,41 @@ bool UChaosClothAssetEditorMode::IsSimulationEnabled() const
 	return false;
 }
 
+void UChaosClothAssetEditorMode::SetLODModel(int32 LODIndex)
+{
+	if (PreviewScene && PreviewScene->GetClothComponent())
+	{
+		PreviewScene->GetClothComponent()->SetForcedLOD(LODIndex + 1);
+	}
+}
+
+bool UChaosClothAssetEditorMode::IsLODModelSelected(int32 LODIndex) const
+{
+	if (PreviewScene && PreviewScene->GetClothComponent())
+	{
+		return PreviewScene->GetClothComponent()->GetForcedLOD() == LODIndex + 1;
+	}
+	return false;
+}
+
+int32 UChaosClothAssetEditorMode::GetLODModel() const
+{
+	if (PreviewScene && PreviewScene->GetClothComponent())
+	{
+		return PreviewScene->GetClothComponent()->GetForcedLOD() - 1;
+	}
+	return INDEX_NONE;
+}
+
+int32 UChaosClothAssetEditorMode::GetNumLODs() const
+{
+	if (PreviewScene && PreviewScene->GetClothComponent())
+	{
+		return PreviewScene->GetClothComponent()->GetNumLODs();
+	}
+	return 0;
+}
+
 UDataflowComponent* UChaosClothAssetEditorMode::GetDataflowComponent() const
 {
 	return DataflowComponent;
diff --git a/Engine/Plugins/Experimental/ChaosClothAssetEditor/Source/ChaosClothAssetEditor/Private/ChaosClothAsset/ClothEditorMode.h b/Engine/Plugins/Experimental/ChaosClothAssetEditor/Source/ChaosClothAssetEditor/Private/ChaosClothAsset/ClothEditorMode.h
--- a/Engine/Plugins/Experimental/ChaosClothAssetEditor/Source/ChaosClothAssetEditor/Private/ChaosClothAsset/ClothEditorMode.h
+++ b/Engine/Plugins/Experimental/ChaosClothAssetEditor/Source/ChaosClothAssetEditor/Private/ChaosClothAsset/ClothEditorMode.h
@@ -80,6 +80,12 @@ public:
 	void SetEnableSimulation(bool bEnabled);
 	bool IsSimulationEnabled() const;
 
+	// LODIndex == INDEX_NONE is LOD Auto
+	void SetLODModel(int32 LODIndex);
+	bool IsLODModelSelected(int32 LODIndex) const;
+	int32 GetLODModel() const;
+	int32 GetNumLODs() const;
+
 	UDataflowComponent* GetDataflowComponent() const;
 
 	TObjectPtr<UEditorInteractiveToolsContext> GetActiveToolsContext()
diff --git a/Engine/Plugins/Experimental/ChaosClothAssetEditor/Source/ChaosClothAssetEditor/Private/ChaosClothAsset/SClothCollectionOutliner.cpp b/Engine/Plugins/Experimental/ChaosClothAssetEditor/Source/ChaosClothAssetEditor/Private/ChaosClothAsset/SClothCollectionOutliner.cpp
--- a/Engine/Plugins/Experimental/ChaosClothAssetEditor/Source/ChaosClothAssetEditor/Private/ChaosClothAsset/SClothCollectionOutliner.cpp
+++ b/Engine/Plugins/Experimental/ChaosClothAssetEditor/Source/ChaosClothAssetEditor/Private/ChaosClothAsset/SClothCollectionOutliner.cpp
@@ -202,6 +202,11 @@ namespace ClothCollectionOutlinerHelpers
 		return FString::FromInt(Value);
 	}
 
+	FString AttributeValueToString(uint8 Value)
+	{
+		return FString::FromInt(Value);
+	}
+
 	FString AttributeValueToString(const FString& Value)
 	{
 		return Value;
@@ -304,6 +309,9 @@ namespace ClothCollectionOutlinerHelpers
 		case FManagedArrayCollection::EArrayType::FIntArrayType:
 			ValueAsString = AttributeValueToString<TSet<int32>>(ClothCollection, AttributeName, GroupName, AttributeArrayIndex);
 			break;
+		case FManagedArrayCollection::EArrayType::FUInt8Type:
+			ValueAsString = AttributeValueToString<uint8>(ClothCollection, AttributeName, GroupName, AttributeArrayIndex);
+			break;
 		default:
 			ensure(false);
 			ValueAsString = "(Unknown Data Type)";
diff --git a/Engine/Plugins/Experimental/ChaosClothAssetEditor/Source/ChaosClothAssetEditor/Private/ChaosClothAsset/SClothEditor3DViewport.cpp b/Engine/Plugins/Experimental/ChaosClothAssetEditor/Source/ChaosClothAssetEditor/Private/ChaosClothAsset/SClothEditor3DViewport.cpp
--- a/Engine/Plugins/Experimental/ChaosClothAssetEditor/Source/ChaosClothAssetEditor/Private/ChaosClothAsset/SClothEditor3DViewport.cpp
+++ b/Engine/Plugins/Experimental/ChaosClothAssetEditor/Source/ChaosClothAssetEditor/Private/ChaosClothAsset/SClothEditor3DViewport.cpp
@@ -118,6 +118,38 @@ void SChaosClothAssetEditor3DViewport::BindCommands()
 		FCanExecuteAction::CreateLambda([this]() { return true; }),
 		FIsActionChecked::CreateLambda([this]() { return StaticCastSharedPtr<FChaosClothAssetEditor3DViewportClient>(Client)->IsSimulationSuspended(); }) );
 
+	ToolkitCommandList->MapAction(
+		CommandInfos.LODAuto,
+		FExecuteAction::CreateLambda([this]()
+		{
+			TSharedPtr<FChaosClothAssetEditor3DViewportClient> ClothViewportClient = StaticCastSharedPtr<FChaosClothAssetEditor3DViewportClient>(Client);
+			ClothViewportClient->SetLODModel(INDEX_NONE);
+		}),
+		FCanExecuteAction(),
+		FIsActionChecked::CreateLambda([this]()
+		{
+			TSharedPtr<FChaosClothAssetEditor3DViewportClient> ClothViewportClient = StaticCastSharedPtr<FChaosClothAssetEditor3DViewportClient>(Client);
+			return ClothViewportClient->IsLODModelSelected(INDEX_NONE);
+		}
+		));
+
+	ToolkitCommandList->MapAction(
+		CommandInfos.LOD0,
+		FExecuteAction::CreateLambda([this]()
+		{
+			TSharedPtr<FChaosClothAssetEditor3DViewportClient> ClothViewportClient = StaticCastSharedPtr<FChaosClothAssetEditor3DViewportClient>(Client);
+			ClothViewportClient->SetLODModel(0);
+		}),
+		FCanExecuteAction(),
+		FIsActionChecked::CreateLambda([this]()
+		{
+			TSharedPtr<FChaosClothAssetEditor3DViewportClient> ClothViewportClient = StaticCastSharedPtr<FChaosClothAssetEditor3DViewportClient>(Client);
+			return ClothViewportClient->IsLODModelSelected(0);
+		}
+		));
+
+	// all other LODs will be added dynamically
+
 }
 
 TSharedPtr<SWidget> SChaosClothAssetEditor3DViewport::MakeViewportToolbar()
diff --git a/Engine/Plugins/Experimental/ChaosClothAssetEditor/Source/ChaosClothAssetEditor/Private/ChaosClothAsset/SClothEditor3DViewportToolBar.cpp b/Engine/Plugins/Experimental/ChaosClothAssetEditor/Source/ChaosClothAssetEditor/Private/ChaosClothAsset/SClothEditor3DViewportToolBar.cpp
--- a/Engine/Plugins/Experimental/ChaosClothAssetEditor/Source/ChaosClothAssetEditor/Private/ChaosClothAsset/SClothEditor3DViewportToolBar.cpp
+++ b/Engine/Plugins/Experimental/ChaosClothAssetEditor/Source/ChaosClothAssetEditor/Private/ChaosClothAsset/SClothEditor3DViewportToolBar.cpp
@@ -6,6 +6,9 @@
 #include "ChaosClothAsset/ClothEditorSimulationVisualization.h"
 #include "ChaosClothAsset/ClothEditorCommands.h"
 #include "ChaosClothAsset/ClothEditorStyle.h"
+#include "ChaosClothAsset/ClothEditorPreviewScene.h"
+#include "ChaosClothAsset/ClothAsset.h"
+#include "ChaosClothAsset/ClothComponent.h"
 #include "EditorViewportCommands.h"
 #include "Styling/AppStyle.h"
 #include "Framework/MultiBox/MultiBoxBuilder.h"
@@ -60,12 +63,25 @@ void SChaosClothAssetEditor3DViewportToolBar::Construct(const FArguments& InArgs
 		];
 
 	MainBoxPtr->AddSlot()
+		.AutoWidth()
 		.Padding(ToolbarSlotPadding)
 		.HAlign(HAlign_Left)
 		[
 			MakeDisplayToolBar(InArgs._Extenders)
 		];
 
+	MainBoxPtr->AddSlot()
+		.AutoWidth()
+		.Padding(ToolbarSlotPadding)
+		.HAlign(HAlign_Left)
+		[
+			SNew(SEditorViewportToolbarMenu)
+			.Label(this, &SChaosClothAssetEditor3DViewportToolBar::GetLODMenuLabel)
+			.Cursor(EMouseCursor::Default)
+			.ParentToolBar(SharedThis(this))
+			.OnGetMenuContent(this, &SChaosClothAssetEditor3DViewportToolBar::MakeLODMenu)
+		];
+
 	MainBoxPtr->AddSlot()
 		.Padding(ToolbarSlotPadding)
 		.HAlign(HAlign_Right)
@@ -165,23 +181,76 @@ FText SChaosClothAssetEditor3DViewportToolBar::GetDisplayString() const
 	return FText();
 }
 
-void SChaosClothAssetEditor3DViewportToolBar::ExtendOptionsMenu(FMenuBuilder& OptionsMenuBuilder) const
+FText SChaosClothAssetEditor3DViewportToolBar::GetLODMenuLabel() const
+{
+	using namespace UE::Chaos::ClothAsset;
+
+	FText Label = LOCTEXT("LODMenu_AutoLabel", "LOD Auto"); 
+
+	TSharedRef<FChaosClothAssetEditor3DViewportClient> ViewportClient = 
+		StaticCastSharedPtr<FChaosClothAssetEditor3DViewportClient>(ChaosClothAssetEditor3DViewportPtr.Pin()->GetViewportClient()).ToSharedRef();
+
+	const int32 LODSelectionType = ViewportClient->GetLODModel();
+	if (LODSelectionType >= 0)
+	{
+		const FString TitleLabel = FString::Printf(TEXT("LOD %d"), LODSelectionType);
+		Label = FText::FromString(TitleLabel);
+	}
+	return Label;
+}
+
+TSharedRef<SWidget> SChaosClothAssetEditor3DViewportToolBar::MakeLODMenu() const
 {
 	using namespace UE::Chaos::ClothAsset;
+	TSharedRef<FChaosClothAssetEditor3DViewportClient> ViewportClient =
+		StaticCastSharedPtr<FChaosClothAssetEditor3DViewportClient>(ChaosClothAssetEditor3DViewportPtr.Pin()->GetViewportClient()).ToSharedRef();
+
+	constexpr bool bInShouldCloseWindowAfterMenuSelection = true;
+	FMenuBuilder MenuBuilder(bInShouldCloseWindowAfterMenuSelection, CommandList);
+
+	MenuBuilder.PushCommandList(CommandList.ToSharedRef());
+	const int32 NumLODs = ViewportClient->GetNumLODs();
 
-	constexpr bool bOpenSubMenuOnClick = false;
-	constexpr bool bShouldCloseWindowAfterMenuSelection = false;
-	OptionsMenuBuilder.AddSubMenu(
-		LOCTEXT("ChaosClothAssetEditor_SimulationVisualization", "Simulation Visualization"),
-		LOCTEXT("ChaosClothAssetEditor_SimulationVisualizationToolTip", "Options to control simulation visualization"),
-		FNewMenuDelegate::CreateLambda([this](FMenuBuilder& MenuBuilder)
+	MenuBuilder.BeginSection("ClothAssetPreviewLODs", LOCTEXT("ShowLOD_PreviewLabel", "Preview LODs"));
 	{
-		TSharedRef<FChaosClothAssetEditor3DViewportClient> ViewportClient = StaticCastSharedPtr<FChaosClothAssetEditor3DViewportClient>(ChaosClothAssetEditor3DViewportPtr.Pin()->GetViewportClient()).ToSharedRef();
-		if (FClothEditorSimulationVisualization* const Visualization = ViewportClient->GetSimulationVisualization().Pin().Get())
+		MenuBuilder.AddMenuEntry(FChaosClothAssetEditorCommands::Get().LODAuto);
+		MenuBuilder.AddMenuEntry(FChaosClothAssetEditorCommands::Get().LOD0);
+
+		for (int32 LODIndex = 1; LODIndex < NumLODs; ++LODIndex)
 		{
-			Visualization->ExtendViewportShowMenu(MenuBuilder, ViewportClient);
+			FString TitleLabel = FString::Printf(TEXT("LOD %d"), LODIndex);
+
+			FUIAction Action(FExecuteAction::CreateSP(ViewportClient, &FChaosClothAssetEditor3DViewportClient::SetLODModel, LODIndex),
+				FCanExecuteAction(),
+				FIsActionChecked::CreateSP(ViewportClient, &FChaosClothAssetEditor3DViewportClient::IsLODModelSelected, LODIndex));
+
+			MenuBuilder.AddMenuEntry(FText::FromString(TitleLabel), FText::GetEmpty(), FSlateIcon(), Action, NAME_None, EUserInterfaceActionType::RadioButton);
 		}
-	}), bOpenSubMenuOnClick, FSlateIcon(), bShouldCloseWindowAfterMenuSelection);
+	}
+	MenuBuilder.EndSection();
+	MenuBuilder.PopCommandList();
+	return MenuBuilder.MakeWidget();
+}
+
+void SChaosClothAssetEditor3DViewportToolBar::ExtendOptionsMenu(FMenuBuilder& OptionsMenuBuilder) const
+{
+	using namespace UE::Chaos::ClothAsset;
+	{
+		// Debug visualization
+		constexpr bool bOpenSubMenuOnClick = false;
+		constexpr bool bShouldCloseWindowAfterMenuSelection = false;
+		OptionsMenuBuilder.AddSubMenu(
+			LOCTEXT("SimulationVisualization", "Simulation Visualization"),
+			LOCTEXT("SimulationVisualizationToolTip", "Options to control simulation visualization"),
+			FNewMenuDelegate::CreateLambda([this](FMenuBuilder& MenuBuilder)
+		{
+			TSharedRef<FChaosClothAssetEditor3DViewportClient> ViewportClient = StaticCastSharedPtr<FChaosClothAssetEditor3DViewportClient>(ChaosClothAssetEditor3DViewportPtr.Pin()->GetViewportClient()).ToSharedRef();
+			if (FClothEditorSimulationVisualization* const Visualization = ViewportClient->GetSimulationVisualization().Pin().Get())
+			{
+				Visualization->ExtendViewportShowMenu(MenuBuilder, ViewportClient);
+			}
+		}), bOpenSubMenuOnClick, FSlateIcon(), bShouldCloseWindowAfterMenuSelection);
+	}
 }
 
 #undef LOCTEXT_NAMESPACE
diff --git a/Engine/Plugins/Experimental/ChaosClothAssetEditor/Source/ChaosClothAssetEditor/Private/ChaosClothAsset/SClothEditor3DViewportToolBar.h b/Engine/Plugins/Experimental/ChaosClothAssetEditor/Source/ChaosClothAssetEditor/Private/ChaosClothAsset/SClothEditor3DViewportToolBar.h
--- a/Engine/Plugins/Experimental/ChaosClothAssetEditor/Source/ChaosClothAssetEditor/Private/ChaosClothAsset/SClothEditor3DViewportToolBar.h
+++ b/Engine/Plugins/Experimental/ChaosClothAssetEditor/Source/ChaosClothAssetEditor/Private/ChaosClothAsset/SClothEditor3DViewportToolBar.h
@@ -26,6 +26,8 @@ private:
 	TSharedRef<SWidget> MakeDisplayToolBar(const TSharedPtr<FExtender> InExtenders);
 	TSharedRef<SWidget> MakeToolBar(const TSharedPtr<FExtender> InExtenders);
 	FText GetDisplayString() const;
+	FText GetLODMenuLabel() const;
+	TSharedRef<SWidget> MakeLODMenu() const;
 
 	virtual void ExtendOptionsMenu(FMenuBuilder& OptionsMenuBuilder) const override;
 
diff --git a/Engine/Source/Runtime/Experimental/Chaos/Private/Chaos/CollectionPropertyFacade.cpp b/Engine/Source/Runtime/Experimental/Chaos/Private/Chaos/CollectionPropertyFacade.cpp
--- a/Engine/Source/Runtime/Experimental/Chaos/Private/Chaos/CollectionPropertyFacade.cpp
+++ b/Engine/Source/Runtime/Experimental/Chaos/Private/Chaos/CollectionPropertyFacade.cpp
@@ -215,8 +215,23 @@ namespace Chaos::Softs
 		return INDEX_NONE;
 	}
 
+	void FCollectionPropertyMutableFacade::Append(const TSharedPtr<const FManagedArrayCollection>& InManagedArrayCollection, bool bUpdateExistingProperties)
+	{
+		
+		Update(InManagedArrayCollection, ECollectionPropertyUpdateFlags::AppendNewProperties | (bUpdateExistingProperties ? ECollectionPropertyUpdateFlags::UpdateExistingProperties : ECollectionPropertyUpdateFlags::None));
+	}
+
 	void FCollectionPropertyMutableFacade::Append(const FManagedArrayCollection& InManagedArrayCollection)
 	{
+		constexpr ECollectionPropertyUpdateFlags UpdateFlagsAppendNewOnly = ECollectionPropertyUpdateFlags::AppendNewProperties;
+		Update(MakeShared<const FManagedArrayCollection>(InManagedArrayCollection), UpdateFlagsAppendNewOnly);
+	}
+
+	void FCollectionPropertyMutableFacade::Copy(const FManagedArrayCollection& InManagedArrayCollection)
+	{
+		GetManagedArrayCollection()->Reset();
+		DefineSchema();
+
 		TArray<FName> GroupsToSkip = InManagedArrayCollection.GroupNames();
 		GroupsToSkip.RemoveSingleSwap(PropertyFacadeNames::PropertyGroup);
 
@@ -225,11 +240,85 @@ namespace Chaos::Softs
 		RebuildKeyIndices();
 	}
 
-	void FCollectionPropertyMutableFacade::Copy(const FManagedArrayCollection& InManagedArrayCollection)
+	void FCollectionPropertyMutableFacade::Update(const TSharedPtr<const FManagedArrayCollection>& InManagedArrayCollection, ECollectionPropertyUpdateFlags UpdateFlags)
 	{
-		GetManagedArrayCollection()->Reset();
-		DefineSchema();
-		Append(InManagedArrayCollection);
+		if (UpdateFlags == ECollectionPropertyUpdateFlags::None)
+		{
+			// Nothing to do
+			return;
+		}
+
+		const bool bAppendNewProperties = (UpdateFlags & ECollectionPropertyUpdateFlags::AppendNewProperties) != ECollectionPropertyUpdateFlags::None;
+		const bool bUpdateExistingProperties = (UpdateFlags & ECollectionPropertyUpdateFlags::UpdateExistingProperties) != ECollectionPropertyUpdateFlags::None;
+		const bool bRemoveMissingProperties = (UpdateFlags & ECollectionPropertyUpdateFlags::RemoveMissingProperties) != ECollectionPropertyUpdateFlags::None;
+		const bool bDisableMissingProperties = (UpdateFlags & ECollectionPropertyUpdateFlags::DisableMissingProperties) != ECollectionPropertyUpdateFlags::None;
+
+		FCollectionPropertyConstFacade InPropertyFacade(InManagedArrayCollection);
+		if (InPropertyFacade.IsValid())
+		{
+			DefineSchema();
+
+			const int32 PrevNumKeys = Num();
+
+			if (bAppendNewProperties || bUpdateExistingProperties)
+			{
+				const int32 NumInKeys = InPropertyFacade.Num();
+				for (int32 InKeyIndex = 0; InKeyIndex < NumInKeys; ++InKeyIndex)
+				{
+					const FString& PropertyName = InPropertyFacade.GetKey(InKeyIndex);
+					int32 NewPropertyIndex = GetKeyIndex(PropertyName);
+					if (NewPropertyIndex == INDEX_NONE)
+					{
+						if (bAppendNewProperties)
+						{
+							NewPropertyIndex = AddProperty(PropertyName, (ECollectionPropertyFlags)InPropertyFacade.GetFlags(InKeyIndex));
+						}
+						else
+						{
+							continue;
+						}
+					}
+					else if (!bUpdateExistingProperties)
+					{
+						continue;
+					}
+					else
+					{
+						SetFlags(NewPropertyIndex, InPropertyFacade.GetFlags(InKeyIndex));
+					}
+
+					// Setting as FVector3f since that is the underlying type
+					SetLowValue(NewPropertyIndex, InPropertyFacade.GetLowValue<FVector3f>(InKeyIndex));
+					SetHighValue(NewPropertyIndex, InPropertyFacade.GetHighValue<FVector3f>(InKeyIndex));
+					SetStringValue(NewPropertyIndex, InPropertyFacade.GetStringValue(InKeyIndex));
+				}
+			}
+			if (bRemoveMissingProperties || bDisableMissingProperties)
+			{
+				TArray<int32> KeyIndicesToRemove;
+				for (int32 ExistingKeyIndex = 0; ExistingKeyIndex < PrevNumKeys; ++ExistingKeyIndex)
+				{
+					if (InPropertyFacade.GetKeyIndex(GetKey(ExistingKeyIndex)) == INDEX_NONE)
+					{
+						if (bRemoveMissingProperties)
+						{
+							KeyIndicesToRemove.Add(ExistingKeyIndex);
+						}
+						else if(bDisableMissingProperties)
+						{
+							SetEnabled(ExistingKeyIndex, false);
+						}
+					}
+				}
+
+				if (!KeyIndicesToRemove.IsEmpty())
+				{
+					GetManagedArrayCollection()->RemoveElements(PropertyFacadeNames::PropertyGroup, KeyIndicesToRemove);
+					UpdateArrays();
+					RebuildKeyIndices();
+				}
+			}
+		}
 	}
 
 	int32 FCollectionPropertyMutableFacade::AddWeightedFloatValue(const FString& Key, const FVector2f& Value, bool bEnabled, bool bAnimatable)
diff --git a/Engine/Source/Runtime/Experimental/Chaos/Public/Chaos/CollectionPropertyFacade.h b/Engine/Source/Runtime/Experimental/Chaos/Public/Chaos/CollectionPropertyFacade.h
--- a/Engine/Source/Runtime/Experimental/Chaos/Public/Chaos/CollectionPropertyFacade.h
+++ b/Engine/Source/Runtime/Experimental/Chaos/Public/Chaos/CollectionPropertyFacade.h
@@ -69,9 +69,9 @@ namespace Chaos::Softs
 		int32 GetKeyIndex(const FString& Key) const { const int32* const Index = KeyIndices.Find(Key); return Index ? *Index : INDEX_NONE; }
 
 		/** Return the property name (key) for the specified index. */
+		//~ Values access per index, fast, no check, index must be valid (0 <= KeyIndex < Num())
 		const FString& GetKey(int32 KeyIndex) const { return GetValue<const FString&>(KeyIndex, KeyArray); }
 
-		//~ Values access per index, fast, no check, index must be valid (0 <= KeyIndex < Num())
 		template<typename T, TEMPLATE_REQUIRES(TIsWeightedType<T>::Value)>
 		T GetLowValue(int32 KeyIndex) const { return GetValue<T>(KeyIndex, LowValueArray); }
 
@@ -90,7 +90,6 @@ namespace Chaos::Softs
 		UE_DEPRECATED(5.3, "Use GetStringValue(int32) or GetStringValue(const FString&, const FString&, int32*) instead.")
 		const FString& GetStringValue(int32 KeyIndex, const FString& Default) const { return GetValue<const FString&>(KeyIndex, StringValueArray); }
 
-		UE_DEPRECATED(5.3, "GetFlags is being phased out to promote correct dirtying operations.")
 		uint8 GetFlags(int32 KeyIndex) const { return FlagsArray[KeyIndex]; }
 
 		bool IsEnabled(int32 KeyIndex) const { return HasAnyFlags(KeyIndex, ECollectionPropertyFlags::Enabled); }
@@ -134,12 +133,9 @@ namespace Chaos::Softs
 			return SafeGet(Key, [this](int32 KeyIndex)->FString { return GetStringValue(KeyIndex); }, Default, OutKeyIndex);
 		}
 
-		UE_DEPRECATED(5.3, "GetFlags is being phased out to promote correct dirtying operations.")
 		uint8 GetFlags(const FString& Key, uint8 Default = 0, int32* OutKeyIndex = nullptr) const
 		{
-PRAGMA_DISABLE_DEPRECATION_WARNINGS
 			return SafeGet(Key, [this](int32 KeyIndex)->uint8 { return GetFlags(KeyIndex); }, Default, OutKeyIndex);
-PRAGMA_ENABLE_DEPRECATION_WARNINGS
 		}
 
 		bool IsEnabled(const FString& Key, bool bDefault = false, int32* OutKeyIndex = nullptr) const
@@ -168,6 +164,7 @@ PRAGMA_ENABLE_DEPRECATION_WARNINGS
 		}
 
 	protected:
+
 		// No init constructor for FCollectionPropertyFacade
 		CHAOS_API FCollectionPropertyConstFacade(const TSharedPtr<const FManagedArrayCollection>& InManagedArrayCollection, ENoInit);
 
@@ -255,8 +252,13 @@ PRAGMA_ENABLE_DEPRECATION_WARNINGS
 
 		void SetStringValue(int32 KeyIndex, const FString& Value) { if (GetStringValueArray()[KeyIndex] != Value) { GetStringValueArray()[KeyIndex] = Value; SetStringDirty(KeyIndex); } }
 
-		UE_DEPRECATED(5.3, "SetFlags is being phased out to promote correct dirtying operations.")
-		void SetFlags(int32 KeyIndex, uint8 Flags) { SetValue(KeyIndex, GetFlagsArray(), (ECollectionPropertyFlags)Flags); }
+		/** SetFlags cannot be used to remove dirty flags. Use ClearDirtyFlags to remove dirty flags.*/
+		void SetFlags(int32 KeyIndex, uint8 Flags) 
+		{ 
+			const ECollectionPropertyFlags DirtyFlags = (ECollectionPropertyFlags)GetFlags(KeyIndex) & (ECollectionPropertyFlags::Dirty | ECollectionPropertyFlags::StringDirty);
+
+			SetValue(KeyIndex, GetFlagsArray(), (ECollectionPropertyFlags)Flags | DirtyFlags);
+		}
 
 		void SetEnabled(int32 KeyIndex, bool bEnabled) { EnableFlags(KeyIndex, ECollectionPropertyFlags::Enabled, bEnabled); }
 		void SetAnimatable(int32 KeyIndex, bool bAnimatable) { EnableFlags(KeyIndex, ECollectionPropertyFlags::Animatable, bAnimatable); }
@@ -301,12 +303,9 @@ PRAGMA_ENABLE_DEPRECATION_WARNINGS
 			return SafeSet(Key, [this, &Value](int32 KeyIndex) { SetStringValue(KeyIndex, Value); });
 		}
 
-		UE_DEPRECATED(5.3, "SetFlags is being phased out to promote correct dirtying operations.")
 		int32 SetFlags(const FString& Key, uint8 Flags)
 		{
-PRAGMA_DISABLE_DEPRECATION_WARNINGS
 			return SafeSet(Key, [this, Flags](int32 KeyIndex) { SetFlags(KeyIndex, Flags); });
-PRAGMA_ENABLE_DEPRECATION_WARNINGS
 		}
 
 		int32 SetEnabled(const FString& Key, bool bEnabled)
@@ -375,6 +374,16 @@ PRAGMA_ENABLE_DEPRECATION_WARNINGS
 		CHAOS_API void EnableFlags(int32 KeyIndex, ECollectionPropertyFlags Flags, bool bEnable);
 	};
 
+	enum class ECollectionPropertyUpdateFlags : uint8
+	{
+		None = 0,
+		AppendNewProperties = 1 << 0, /** Add new properties.*/
+		UpdateExistingProperties = 1 << 1, /** Update values of existing properties.*/
+		RemoveMissingProperties = 1 << 2, /** Existing properties not in InManagedArrayCollection will be removed */
+		DisableMissingProperties = 1 << 3 /** Existing properties not in InManagedArrayCollection will be disabled. (Does nothing if RemoveMissingProperties is set).*/
+	};
+	ENUM_CLASS_FLAGS(ECollectionPropertyUpdateFlags)
+
 	/**
 	 * Defines common API for reading and writing, and adding/removing simulation properties data and metadata.
 	 * This is mainly used for the cloth simulation properties to provide
@@ -411,17 +420,32 @@ PRAGMA_ENABLE_DEPRECATION_WARNINGS
 		/**
 		 * Append all properties and values from an existing collection to this property collection.
 		 * This won't copy any other groups, only data from PropertyGroup.
-		 * Any pre-existing data will be preserved.
+		 * Modified properties will be marked dirty.
+		 * 
+		 * @param bUpdateExistingProperties: Update existing property values. Otherwise, only new properties will be added.
 		 */
+		CHAOS_API void Append(const TSharedPtr<const FManagedArrayCollection>& InManagedArrayCollection, bool bUpdateExistingProperties);
+
+		UE_DEPRECATED(5.3, "Use SharedPtr version of Append to avoid additional copy.")
 		CHAOS_API void Append(const FManagedArrayCollection& InManagedArrayCollection);
 
 		/**
 		 * Copy all properties and values from an existing collection to this property collection.
+		 * Dirty flags will be copied directly.
 		 * This won't copy any other groups, only data from PropertyGroup.
 		 * Any pre-xisting data will be removed/replaced.
 		 */
 		CHAOS_API void Copy(const FManagedArrayCollection& InManagedArrayCollection);
 
+
+		/**
+		 * Update all properties and values from an existing collection to this property collection.
+		 * This won't copy any other groups, only data from PropertyGroup.
+		 * Modified properties will be marked dirty.
+		 */
+		CHAOS_API void Update(const TSharedPtr<const FManagedArrayCollection>& InManagedArrayCollection, ECollectionPropertyUpdateFlags UpdateFlags);
+
+
 		//~ Add values
 		template<typename T, TEMPLATE_REQUIRES(TIsWeightedType<T>::Value)>
 		inline int32 AddWeightedValue(const FString& Key, const T& LowValue, const T& HighValue, ECollectionPropertyFlags Flags = ECollectionPropertyFlags::Enabled);
@@ -583,7 +607,7 @@ PRAGMA_ENABLE_DEPRECATION_WARNINGS \
 	bool Is##PropertyName##Mutable(const FCollectionPropertyConstFacade& PropertyCollection) const \
 	{ \
 		return PropertyName##Index != INDEX_NONE && \
-			PropertyCollection.IsEnabled(PropertyName##Index) && PropertyCollection.IsAnimatable(PropertyName##Index) && PropertyCollection.IsDirty(PropertyName##Index); \
+			PropertyCollection.IsEnabled(PropertyName##Index) && PropertyCollection.IsAnimatable(PropertyName##Index) && (PropertyCollection.IsDirty(PropertyName##Index) || PropertyCollection.IsStringDirty(PropertyName##Index)); \
 	} \
 	struct F##PropertyName##Index \
 	{ \
