Speculative fix for crash caused by not initializing the cloth simulation solver config when the solver is reset now that it's required after CL26744966.
I have not yet been able to repro the actual bug, but this seems like the most likely cause.

diff --git a/Engine/Plugins/Experimental/ChaosCloth/Source/Chaos/Private/ChaosCloth/ChaosClothingSimulationSolver.cpp b/Engine/Plugins/Experimental/ChaosCloth/Source/Chaos/Private/ChaosCloth/ChaosClothingSimulationSolver.cpp
--- a/Engine/Plugins/Experimental/ChaosCloth/Source/Chaos/Private/ChaosCloth/ChaosClothingSimulationSolver.cpp
+++ b/Engine/Plugins/Experimental/ChaosCloth/Source/Chaos/Private/ChaosCloth/ChaosClothingSimulationSolver.cpp
@@ -109,6 +109,8 @@ FClothingSimulationSolver::FClothingSimulationSolver()
 	, bIsClothGravityOverrideEnabled(false)
 	, bEnableSolver(true)
 {
+	SetConfig(nullptr); // This will generate a local default config so we have something to use if there are no cloth assets..
+
 	Softs::FSolverParticles LocalParticles;
 	Softs::FSolverRigidParticles RigidParticles;
 	Evolution.Reset(
