Chaos Cloth Asset: handle different weight map name lists for different LODs.

diff --git a/Engine/Plugins/Experimental/ChaosCloth/Source/Chaos/Private/ChaosCloth/ChaosClothingSimulationCloth.cpp b/Engine/Plugins/Experimental/ChaosCloth/Source/Chaos/Private/ChaosCloth/ChaosClothingSimulationCloth.cpp
--- a/Engine/Plugins/Experimental/ChaosCloth/Source/Chaos/Private/ChaosCloth/ChaosClothingSimulationCloth.cpp
+++ b/Engine/Plugins/Experimental/ChaosCloth/Source/Chaos/Private/ChaosCloth/ChaosClothingSimulationCloth.cpp
@@ -543,26 +543,25 @@ PRAGMA_DISABLE_DEPRECATION_WARNINGS  // TODO: CHAOS_IS_CLOTHINGSIMULATIONMESH_AB
 void FClothingSimulationCloth::SetMesh(FClothingSimulationMesh* InMesh)
 {
 	Mesh = InMesh;
-
-	// Build weight map index lookup
-	const TMap<FString, int32> WeightMapIndices = Mesh ? Mesh->GetWeightMapIndices() : TMap<FString, int32>();
-
+	
 	// Reset LODs
 	const int32 NumLODs = Mesh ? Mesh->GetNumLODs() : 0;
 	LODData.Reset(NumLODs);
 	for (int32 LODIndex = 0; LODIndex < NumLODs; ++LODIndex)
 	{
 		// Regenerate LOD weight maps lookup map
+		const TArray<FName> WeightMapNames = Mesh->GetWeightMapNames(LODIndex);
 		TMap<FString, TConstArrayView<FRealSingle>> WeightMaps;
-		WeightMaps.Reserve(WeightMapIndices.Num());
+		WeightMaps.Reserve(WeightMapNames.Num());
 
 		const TArray<TConstArrayView<FRealSingle>> WeightMapArray = Mesh->GetWeightMaps(LODIndex);
+		ensure(WeightMapArray.Num() == WeightMapNames.Num());
 
-		for (const TPair<FString, int32>& WeightMapIndex : WeightMapIndices)
+		for (int32 WeightMapIndex = 0; WeightMapIndex < WeightMapNames.Num(); ++WeightMapIndex)
 		{
-			WeightMaps.Add(WeightMapIndex.Key,
-				WeightMapArray.IsValidIndex(WeightMapIndex.Value) ?
-					WeightMapArray[WeightMapIndex.Value] :
+			WeightMaps.Add(WeightMapNames[WeightMapIndex].ToString(),
+				WeightMapArray.IsValidIndex(WeightMapIndex) ?
+					WeightMapArray[WeightMapIndex] :
 					TConstArrayView<FRealSingle>());
 		}
 
diff --git a/Engine/Plugins/Experimental/ChaosCloth/Source/Chaos/Private/ChaosCloth/ChaosClothingSimulationMesh.cpp b/Engine/Plugins/Experimental/ChaosCloth/Source/Chaos/Private/ChaosCloth/ChaosClothingSimulationMesh.cpp
--- a/Engine/Plugins/Experimental/ChaosCloth/Source/Chaos/Private/ChaosCloth/ChaosClothingSimulationMesh.cpp
+++ b/Engine/Plugins/Experimental/ChaosCloth/Source/Chaos/Private/ChaosCloth/ChaosClothingSimulationMesh.cpp
@@ -149,38 +149,35 @@ TConstArrayView<uint32> FClothingSimulationMesh::GetPatternToWeldedIndices(int32
 	return TConstArrayView<uint32>();
 }
 
-TArray<FName> FClothingSimulationMesh::GetWeightMapNames() const
+TArray<FName> FClothingSimulationMesh::GetWeightMapNames(int32 LODIndex) const
 {
 	TArray<FName> WeightMapNames;
-
-	const UEnum* const ChaosWeightMapTargetEnum = StaticEnum<EChaosWeightMapTarget>();
-	const int32 NumWeightMaps = (int32)ChaosWeightMapTargetEnum->GetMaxEnumValue() + 1;
-
-	WeightMapNames.SetNum(NumWeightMaps);
-
-	for (int32 EnumIndex = 0; EnumIndex < ChaosWeightMapTargetEnum->NumEnums(); ++EnumIndex)
+	if (IsValidLODIndex(LODIndex))
 	{
-		const int32 TargetIndex = (int32)ChaosWeightMapTargetEnum->GetValueByIndex(EnumIndex);
-		WeightMapNames[TargetIndex] = FName(ChaosWeightMapTargetEnum->GetNameStringByIndex(EnumIndex));
+		// This must match the order of GetWeightMaps
+		const FClothLODDataCommon& ClothLODData = Asset->LodData[LODIndex];
+		const FClothPhysicalMeshData& ClothPhysicalMeshData = ClothLODData.PhysicalMeshData;
+		WeightMapNames.Reserve(ClothPhysicalMeshData.WeightMaps.Num());
+
+		const UEnum* const ChaosWeightMapTargetEnum = StaticEnum<EChaosWeightMapTarget>();
+		for (const TPair<uint32, FPointWeightMap>& TargetIDAndMap : ClothPhysicalMeshData.WeightMaps)
+		{
+			WeightMapNames.Add(FName(ChaosWeightMapTargetEnum->GetNameStringByValue(TargetIDAndMap.Get<0>())));
+		}
 	}
+
 	return WeightMapNames;
 }
 
-TMap<FString, int32> FClothingSimulationMesh::GetWeightMapIndices() const
+TMap<FString, int32> FClothingSimulationMesh::GetWeightMapIndices(int32 LODIndex) const
 {
 	TMap<FString, int32> WeightMapIndices;
-
-	const UEnum* const ChaosWeightMapTargetEnum = StaticEnum<EChaosWeightMapTarget>();
-	const int32 NumWeightMaps = (int32)ChaosWeightMapTargetEnum->GetMaxEnumValue() + 1;
-
-	WeightMapIndices.Reserve(NumWeightMaps);
-
-	for (int32 EnumIndex = 0; EnumIndex < ChaosWeightMapTargetEnum->NumEnums(); ++EnumIndex)
+	const TArray<FName> WeightMapNames = GetWeightMapNames(LODIndex);
+	WeightMapIndices.Reserve(WeightMapNames.Num());
+	for (int32 WeightMapIndex = 0; WeightMapIndex < WeightMapNames.Num(); ++WeightMapIndex)
 	{
-		const int32 TargetIndex = (int32)ChaosWeightMapTargetEnum->GetValueByIndex(EnumIndex);
-		const FString WeightMapName = ChaosWeightMapTargetEnum->GetNameStringByIndex(EnumIndex);
-
-		WeightMapIndices.Emplace(WeightMapName, TargetIndex);
+		const FName& WeightMapName = WeightMapNames[WeightMapIndex];
+		WeightMapIndices.Emplace(WeightMapName.ToString(), WeightMapIndex);
 	}
 	return WeightMapIndices;
 }
@@ -190,21 +187,14 @@ TArray<TConstArrayView<FRealSingle>> FClothingSimulationMesh::GetWeightMaps(int3
 	TArray<TConstArrayView<FRealSingle>> WeightMaps;
 	if (IsValidLODIndex(LODIndex))
 	{
+		// This must match the order of GetWeightMapNames
 		const FClothLODDataCommon& ClothLODData = Asset->LodData[LODIndex];
 		const FClothPhysicalMeshData& ClothPhysicalMeshData = ClothLODData.PhysicalMeshData;
 
-		const UEnum* const ChaosWeightMapTargetEnum = StaticEnum<EChaosWeightMapTarget>();
-		const int32 NumWeightMaps = (int32)ChaosWeightMapTargetEnum->GetMaxEnumValue() + 1;
-
-		WeightMaps.SetNum(NumWeightMaps);
-
-		for (int32 EnumIndex = 0; EnumIndex < ChaosWeightMapTargetEnum->NumEnums(); ++EnumIndex)
+		WeightMaps.Reserve(ClothPhysicalMeshData.WeightMaps.Num());
+		for (const TPair<uint32, FPointWeightMap>& TargetIDAndMap : ClothPhysicalMeshData.WeightMaps)
 		{
-			const int32 TargetIndex = (int32)ChaosWeightMapTargetEnum->GetValueByIndex(EnumIndex);
-			if (const FPointWeightMap* const WeightMap = ClothPhysicalMeshData.FindWeightMap(TargetIndex))
-			{
-				WeightMaps[TargetIndex] = WeightMap->Values;
-			}
+			WeightMaps.Add(TargetIDAndMap.Get<1>().Values);
 		}
 	}
 	return WeightMaps;
diff --git a/Engine/Plugins/Experimental/ChaosCloth/Source/Chaos/Private/ChaosCloth/ChaosClothingSimulationSkeletalMesh.cpp b/Engine/Plugins/Experimental/ChaosCloth/Source/Chaos/Private/ChaosCloth/ChaosClothingSimulationSkeletalMesh.cpp
--- a/Engine/Plugins/Experimental/ChaosCloth/Source/Chaos/Private/ChaosCloth/ChaosClothingSimulationSkeletalMesh.cpp
+++ b/Engine/Plugins/Experimental/ChaosCloth/Source/Chaos/Private/ChaosCloth/ChaosClothingSimulationSkeletalMesh.cpp
@@ -108,38 +108,35 @@ PRAGMA_ENABLE_DEPRECATION_WARNINGS
 		return TConstArrayView<uint32>();
 	}
 
-	TArray<FName> FClothingSimulationSkeletalMesh::GetWeightMapNames() const
+	TArray<FName> FClothingSimulationSkeletalMesh::GetWeightMapNames(int32 LODIndex) const
 	{
 		TArray<FName> WeightMapNames;
-
-		const UEnum* const ChaosWeightMapTargetEnum = StaticEnum<EChaosWeightMapTarget>();
-		const int32 NumWeightMaps = (int32)ChaosWeightMapTargetEnum->GetMaxEnumValue() + 1;
-		
-		WeightMapNames.SetNum(NumWeightMaps);
-
-		for (int32 EnumIndex = 0; EnumIndex < ChaosWeightMapTargetEnum->NumEnums(); ++EnumIndex)
+		if (IsValidLODIndex(LODIndex))
 		{
-			const int32 TargetIndex = (int32)ChaosWeightMapTargetEnum->GetValueByIndex(EnumIndex);
-			WeightMapNames[TargetIndex] = FName(ChaosWeightMapTargetEnum->GetNameStringByIndex(EnumIndex));
+			// This must match the order of GetWeightMaps
+			const FClothLODDataCommon& ClothLODData = Asset->LodData[LODIndex];
+			const FClothPhysicalMeshData& ClothPhysicalMeshData = ClothLODData.PhysicalMeshData;
+			WeightMapNames.Reserve(ClothPhysicalMeshData.WeightMaps.Num());
+
+			const UEnum* const ChaosWeightMapTargetEnum = StaticEnum<EChaosWeightMapTarget>();
+			for (const TPair<uint32, FPointWeightMap>& TargetIDAndMap : ClothPhysicalMeshData.WeightMaps)
+			{
+				WeightMapNames.Add(FName(ChaosWeightMapTargetEnum->GetNameStringByValue(TargetIDAndMap.Get<0>())));
+			}
 		}
+
 		return WeightMapNames;
 	}
 
-	TMap<FString, int32> FClothingSimulationSkeletalMesh::GetWeightMapIndices() const
+	TMap<FString, int32> FClothingSimulationSkeletalMesh::GetWeightMapIndices(int32 LODIndex) const
 	{
 		TMap<FString, int32> WeightMapIndices;
-
-		const UEnum* const ChaosWeightMapTargetEnum = StaticEnum<EChaosWeightMapTarget>();
-		const int32 NumWeightMaps = (int32)ChaosWeightMapTargetEnum->GetMaxEnumValue() + 1;
-		
-		WeightMapIndices.Reserve(NumWeightMaps);
-
-		for (int32 EnumIndex = 0; EnumIndex < ChaosWeightMapTargetEnum->NumEnums(); ++EnumIndex)
+		const TArray<FName> WeightMapNames = GetWeightMapNames(LODIndex);
+		WeightMapIndices.Reserve(WeightMapNames.Num());
+		for (int32 WeightMapIndex = 0; WeightMapIndex < WeightMapNames.Num(); ++WeightMapIndex)
 		{
-			const int32 TargetIndex = (int32)ChaosWeightMapTargetEnum->GetValueByIndex(EnumIndex);
-			const FString WeightMapName = ChaosWeightMapTargetEnum->GetNameStringByIndex(EnumIndex);
-
-			WeightMapIndices.Emplace(WeightMapName, TargetIndex);
+			const FName& WeightMapName = WeightMapNames[WeightMapIndex];
+			WeightMapIndices.Emplace(WeightMapName.ToString(), WeightMapIndex);
 		}
 		return WeightMapIndices;
 	}
@@ -149,21 +146,14 @@ PRAGMA_ENABLE_DEPRECATION_WARNINGS
 		TArray<TConstArrayView<FRealSingle>> WeightMaps;
 		if (IsValidLODIndex(LODIndex))
 		{
+			// This must match the order of GetWeightMapNames
 			const FClothLODDataCommon& ClothLODData = Asset->LodData[LODIndex];
 			const FClothPhysicalMeshData& ClothPhysicalMeshData = ClothLODData.PhysicalMeshData;
 
-			const UEnum* const ChaosWeightMapTargetEnum = StaticEnum<EChaosWeightMapTarget>();
-			const int32 NumWeightMaps = (int32)ChaosWeightMapTargetEnum->GetMaxEnumValue() + 1;
-
-			WeightMaps.SetNum(NumWeightMaps);
-
-			for (int32 EnumIndex = 0; EnumIndex < ChaosWeightMapTargetEnum->NumEnums(); ++EnumIndex)
+			WeightMaps.Reserve(ClothPhysicalMeshData.WeightMaps.Num());
+			for (const TPair<uint32, FPointWeightMap>& TargetIDAndMap : ClothPhysicalMeshData.WeightMaps)
 			{
-				const int32 TargetIndex = (int32)ChaosWeightMapTargetEnum->GetValueByIndex(EnumIndex);
-				if (const FPointWeightMap* const WeightMap = ClothPhysicalMeshData.FindWeightMap(TargetIndex))
-				{
-					WeightMaps[TargetIndex] = WeightMap->Values;
-				}
+				WeightMaps.Add(TargetIDAndMap.Get<1>().Values);
 			}
 		}
 		return WeightMaps;
diff --git a/Engine/Plugins/Experimental/ChaosCloth/Source/Chaos/Public/ChaosCloth/ChaosClothingSimulationMesh.h b/Engine/Plugins/Experimental/ChaosCloth/Source/Chaos/Public/ChaosCloth/ChaosClothingSimulationMesh.h
--- a/Engine/Plugins/Experimental/ChaosCloth/Source/Chaos/Public/ChaosCloth/ChaosClothingSimulationMesh.h
+++ b/Engine/Plugins/Experimental/ChaosCloth/Source/Chaos/Public/ChaosCloth/ChaosClothingSimulationMesh.h
@@ -65,8 +65,12 @@ namespace Chaos
 		virtual TConstArrayView<uint32> GetIndices(int32 LODIndex) const;
 		virtual TConstArrayView<uint32> GetPatternIndices(int32 LODIndex) const;
 		virtual TConstArrayView<uint32> GetPatternToWeldedIndices(int32 LODIndex) const;
-		virtual TArray<FName> GetWeightMapNames() const;
-		virtual TMap<FString, int32> GetWeightMapIndices() const;
+		virtual TArray<FName> GetWeightMapNames(int32 LODIndex) const;
+		UE_DEPRECATED(5.3, "Use LODIndex version.")
+		virtual TArray<FName> GetWeightMapNames() const { return GetWeightMapNames(0); }
+		virtual TMap<FString, int32> GetWeightMapIndices(int32 LODIndex) const;
+		UE_DEPRECATED(5.3, "Use LODIndex version.")
+		virtual TMap<FString, int32> GetWeightMapIndices() const { return GetWeightMapIndices(0); }
 		virtual TArray<TConstArrayView<FRealSingle>> GetWeightMaps(int32 LODIndex) const;
 		virtual TArray<TConstArrayView<TTuple<int32, int32, float>>> GetTethers(int32 LODIndex, bool bUseGeodesicTethers) const;
 		virtual int32 GetReferenceBoneIndex() const;
@@ -116,9 +120,13 @@ namespace Chaos
 		virtual TConstArrayView<uint32> GetPatternToWeldedIndices(int32 LODIndex) const = 0;
 
 		/* Return all weight maps associated with this mesh returned in the same order as GetWeightMaps. */
+		virtual TArray<FName> GetWeightMapNames(int32 LODIndex) const = 0;
+		UE_DEPRECATED(5.3, "Use LODIndex version.")
 		virtual TArray<FName> GetWeightMapNames() const = 0;
 
 		/* Return a map of all weight map names associated with this mesh to the index in the array returned by GetWeightMaps. */
+		virtual TMap<FString, int32> GetWeightMapIndices(int32 LODIndex) const = 0;
+		UE_DEPRECATED(5.3, "Use LODIndex version.")
 		virtual TMap<FString, int32> GetWeightMapIndices() const = 0;
 
 		/* Return the specified LOD's weight map. */
diff --git a/Engine/Plugins/Experimental/ChaosCloth/Source/Chaos/Public/ChaosCloth/ChaosClothingSimulationSkeletalMesh.h b/Engine/Plugins/Experimental/ChaosCloth/Source/Chaos/Public/ChaosCloth/ChaosClothingSimulationSkeletalMesh.h
--- a/Engine/Plugins/Experimental/ChaosCloth/Source/Chaos/Public/ChaosCloth/ChaosClothingSimulationSkeletalMesh.h
+++ b/Engine/Plugins/Experimental/ChaosCloth/Source/Chaos/Public/ChaosCloth/ChaosClothingSimulationSkeletalMesh.h
@@ -45,8 +45,12 @@ PRAGMA_ENABLE_DEPRECATION_WARNINGS
 		virtual TConstArrayView<uint32> GetIndices(int32 LODIndex) const override;
 		virtual TConstArrayView<uint32> GetPatternIndices(int32 LODIndex) const override;
 		virtual TConstArrayView<uint32> GetPatternToWeldedIndices(int32 LODIndex) const override;
-		virtual TArray<FName> GetWeightMapNames() const override;
-		virtual TMap<FString, int32> GetWeightMapIndices() const override;
+		virtual TArray<FName> GetWeightMapNames(int32 LODIndex) const override;
+		UE_DEPRECATED(5.3, "Use LODIndex version.")
+		virtual TArray<FName> GetWeightMapNames() const override { return GetWeightMapNames(0); }
+		virtual TMap<FString, int32> GetWeightMapIndices(int32 LODIndex) const override;
+		UE_DEPRECATED(5.3, "Use LODIndex version.")
+		virtual TMap<FString, int32> GetWeightMapIndices() const override { return GetWeightMapIndices(0); }
 		virtual TArray<TConstArrayView<FRealSingle>> GetWeightMaps(int32 LODIndex) const override;
 		virtual TArray<TConstArrayView<TTuple<int32, int32, float>>> GetTethers(int32 LODIndex, bool bUseGeodesicTethers) const override;
 		virtual int32 GetReferenceBoneIndex() const override;
diff --git a/Engine/Plugins/Experimental/ChaosClothAsset/Source/ChaosClothAssetEngine/Private/ChaosClothAsset/ClothSimulationMesh.cpp b/Engine/Plugins/Experimental/ChaosClothAsset/Source/ChaosClothAssetEngine/Private/ChaosClothAsset/ClothSimulationMesh.cpp
--- a/Engine/Plugins/Experimental/ChaosClothAsset/Source/ChaosClothAssetEngine/Private/ChaosClothAsset/ClothSimulationMesh.cpp
+++ b/Engine/Plugins/Experimental/ChaosClothAsset/Source/ChaosClothAssetEngine/Private/ChaosClothAsset/ClothSimulationMesh.cpp
@@ -77,22 +77,22 @@ PRAGMA_ENABLE_DEPRECATION_WARNINGS
 		return TConstArrayView<uint32>(ClothSimulationModel.GetPatternToWeldedIndices(LODIndex));
 	}
 
-	TArray<FName> FClothSimulationMesh::GetWeightMapNames() const
+	TArray<FName> FClothSimulationMesh::GetWeightMapNames(int32 LODIndex) const
 	{
 		TArray<FName> WeightMapNames;
-		if (ClothSimulationModel.ClothSimulationLodModels.Num())
+		if (ClothSimulationModel.ClothSimulationLodModels.IsValidIndex(LODIndex))
 		{
-			ClothSimulationModel.ClothSimulationLodModels[0].WeightMaps.GetKeys(WeightMapNames);
+			ClothSimulationModel.ClothSimulationLodModels[LODIndex].WeightMaps.GetKeys(WeightMapNames);
 		}
 		return WeightMapNames;
 	}
 
-	TMap<FString, int32> FClothSimulationMesh::GetWeightMapIndices() const
+	TMap<FString, int32> FClothSimulationMesh::GetWeightMapIndices(int32 LODIndex) const
 	{
 		TMap<FString, int32> WeightMapIndices;
 
 		// Retrieve weight map names for this cloth
-		const TArray<FName> WeightMapNames = GetWeightMapNames();
+		const TArray<FName> WeightMapNames = GetWeightMapNames(LODIndex);
 		WeightMapIndices.Reserve(WeightMapNames.Num());
 
 		for (int32 WeightMapIndex = 0; WeightMapIndex < WeightMapNames.Num(); ++WeightMapIndex)
@@ -108,7 +108,7 @@ PRAGMA_ENABLE_DEPRECATION_WARNINGS
 		TArray<TConstArrayView<::Chaos::FRealSingle>> WeightMaps;
 
 		// Retrieve weight map names for this cloth
-		const TArray<FName> WeightMapNames = GetWeightMapNames();
+		const TArray<FName> WeightMapNames = GetWeightMapNames(LODIndex);
 		WeightMaps.Reserve(WeightMapNames.Num());
 
 		for (int32 WeightMapIndex = 0; WeightMapIndex < WeightMapNames.Num(); ++WeightMapIndex)
diff --git a/Engine/Plugins/Experimental/ChaosClothAsset/Source/ChaosClothAssetEngine/Private/ChaosClothAsset/ClothSimulationMesh.h b/Engine/Plugins/Experimental/ChaosClothAsset/Source/ChaosClothAssetEngine/Private/ChaosClothAsset/ClothSimulationMesh.h
--- a/Engine/Plugins/Experimental/ChaosClothAsset/Source/ChaosClothAssetEngine/Private/ChaosClothAsset/ClothSimulationMesh.h
+++ b/Engine/Plugins/Experimental/ChaosClothAsset/Source/ChaosClothAssetEngine/Private/ChaosClothAsset/ClothSimulationMesh.h
@@ -38,8 +38,12 @@ PRAGMA_ENABLE_DEPRECATION_WARNINGS
 		virtual TConstArrayView<uint32> GetIndices(int32 LODIndex) const override;
 		virtual TConstArrayView<uint32> GetPatternIndices(int32 LODIndex) const override;
 		virtual TConstArrayView<uint32> GetPatternToWeldedIndices(int32 LODIndex) const override;
-		virtual TArray<FName> GetWeightMapNames() const override;
-		virtual TMap<FString, int32> GetWeightMapIndices() const override;
+		virtual TArray<FName> GetWeightMapNames(int32 LODIndex) const override;
+		UE_DEPRECATED(5.3, "Use LODIndex version.")
+		virtual TArray<FName> GetWeightMapNames() const override { return GetWeightMapNames(0); }
+		virtual TMap<FString, int32> GetWeightMapIndices(int32 LODIndex) const override;
+		UE_DEPRECATED(5.3, "Use LODIndex version.")
+		virtual TMap<FString, int32> GetWeightMapIndices() const override { return GetWeightMapIndices(0); }
 		virtual TArray<TConstArrayView<::Chaos::FRealSingle>> GetWeightMaps(int32 LODIndex) const override;
 		// Note: there is only one set of tethers stored on ClothSimulationMesh assets
 		virtual TArray<TConstArrayView<TTuple<int32, int32, float>>> GetTethers(int32 LODIndex, bool /*bUseGeodesicTethers*/) const override;
