Chaos Cloth: allow adjusting  MinPerParticleMass in the SimulationMassConfig node.

There's a bit of a hack to allow clamping to 1e-8 because the editor will clip values less than ~1e-6 to 0 (including if you set ClampMin = 1e-8) by doing it in code when we read the value.

diff --git a/Engine/Plugins/Experimental/ChaosCloth/Source/Chaos/Private/ChaosCloth/ChaosClothingSimulationCloth.cpp b/Engine/Plugins/Experimental/ChaosCloth/Source/Chaos/Private/ChaosCloth/ChaosClothingSimulationCloth.cpp
--- a/Engine/Plugins/Experimental/ChaosCloth/Source/Chaos/Private/ChaosCloth/ChaosClothingSimulationCloth.cpp
+++ b/Engine/Plugins/Experimental/ChaosCloth/Source/Chaos/Private/ChaosCloth/ChaosClothingSimulationCloth.cpp
@@ -187,7 +187,9 @@ PRAGMA_ENABLE_DEPRECATION_WARNINGS
 
 	const int32 MassMode = ConfigProperties.GetValue<int32>(TEXT("MassMode"), ClothingSimulationClothDefault::MassMode);
 	const FRealSingle MassValue = (FRealSingle)ConfigProperties.GetValue<float>(TEXT("MassValue"), ClothingSimulationClothDefault::MassValue);
-	const FRealSingle MinPerParticleMass = (FRealSingle)ConfigProperties.GetValue<float>(TEXT("MinPerParticleMass"), ClothingSimulationClothDefault::MinPerParticleMass);
+
+	constexpr FRealSingle MinPerParticleMassClampMin = UE_SMALL_NUMBER;
+	const FRealSingle MinPerParticleMass = FMath::Max((FRealSingle)ConfigProperties.GetValue<float>(TEXT("MinPerParticleMass"), ClothingSimulationClothDefault::MinPerParticleMass), MinPerParticleMassClampMin);
 
 	switch (MassMode)
 	{
diff --git a/Engine/Plugins/Experimental/ChaosClothAssetEditor/Source/ChaosClothAssetDataflowNodes/Private/ChaosClothAsset/SimulationMassConfigNode.cpp b/Engine/Plugins/Experimental/ChaosClothAssetEditor/Source/ChaosClothAssetDataflowNodes/Private/ChaosClothAsset/SimulationMassConfigNode.cpp
--- a/Engine/Plugins/Experimental/ChaosClothAssetEditor/Source/ChaosClothAssetDataflowNodes/Private/ChaosClothAsset/SimulationMassConfigNode.cpp
+++ b/Engine/Plugins/Experimental/ChaosClothAssetEditor/Source/ChaosClothAssetDataflowNodes/Private/ChaosClothAsset/SimulationMassConfigNode.cpp
@@ -25,4 +25,5 @@ void FChaosClothAssetSimulationMassConfigNode::AddProperties(Dataflow::FContext&
 
 	UE_CHAOS_CLOTHASSET_SIMULATIONCONFIG_SETPROPERTYENUM(MassMode);
 	UE_CHAOS_CLOTHASSET_SIMULATIONCONFIG_SETPROPERTY(MassValue);
+	UE_CHAOS_CLOTHASSET_SIMULATIONCONFIG_SETPROPERTY(MinPerParticleMass);
 }
diff --git a/Engine/Plugins/Experimental/ChaosClothAssetEditor/Source/ChaosClothAssetDataflowNodes/Public/ChaosClothAsset/SimulationMassConfigNode.h b/Engine/Plugins/Experimental/ChaosClothAssetEditor/Source/ChaosClothAssetDataflowNodes/Public/ChaosClothAsset/SimulationMassConfigNode.h
--- a/Engine/Plugins/Experimental/ChaosClothAssetEditor/Source/ChaosClothAssetDataflowNodes/Public/ChaosClothAsset/SimulationMassConfigNode.h
+++ b/Engine/Plugins/Experimental/ChaosClothAssetEditor/Source/ChaosClothAssetDataflowNodes/Public/ChaosClothAsset/SimulationMassConfigNode.h
@@ -44,6 +44,10 @@ public:
 	UPROPERTY(EditAnywhere, Category = "Mass Properties", meta = (UIMin = "0.001", UIMax = "1", ClampMin = "0", EditCondition = "MassMode == EClothMassMode::Density"))
 	float Density = 0.35f;
 
+	/** Calculated particle masses will be clamped to this minimum value (or 1e-8, whichever is larger). */
+	UPROPERTY(EditAnywhere, Category = "Mass Properties", meta = (ClampMin = "0"))
+	float MinPerParticleMass = 0.0001f;
+
 	FChaosClothAssetSimulationMassConfigNode(const Dataflow::FNodeParameters& InParam, FGuid InGuid = FGuid::NewGuid());
 
 private:
