Fix indexing when debug drawing cloth bending with multiple LODs.

diff --git a/Engine/Plugins/Experimental/ChaosCloth/Source/Chaos/Private/ChaosCloth/ChaosClothVisualization.cpp b/Engine/Plugins/Experimental/ChaosCloth/Source/Chaos/Private/ChaosCloth/ChaosClothVisualization.cpp
--- a/Engine/Plugins/Experimental/ChaosCloth/Source/Chaos/Private/ChaosCloth/ChaosClothVisualization.cpp
+++ b/Engine/Plugins/Experimental/ChaosCloth/Source/Chaos/Private/ChaosCloth/ChaosClothVisualization.cpp
@@ -1600,7 +1600,7 @@ static FAutoConsoleVariableRef CVarClothVizWeightMapName(TEXT("p.ChaosClothVisua
 			// Draw constraints
 			const FClothConstraints& ClothConstraints = Solver->GetClothConstraints(Offset);
 
-			const TConstArrayView<Softs::FSolverVec3> Positions = Cloth->GetParticlePositions(Solver);
+			const TArray<Softs::FSolverVec3>& Positions = Solver->GetParticleXs();
 
 			if (const Softs::FPBDEdgeSpringConstraints* const EdgeConstraints = ClothConstraints.GetEdgeSpringConstraints().Get())
 			{
@@ -1709,7 +1709,7 @@ static FAutoConsoleVariableRef CVarClothVizWeightMapName(TEXT("p.ChaosClothVisua
 			// Draw constraints
 			const FClothConstraints& ClothConstraints = Solver->GetClothConstraints(Offset);
 
-			const TConstArrayView<Softs::FSolverVec3> Positions = Cloth->GetParticlePositions(Solver);
+			const TArray<Softs::FSolverVec3>& Positions = Solver->GetParticleXs();
 
 			if (const Softs::FPBDBendingSpringConstraints* const BendingConstraints = ClothConstraints.GetBendingSpringConstraints().Get())
 			{
@@ -1885,21 +1885,21 @@ static FAutoConsoleVariableRef CVarClothVizWeightMapName(TEXT("p.ChaosClothVisua
 
 			const TConstArrayView<TVec3<int32>>& Elements = VelocityField.GetElements();
 			const TConstArrayView<Softs::FSolverVec3> Forces = VelocityField.GetForces();
-			const TConstArrayView<Softs::FSolverVec3> Positions = Cloth->GetParticlePositions(Solver);
-			const TConstArrayView<Softs::FSolverReal> InvMasses = Cloth->GetParticleInvMasses(Solver);
+			const TArray<Softs::FSolverVec3>& Positions = Solver->GetParticleXs();
+			const TArray<Softs::FSolverReal>& InvMasses = Solver->GetParticleInvMasses();
 			check(InvMasses.Num() == Positions.Num());
 
 			for (int32 ElementIndex = 0; ElementIndex < Elements.Num(); ++ElementIndex)
 			{
 				const TVec3<int32>& Element = Elements[ElementIndex];
 				const FVec3 Position = LocalSpaceLocation + (
-					FVec3(Positions[Element.X - Offset]) +
-					FVec3(Positions[Element.Y - Offset]) +
-					FVec3(Positions[Element.Z - Offset])) / (FReal)3.;
+					FVec3(Positions[Element.X]) +
+					FVec3(Positions[Element.Y]) +
+					FVec3(Positions[Element.Z])) / (FReal)3.;
 
-				const bool bIsKinematic0 = !InvMasses[Element.X - Offset];
-				const bool bIsKinematic1 = !InvMasses[Element.Y - Offset];
-				const bool bIsKinematic2 = !InvMasses[Element.Z - Offset];
+				const bool bIsKinematic0 = !InvMasses[Element.X];
+				const bool bIsKinematic1 = !InvMasses[Element.Y];
+				const bool bIsKinematic2 = !InvMasses[Element.Z];
 				const bool bIsKinematic = bIsKinematic0 || bIsKinematic1 || bIsKinematic2;
 
 				const FVec3 Force = FVec3(Forces[ElementIndex]) * ForceLength;
@@ -1952,7 +1952,7 @@ static FAutoConsoleVariableRef CVarClothVizWeightMapName(TEXT("p.ChaosClothVisua
 
 			if (const Softs::FPBDCollisionSpringConstraints* const SelfCollisionConstraints = ClothConstraints.GetSelfCollisionConstraints().Get())
 			{
-				const TConstArrayView<Softs::FSolverVec3> Positions = Cloth->GetParticlePositions(Solver);
+				const TArray<Softs::FSolverVec3>& Positions = Solver->GetParticleXs();
 				const TArray<TVec4<int32>>& Constraints = SelfCollisionConstraints->GetConstraints();
 				const TArray<Softs::FSolverVec3>& Barys = SelfCollisionConstraints->GetBarys();
 				const FReal Thickness = (FReal)SelfCollisionConstraints->GetThickness();
@@ -1964,11 +1964,10 @@ static FAutoConsoleVariableRef CVarClothVizWeightMapName(TEXT("p.ChaosClothVisua
 					const TVec4<int32>& Constraint = Constraints[Index];
 					const FVec3 Bary(Barys[Index]);
 
-					// Constraint index includes Offset, but so does Positions.
-					const FVector P = LocalSpaceLocation + FVector(Positions[Constraint[0] - Offset]);
-					const FVector P0 = LocalSpaceLocation + FVector(Positions[Constraint[1] - Offset]);
-					const FVector P1 = LocalSpaceLocation + FVector(Positions[Constraint[2] - Offset]);
-					const FVector P2 = LocalSpaceLocation + FVector(Positions[Constraint[3] - Offset]);
+					const FVector P = LocalSpaceLocation + FVector(Positions[Constraint[0]]);
+					const FVector P0 = LocalSpaceLocation + FVector(Positions[Constraint[1]]);
+					const FVector P1 = LocalSpaceLocation + FVector(Positions[Constraint[2]]);
+					const FVector P2 = LocalSpaceLocation + FVector(Positions[Constraint[3]]);
 
 					const FVector Pos0 = P0 * Bary[0] + P1 * Bary[1] + P2 * Bary[2];
 
@@ -2018,7 +2017,7 @@ static FAutoConsoleVariableRef CVarClothVizWeightMapName(TEXT("p.ChaosClothVisua
 
 			if (const Softs::FPBDTriangleMeshCollisions* const SelfCollisionInit = ClothConstraints.GetSelfCollisionInit().Get())
 			{
-				const TConstArrayView<Softs::FSolverVec3> Positions = Cloth->GetParticlePositions(Solver);
+				const TArray<Softs::FSolverVec3>& Positions = Solver->GetParticleXs();
 				const FTriangleMesh& TriangleMesh = Cloth->GetTriangleMesh(Solver);
 
 				// Draw contours
@@ -2040,9 +2039,9 @@ static FAutoConsoleVariableRef CVarClothVizWeightMapName(TEXT("p.ChaosClothVisua
 					for (int32 PointIdx = 0; PointIdx < Contour.Num() - 1; ++PointIdx)
 					{
 						const Softs::FPBDTriangleMeshCollisions::FBarycentricPoint& Point0 = Contour[PointIdx];
-						const FVector EndPoint0 = LocalSpaceLocation + (1.f - Point0.Bary[0] - Point0.Bary[1]) * Positions[Point0.Vertices[0] - Offset] + Point0.Bary[0] * Positions[Point0.Vertices[1] - Offset] + Point0.Bary[1] * Positions[Point0.Vertices[2] - Offset];
+						const FVector EndPoint0 = LocalSpaceLocation + (1.f - Point0.Bary[0] - Point0.Bary[1]) * Positions[Point0.Vertices[0]] + Point0.Bary[0] * Positions[Point0.Vertices[1]] + Point0.Bary[1] * Positions[Point0.Vertices[2]];
 						const Softs::FPBDTriangleMeshCollisions::FBarycentricPoint& Point1 = Contour[PointIdx+1];
-						const FVector EndPoint1 = LocalSpaceLocation + (1.f - Point1.Bary[0] - Point1.Bary[1]) * Positions[Point1.Vertices[0] - Offset] + Point1.Bary[0] * Positions[Point1.Vertices[1] - Offset] + Point1.Bary[1] * Positions[Point1.Vertices[2] - Offset];
+						const FVector EndPoint1 = LocalSpaceLocation + (1.f - Point1.Bary[0] - Point1.Bary[1]) * Positions[Point1.Vertices[0]] + Point1.Bary[0] * Positions[Point1.Vertices[1]] + Point1.Bary[1] * Positions[Point1.Vertices[2]];
 						DrawLine(PDI, EndPoint0, EndPoint1, ContourColor);
 						DrawPoint(PDI, EndPoint0, ContourColor, nullptr, 1.f);
 						DrawPoint(PDI, EndPoint1, ContourColor, nullptr, 1.f);
@@ -2064,7 +2063,7 @@ static FAutoConsoleVariableRef CVarClothVizWeightMapName(TEXT("p.ChaosClothVisua
 							const bool bAnyBlack = (VertexGIAColors[ParticleIdx].ContourIndexBits & VertexGIAColors[ParticleIdx].ColorBits);
 							const FLinearColor& VertColor = bIsLoop ? Red : (bAnyWhite && bAnyBlack) ? Gray : bAnyWhite ? White : Black;
 						
-							DrawPoint(PDI, LocalSpaceLocation + Positions[ParticleIdx - Offset], VertColor, nullptr, 5.f);
+							DrawPoint(PDI, LocalSpaceLocation + Positions[ParticleIdx], VertColor, nullptr, 5.f);
 						}
 					}
 				}
@@ -2080,9 +2079,9 @@ static FAutoConsoleVariableRef CVarClothVizWeightMapName(TEXT("p.ChaosClothVisua
 							const bool bAnyWhite = (TriangleGIAColors[TriangleIdx].ContourIndexBits & ~TriangleGIAColors[TriangleIdx].ColorBits);
 							const bool bAnyBlack = (TriangleGIAColors[TriangleIdx].ContourIndexBits & TriangleGIAColors[TriangleIdx].ColorBits);
 							const FLinearColor& TriColor = bIsLoop ? Red : (bAnyWhite && bAnyBlack) ? Gray : bAnyWhite ? White : Black;
-							DrawLine(PDI, LocalSpaceLocation + Positions[Elements[TriangleIdx][0] - Offset], LocalSpaceLocation + Positions[Elements[TriangleIdx][1] - Offset], TriColor);
-							DrawLine(PDI, LocalSpaceLocation + Positions[Elements[TriangleIdx][1] - Offset], LocalSpaceLocation + Positions[Elements[TriangleIdx][2] - Offset], TriColor);
-							DrawLine(PDI, LocalSpaceLocation + Positions[Elements[TriangleIdx][0] - Offset], LocalSpaceLocation + Positions[Elements[TriangleIdx][2] - Offset], TriColor);
+							DrawLine(PDI, LocalSpaceLocation + Positions[Elements[TriangleIdx][0]], LocalSpaceLocation + Positions[Elements[TriangleIdx][1]], TriColor);
+							DrawLine(PDI, LocalSpaceLocation + Positions[Elements[TriangleIdx][1]], LocalSpaceLocation + Positions[Elements[TriangleIdx][2]], TriColor);
+							DrawLine(PDI, LocalSpaceLocation + Positions[Elements[TriangleIdx][0]], LocalSpaceLocation + Positions[Elements[TriangleIdx][2]], TriColor);
 						}
 					}
 				}
@@ -2098,8 +2097,8 @@ static FAutoConsoleVariableRef CVarClothVizWeightMapName(TEXT("p.ChaosClothVisua
 					Intersection.GlobalGradientVector.ToDirectionAndLength(GradientDir, GradientLength);
 					const FVector Delta = FVector(GradientDir) * MaxDrawImpulse * GradientLength * FMath::InvSqrt(GradientLength * GradientLength + RegularizeEpsilonSq);
 
-					const FVector EdgeCenter = LocalSpaceLocation + .5 * (Positions[Intersection.EdgeVertices[0] - Offset] + Positions[Intersection.EdgeVertices[1] - Offset]);
-					const FVector TriCenter = LocalSpaceLocation + (Positions[Intersection.FaceVertices[0] - Offset] + Positions[Intersection.FaceVertices[1] - Offset] + Positions[Intersection.FaceVertices[2] - Offset]) / 3.;
+					const FVector EdgeCenter = LocalSpaceLocation + .5 * (Positions[Intersection.EdgeVertices[0]] + Positions[Intersection.EdgeVertices[1]]);
+					const FVector TriCenter = LocalSpaceLocation + (Positions[Intersection.FaceVertices[0]] + Positions[Intersection.FaceVertices[1]] + Positions[Intersection.FaceVertices[2]]) / 3.;
 
 					DrawPoint(PDI, EdgeCenter, Green, nullptr, 2.f);
 					DrawLine(PDI, EdgeCenter, EdgeCenter + Delta, Green);
diff --git a/Engine/Plugins/Experimental/ChaosCloth/Source/Chaos/Private/ChaosCloth/ChaosClothingSimulationSolver.cpp b/Engine/Plugins/Experimental/ChaosCloth/Source/Chaos/Private/ChaosCloth/ChaosClothingSimulationSolver.cpp
--- a/Engine/Plugins/Experimental/ChaosCloth/Source/Chaos/Private/ChaosCloth/ChaosClothingSimulationSolver.cpp
+++ b/Engine/Plugins/Experimental/ChaosCloth/Source/Chaos/Private/ChaosCloth/ChaosClothingSimulationSolver.cpp
@@ -417,6 +417,11 @@ void FClothingSimulationSolver::ResetStartPose(int32 Offset, int32 NumParticles)
 	}
 }
 
+const TArray<Softs::FPAndInvM>& FClothingSimulationSolver::GetParticlePandInvMs() const
+{
+	return Evolution->Particles().GetPAndInvM();
+}
+
 const Softs::FPAndInvM* FClothingSimulationSolver::GetParticlePandInvMs(int32 Offset) const
 {
 	return &Evolution->Particles().PAndInvM(Offset);
@@ -427,6 +432,11 @@ Softs::FPAndInvM* FClothingSimulationSolver::GetParticlePandInvMs(int32 Offset)
 	return &Evolution->Particles().PAndInvM(Offset);
 }
 
+const TArray<Softs::FSolverVec3>& FClothingSimulationSolver::GetParticleXs() const
+{
+	return Evolution->Particles().XArray();
+}
+
 const Softs::FSolverVec3* FClothingSimulationSolver::GetParticleXs(int32 Offset) const
 {
 	return &Evolution->Particles().X(Offset);
@@ -437,6 +447,11 @@ Softs::FSolverVec3* FClothingSimulationSolver::GetParticleXs(int32 Offset)
 	return &Evolution->Particles().X(Offset);
 }
 
+const TArray<Softs::FSolverVec3>& FClothingSimulationSolver::GetParticleVs() const
+{
+	return Evolution->Particles().GetV();
+}
+
 const Softs::FSolverVec3* FClothingSimulationSolver::GetParticleVs(int32 Offset) const
 {
 	return &Evolution->Particles().V(Offset);
@@ -447,6 +462,11 @@ Softs::FSolverVec3* FClothingSimulationSolver::GetParticleVs(int32 Offset)
 	return &Evolution->Particles().V(Offset);
 }
 
+const TArray<Softs::FSolverReal>& FClothingSimulationSolver::GetParticleInvMasses() const
+{
+	return Evolution->Particles().GetInvM();
+}
+
 const Softs::FSolverReal* FClothingSimulationSolver::GetParticleInvMasses(int32 Offset) const
 {
 	return &Evolution->Particles().InvM(Offset);
diff --git a/Engine/Plugins/Experimental/ChaosCloth/Source/Chaos/Public/ChaosCloth/ChaosClothingSimulationSolver.h b/Engine/Plugins/Experimental/ChaosCloth/Source/Chaos/Public/ChaosCloth/ChaosClothingSimulationSolver.h
--- a/Engine/Plugins/Experimental/ChaosCloth/Source/Chaos/Public/ChaosCloth/ChaosClothingSimulationSolver.h
+++ b/Engine/Plugins/Experimental/ChaosCloth/Source/Chaos/Public/ChaosCloth/ChaosClothingSimulationSolver.h
@@ -181,28 +181,41 @@ namespace Chaos
 		// Add external forces to the particles
 		CHAOSCLOTH_API void AddExternalForces(uint32 GroupId, bool bUseLegacyWind);
 
+		const TArray<Softs::FSolverVec3>& GetOldAnimationPositions() const { return OldAnimationPositions; }
 		const Softs::FSolverVec3* GetOldAnimationPositions(int32 Offset) const { return OldAnimationPositions.GetData() + Offset; }
 		Softs::FSolverVec3* GetOldAnimationPositions(int32 Offset) { return OldAnimationPositions.GetData() + Offset; }
+		const TArray<Softs::FSolverVec3>& GetAnimationPositions() const { return AnimationPositions; }
 		const Softs::FSolverVec3* GetAnimationPositions(int32 Offset) const { return AnimationPositions.GetData() + Offset; }
 		Softs::FSolverVec3* GetAnimationPositions(int32 Offset) { return AnimationPositions.GetData() + Offset; }
+		const TArray<Softs::FSolverVec3>& GetInterpolatedAnimationPositions() const { return InterpolatedAnimationPositions; }
 		const Softs::FSolverVec3* GetInterpolatedAnimationPositions(int32 Offset) const { return InterpolatedAnimationPositions.GetData() + Offset; }
 		Softs::FSolverVec3* GetInterpolatedAnimationPositions(int32 Offset) { return InterpolatedAnimationPositions.GetData() + Offset; }
+		const TArray<Softs::FSolverVec3>& GetOldAnimationNormals() const { return OldAnimationNormals; }
 		const Softs::FSolverVec3* GetOldAnimationNormals(int32 Offset) const { return OldAnimationNormals.GetData() + Offset; }
 		Softs::FSolverVec3* GetOldAnimationNormals(int32 Offset) { return OldAnimationNormals.GetData() + Offset; }
+		const TArray<Softs::FSolverVec3>& GetAnimationNormals() const { return AnimationNormals; }
 		const Softs::FSolverVec3* GetAnimationNormals(int32 Offset) const { return AnimationNormals.GetData() + Offset; }
 		Softs::FSolverVec3* GetAnimationNormals(int32 Offset) { return AnimationNormals.GetData() + Offset; }
+		const TArray<Softs::FSolverVec3>& GetInterpolatedAnimationNormals() const { return InterpolatedAnimationNormals; }
 		const Softs::FSolverVec3* GetInterpolatedAnimationNormals(int32 Offset) const { return InterpolatedAnimationNormals.GetData() + Offset; }
 		Softs::FSolverVec3* GetInterpolatedAnimationNormals(int32 Offset) { return InterpolatedAnimationNormals.GetData() + Offset; }
+		const TArray<Softs::FSolverVec3>& GetNormals() const { return Normals; }
 		const Softs::FSolverVec3* GetNormals(int32 Offset) const { return Normals.GetData() + Offset; }
+		Softs::FSolverVec3* GetNormals(int32 Offset) { return Normals.GetData() + Offset; }
+		const TArray<FSolverVec3>& GetAnimationVelocities() const { return AnimationVelocities; }
 		const Softs::FSolverVec3* GetAnimationVelocities(int32 Offset) const { return AnimationVelocities.GetData() + Offset; }
 		Softs::FSolverVec3* GetAnimationVelocities(int32 Offset) { return AnimationVelocities.GetData() + Offset; }
-		Softs::FSolverVec3* GetNormals(int32 Offset) { return Normals.GetData() + Offset; }
+
+		CHAOSCLOTH_API const TArray<Softs::FPAndInvM>& GetParticlePandInvMs() const;
 		CHAOSCLOTH_API const Softs::FPAndInvM* GetParticlePandInvMs(int32 Offset) const;
-		CHAOSCLOTH_API Softs::FPAndInvM* GetParticlePandInvMs(int32 Offset);
+		CHAOSCLOTH_API Softs::FPAndInvM* GetParticlePandInvMs(int32 Offset);		
+		CHAOSCLOTH_API const TArray<Softs::FSolverVec3>& GetParticleXs() const;
 		CHAOSCLOTH_API const Softs::FSolverVec3* GetParticleXs(int32 Offset) const;
 		CHAOSCLOTH_API Softs::FSolverVec3* GetParticleXs(int32 Offset);
+		CHAOSCLOTH_API const TArray<Softs::FSolverVec3>& GetParticleVs() const;
 		CHAOSCLOTH_API const Softs::FSolverVec3* GetParticleVs(int32 Offset) const;
 		CHAOSCLOTH_API Softs::FSolverVec3* GetParticleVs(int32 Offset);
+		CHAOSCLOTH_API const TArray<Softs::FSolverReal>& GetParticleInvMasses() const;
 		CHAOSCLOTH_API const Softs::FSolverReal* GetParticleInvMasses(int32 Offset) const;
 		const FClothConstraints& GetClothConstraints(int32 Offset) const { return *ClothsConstraints.FindChecked(Offset); }
 		FClothConstraints& GetClothConstraints(int32 Offset) { return *ClothsConstraints.FindChecked(Offset); }
