Chaos Cloth - Added a global gravity multiplier console variable p.ChaosCloth.GravityMultiplier (default to 1).

diff --git a/Engine/Plugins/Experimental/ChaosCloth/Source/Chaos/Private/ChaosCloth/ChaosClothingSimulationCloth.cpp b/Engine/Plugins/Experimental/ChaosCloth/Source/Chaos/Private/ChaosCloth/ChaosClothingSimulationCloth.cpp
--- a/Engine/Plugins/Experimental/ChaosCloth/Source/Chaos/Private/ChaosCloth/ChaosClothingSimulationCloth.cpp
+++ b/Engine/Plugins/Experimental/ChaosCloth/Source/Chaos/Private/ChaosCloth/ChaosClothingSimulationCloth.cpp
@@ -38,7 +38,15 @@ namespace ClothingSimulationClothDefault
 
 namespace ClothingSimulationClothConsoleVariables
 {
-	TAutoConsoleVariable<bool> CVarLegacyDisablesAccurateWind(TEXT("p.ChaosCloth.LegacyDisablesAccurateWind"), true, TEXT("Whether using the Legacy wind model switches off the accurate wind model, or adds up to it"));
+	TAutoConsoleVariable<bool> CVarLegacyDisablesAccurateWind(
+		TEXT("p.ChaosCloth.LegacyDisablesAccurateWind"),
+		true,
+		TEXT("Whether using the Legacy wind model switches off the accurate wind model, or adds up to it"));
+
+	TAutoConsoleVariable<float> CVarGravityMultiplier(
+		TEXT("p.ChaosCloth.GravityMultiplier"),
+		1.f,
+		TEXT("Scalar multiplier applied at the final stage of the cloth's gravity formulation."));
 }
 
 struct FClothingSimulationCloth::FLODData
@@ -770,8 +778,9 @@ TVec3<FRealSingle> FClothingSimulationCloth::GetGravity(const FClothingSimulatio
 	const bool bUseGravityOverride = ConfigProperties.GetValue<bool>(TEXT("UseGravityOverride"));
 	const TVec3<FRealSingle> GravityOverride = (TVec3<FRealSingle>)ConfigProperties.GetValue<FVector3f>(TEXT("GravityOverride"), FVector3f(0.f, 0.f, ClothingSimulationClothDefault::GravityZOverride));
 	const FRealSingle GravityScale = (FRealSingle)ConfigProperties.GetValue<float>(TEXT("GravityScale"), 1.f);
+	const FRealSingle GravityMultiplier = (FRealSingle)ClothingSimulationClothConsoleVariables::CVarGravityMultiplier.GetValueOnAnyThread();
 
-	return Solver->IsClothGravityOverrideEnabled() && bUseGravityOverride ? GravityOverride : Solver->GetGravity() * GravityScale;
+	return (Solver->IsClothGravityOverrideEnabled() && bUseGravityOverride ? GravityOverride : Solver->GetGravity() * GravityScale) * GravityMultiplier;
 }
 
 FAABB3 FClothingSimulationCloth::CalculateBoundingBox(const FClothingSimulationSolver* Solver) const
