Update the Cloth solver plugin to optionally use the new Softs::FEvolution.

Currently Softs::FEvolution should have the same behavior as PBDEvolution, but Softs::FEvolution will ultimately support a force-based linear system solver, hence why it is keyed on EnableForceBasedSolver. The new evolution will be created if ANY LOD requests a ForceBasedSolver since the ForceBasedSolver can also just do vanilla PBD.

In the long run, we might want to just swap to the new solver entirely and clean up a bunch of code, but not doing that yet.

I did some basic tests, and due to the optimizations in presubstep, the new solver is slightly faster than the old solver.

diff --git a/Engine/Plugins/Experimental/ChaosCloth/Source/Chaos/Private/ChaosCloth/ChaosClothConstraints.cpp b/Engine/Plugins/Experimental/ChaosCloth/Source/Chaos/Private/ChaosCloth/ChaosClothConstraints.cpp
--- a/Engine/Plugins/Experimental/ChaosCloth/Source/Chaos/Private/ChaosCloth/ChaosClothConstraints.cpp
+++ b/Engine/Plugins/Experimental/ChaosCloth/Source/Chaos/Private/ChaosCloth/ChaosClothConstraints.cpp
@@ -20,14 +20,29 @@
 #include "Chaos/PBDTriangleMeshCollisions.h"
 #include "Chaos/PBDTriangleMeshIntersections.h"
 #include "Chaos/PBDEvolution.h"
+#include "Chaos/PBDSoftBodyCollisionConstraint.h"
+#include "Chaos/SoftsExternalForces.h"
+#include "Chaos/SoftsEvolution.h"
+#include "Chaos/VelocityField.h"
 #include "Chaos/CollectionPropertyFacade.h"
 #include "Chaos/Deformable/GaussSeidelMasterConstraint.h"
 #include "Chaos/Deformable/GaussSeidelCorotatedCodimensionalConstraints.h"
+#include "PhysicsProxy/PerSolverFieldSystem.h"
+#include "HAL/IConsoleManager.h"
 
 namespace Chaos {
+
+namespace ClothingSimulationClothConsoleVariables
+{
+// These are defined in ChaosClothingSimulationCloth.cpp
+extern TAutoConsoleVariable<bool> CVarLegacyDisablesAccurateWind;
+extern TAutoConsoleVariable<float> CVarGravityMultiplier;
+}
+
+bool bEnableGS = false;
+
 #if !UE_BUILD_SHIPPING
 
-	bool bEnableGS = false;
 
     static FAutoConsoleVariableRef CVarClothbEnableGS(TEXT("p.Chaos.Cloth.EnableGaussSeidel"), bEnableGS, TEXT("Use Gauss Seidel constraints instead of XPBD [def: false]"));
 
@@ -86,21 +101,35 @@ namespace Chaos {
 #endif
 
 FClothConstraints::FClothConstraints()
-	: Evolution(nullptr)
+	: Evolution(nullptr), PBDEvolution(nullptr)
 	, AnimationPositions(nullptr)
-	, OldAnimationPositions_Deprecated(nullptr)
 	, AnimationNormals(nullptr)
 	, AnimationVelocities(nullptr)
 	, ParticleOffset(0)
+	, ParticleRangeId(0)
 	, NumParticles(0)
-	, ConstraintInitOffset(INDEX_NONE)
-	, ConstraintRuleOffset(INDEX_NONE)
-	, PostCollisionConstraintRuleOffset(INDEX_NONE)
-	, PostprocessingConstraintRuleOffset(INDEX_NONE)
 	, NumConstraintInits(0)
 	, NumConstraintRules(0)
 	, NumPostCollisionConstraintRules(0)
 	, NumPostprocessingConstraintRules(0)
+
+	, PerSolverField(nullptr)
+	, Normals(nullptr)
+	, LastSubframeCollisionTransformsCCD(nullptr)
+	, CollisionParticleCollided(nullptr)
+	, CollisionContacts(nullptr)
+	, CollisionNormals(nullptr)
+	, CollisionPhis(nullptr)
+	, NumPreSubstepInits(0)
+	, NumExternalForceRules(0)
+	, NumPreSubstepConstraintRules(0)
+	, NumCollisionConstraintRules(0)
+
+	, OldAnimationPositions_Deprecated(nullptr)
+	, ConstraintInitOffset(INDEX_NONE)
+	, ConstraintRuleOffset(INDEX_NONE)
+	, PostCollisionConstraintRuleOffset(INDEX_NONE)
+	, PostprocessingConstraintRuleOffset(INDEX_NONE)
 {
 }
 
@@ -111,20 +140,34 @@ FClothConstraints::~FClothConstraints()
 PRAGMA_ENABLE_DEPRECATION_WARNINGS
 
 void FClothConstraints::Initialize(
-	Softs::FPBDEvolution* InEvolution,
-	const TArray<Softs::FSolverVec3>& InAnimationPositions,
-	const TArray<Softs::FSolverVec3>& InOldAnimationPositions,
-	const TArray<Softs::FSolverVec3>& InAnimationNormals,
-	int32 InParticleOffset,
-	int32 InNumParticles)
+	Softs::FEvolution* InEvolution,
+	FPerSolverFieldSystem* InPerSolverField,
+	const TArray<Softs::FSolverVec3>& InInterpolatedAnimationPositions,
+	const TArray<Softs::FSolverVec3>& InInterpolatedAnimationNormals,
+	const TArray<Softs::FSolverVec3>& InAnimationVelocities,
+	const TArray<Softs::FSolverVec3>& InNormals,
+	const TArray<Softs::FSolverRigidTransform3>& InLastSubframeCollisionTransformsCCD,
+	TArray<bool>& InCollisionParticleCollided,
+	TArray<Softs::FSolverVec3>& InCollisionContacts,
+	TArray<Softs::FSolverVec3>& InCollisionNormals,
+	TArray<Softs::FSolverReal>& InCollisionPhis,
+	int32 InParticleRangeId)
 {
+	PBDEvolution = nullptr;
 	Evolution = InEvolution;
-	AnimationPositions = &InAnimationPositions;
-	OldAnimationPositions_Deprecated = &InOldAnimationPositions;
-	AnimationNormals = &InAnimationNormals;
-	AnimationVelocities = nullptr;
-	ParticleOffset = InParticleOffset;
-	NumParticles = InNumParticles;
+	PerSolverField = InPerSolverField;
+	AnimationPositions = &InInterpolatedAnimationPositions;
+	AnimationNormals = &InInterpolatedAnimationNormals;
+	AnimationVelocities = &InAnimationVelocities;
+	Normals = &InNormals;
+	LastSubframeCollisionTransformsCCD = &InLastSubframeCollisionTransformsCCD;
+	CollisionParticleCollided = &InCollisionParticleCollided;
+	CollisionContacts = &InCollisionContacts;
+	CollisionNormals = &InCollisionNormals;
+	CollisionPhis = &InCollisionPhis;
+	ParticleOffset = 0;
+	ParticleRangeId = InParticleRangeId;
+	NumParticles = Evolution->GetSoftBodyParticleNum(ParticleRangeId);
 }
 
 void FClothConstraints::Initialize(
@@ -136,33 +179,38 @@ void FClothConstraints::Initialize(
 	int32 InParticleOffset,
 	int32 InNumParticles)
 {
-	Evolution = InEvolution;
+	Evolution = nullptr;
+	PBDEvolution = InEvolution;
 	AnimationPositions = &InInterpolatedAnimationPositions;
 	OldAnimationPositions_Deprecated = nullptr;
 	AnimationNormals = &InInterpolatedAnimationNormals;
 	AnimationVelocities = &InAnimationVelocities;
 	ParticleOffset = InParticleOffset;
+	ParticleRangeId = InParticleOffset;
 	NumParticles = InNumParticles;
 }
 
 void FClothConstraints::Enable(bool bEnable)
 {
-	check(Evolution);
 	if (ConstraintInitOffset != INDEX_NONE)
 	{
-		Evolution->ActivateConstraintInitRange(ConstraintInitOffset, bEnable);
+		check(PBDEvolution);
+		PBDEvolution->ActivateConstraintInitRange(ConstraintInitOffset, bEnable);
 	}
 	if (ConstraintRuleOffset != INDEX_NONE)
 	{
-		Evolution->ActivateConstraintRuleRange(ConstraintRuleOffset, bEnable);
+		check(PBDEvolution);
+		PBDEvolution->ActivateConstraintRuleRange(ConstraintRuleOffset, bEnable);
 	}
 	if (PostCollisionConstraintRuleOffset != INDEX_NONE)
 	{
-		Evolution->ActivatePostCollisionConstraintRuleRange(PostCollisionConstraintRuleOffset, bEnable);
+		check(PBDEvolution);
+		PBDEvolution->ActivatePostCollisionConstraintRuleRange(PostCollisionConstraintRuleOffset, bEnable);
 	}
 	if (PostprocessingConstraintRuleOffset != INDEX_NONE)
 	{
-		Evolution->ActivateConstraintPostprocessingsRange(PostprocessingConstraintRuleOffset, bEnable);
+		check(PBDEvolution);
+		PBDEvolution->ActivateConstraintPostprocessingsRange(PostprocessingConstraintRuleOffset, bEnable);
 	}
 }
 
@@ -211,6 +259,8 @@ void FClothConstraints::AddRules(
 void FClothConstraints::GetGSNumRules()
 {
 #if !UE_BUILD_SHIPPING
+	check(PBDEvolution);
+
 	NumConstraintInits = 1;
 	NumConstraintRules = 0;
 	if (XStretchBiasConstraints)
@@ -245,15 +295,15 @@ void FClothConstraints::GetGSNumRules()
 void FClothConstraints::CreateGSRules()
 {
 #if !UE_BUILD_SHIPPING
-	check(Evolution);
+	check(PBDEvolution);
 	check(ConstraintInitOffset == INDEX_NONE)
 	GetGSNumRules();
 
-	ConstraintInitOffset = Evolution->AddConstraintInitRange(NumConstraintInits, false);
+	ConstraintInitOffset = PBDEvolution->AddConstraintInitRange(NumConstraintInits, false);
 
 	check(ConstraintRuleOffset == INDEX_NONE)
 
-	ConstraintRuleOffset = Evolution->AddConstraintRuleRange(NumConstraintRules, false);
+	ConstraintRuleOffset = PBDEvolution->AddConstraintRuleRange(NumConstraintRules, false);
 
 	check(PostprocessingConstraintRuleOffset == INDEX_NONE)
 
@@ -265,24 +315,24 @@ void FClothConstraints::CreateGSRules()
 
 	if (PostProcessingConstraintRule)
 	{
-		PostprocessingConstraintRuleOffset = Evolution->AddConstraintPostprocessingsRange(PostProcessingConstraintRule, false);
+		PostprocessingConstraintRuleOffset = PBDEvolution->AddConstraintPostprocessingsRange(PostProcessingConstraintRule, false);
 	}
 
-	TFunction<void(Softs::FSolverParticles&, const Softs::FSolverReal)>* const ConstraintInits = Evolution->ConstraintInits().GetData() + ConstraintInitOffset;
-	TFunction<void(Softs::FSolverParticles&, const Softs::FSolverReal)>* const ConstraintRules = Evolution->ConstraintRules().GetData() + ConstraintRuleOffset;
-	TFunction<void(Softs::FSolverParticles&, const Softs::FSolverReal)>* const PostprocessingRules = Evolution->ConstraintPostprocessings().GetData() + PostprocessingConstraintRuleOffset;
+	TFunction<void(Softs::FSolverParticles&, const Softs::FSolverReal)>* const ConstraintInits = PBDEvolution->ConstraintInits().GetData() + ConstraintInitOffset;
+	TFunction<void(Softs::FSolverParticles&, const Softs::FSolverReal)>* const ConstraintRules = PBDEvolution->ConstraintRules().GetData() + ConstraintRuleOffset;
+	TFunction<void(Softs::FSolverParticles&, const Softs::FSolverReal)>* const PostprocessingRules = PBDEvolution->ConstraintPostprocessings().GetData() + PostprocessingConstraintRuleOffset;
 
 	int32 ConstraintInitIndex = 0;
 	int32 ConstraintRuleIndex = 0;
 
 	int32 ConstraintPostProcessingIndex = 0;
 
-	GSMasterConstraint = MakeShared<Chaos::Softs::FGaussSeidelMasterConstraint<Softs::FSolverReal, Softs::FSolverParticles>>(Evolution->Particles(), bClothDoQuasistatics, bUseSOR, SOROmega, 100);
+	GSMasterConstraint = MakeShared<Chaos::Softs::FGaussSeidelMasterConstraint<Softs::FSolverReal, Softs::FSolverParticles>>(PBDEvolution->Particles(), bClothDoQuasistatics, bUseSOR, SOROmega, 100);
 
 	if (bClothDoQuasistatics)
 	{
-		Evolution->SetQuasistatics(bClothDoQuasistatics);
-		const Softs::FSolverVec3& GravityTerm = Evolution->GetGravity(0);
+		PBDEvolution->SetQuasistatics(bClothDoQuasistatics);
+		const Softs::FSolverVec3& GravityTerm = PBDEvolution->GetGravity(0);
 		GSMasterConstraint->ExternalForce[0] = 0.f;
 		GSMasterConstraint->ExternalForce[1] = 0.f;
 		GSMasterConstraint->ExternalForce[2] = GravityTerm[2];
@@ -411,7 +461,7 @@ void FClothConstraints::CreateGSRules()
 
 			const TArray<TVec3<int32>> CodimensionalMesh = XStretchBiasConstraints->GetConstraints();
 
-			GSCorotatedCodimensionalConstraint = MakeShared<Chaos::Softs::FGaussSeidelCorotatedCodimensionalConstraints<Softs::FSolverReal, Softs::FSolverParticles>>(Evolution->Particles(), CodimensionalMesh, false, YoungsModulus);
+			GSCorotatedCodimensionalConstraint = MakeShared<Chaos::Softs::FGaussSeidelCorotatedCodimensionalConstraints<Softs::FSolverReal, Softs::FSolverParticles>>(PBDEvolution->Particles(), CodimensionalMesh, false, YoungsModulus);
 
 			GSMasterConstraint->AddStaticConstraints(XStretchBiasConstraints->GetConstraintsArray(), IncidentElements, IncidentElementsLocal);
 			const int32 StaticIndex = GSMasterConstraint->AddStaticConstraintResidualAndHessianRange(1);
@@ -444,7 +494,7 @@ void FClothConstraints::CreateGSRules()
 
 		}
 
-		XStretchBiasConstraints->InitializeDmInvAndMeasures(Evolution->Particles());
+		XStretchBiasConstraints->InitializeDmInvAndMeasures(PBDEvolution->Particles());
 
 	}
 
@@ -494,7 +544,7 @@ void FClothConstraints::CreateGSRules()
 		}
 	}
 
-	GSMasterConstraint->InitStaticColor(Evolution->Particles());
+	GSMasterConstraint->InitStaticColor(PBDEvolution->Particles());
 
 	if (bWriteFinalResiduals)
 	{
@@ -545,27 +595,46 @@ void FClothConstraints::AddRules(
 	// Animation Drive Constraints
 	CreateAnimDriveConstraints(ConfigProperties, WeightMaps);
 
+	if (Evolution)
+	{
+		// External Forces
+		CreateExternalForces(ConfigProperties);
 
-	// Commit rules to solver
-PRAGMA_DISABLE_DEPRECATION_WARNINGS
+		// Velocity Field
+		CreateVelocityAndPressureField(ConfigProperties, WeightMaps, TriangleMesh);
 
-#if !UE_BUILD_SHIPPING
+		// PerSolverField
+		if (PerSolverField)
+		{
+			++NumExternalForceRules;
+		}
+
+		// Body collisions
+		CreateCollisionConstraint(ConfigProperties, MeshScale);
+	}
+
+	// Commit rules to solver
 	if (!bEnableGS)
 	{
-		CreateRules();  // TODO: Move CreateRules to private
+		if (Evolution)
+		{
+			CreateForceBasedRules();
+		}
+		else
+		{
+			CreatePBDRules();
+		}
 	} 
 	else
 	{
 		CreateGSRules();
 	}
-#else
-	CreateRules();
-#endif
-
-PRAGMA_ENABLE_DEPRECATION_WARNINGS
 
-	// Enable or disable constraints as requested
-	Enable(bEnabled);
+	if (PBDEvolution)
+	{
+		// Enable or disable constraints as requested
+		Enable(bEnabled);
+	}
 }
 
 void FClothConstraints::CreateSelfCollisionConstraints(const Softs::FCollectionPropertyConstFacade& ConfigProperties, 
@@ -611,12 +680,19 @@ void FClothConstraints::CreateSelfCollisionConstraints(const Softs::FCollectionP
 			ParticleOffset,
 			NumParticles,
 			TriangleMesh);
-		++NumConstraintInits;
+		if (Evolution)
+		{
+			++NumPreSubstepConstraintRules;
+		}
+		else
+		{
+			++NumConstraintInits;
+		}
 		++NumPostprocessingConstraintRules;
 	}
 	else if (Softs::FPBDSelfCollisionSphereConstraints::IsEnabled(ConfigProperties))
 	{
-		SelfCollisionSphereConstraints = MakeShared< Softs::FPBDSelfCollisionSphereConstraints>(
+		SelfCollisionSphereConstraints = MakeShared<Softs::FPBDSelfCollisionSphereConstraints>(
 			ParticleOffset,
 			NumParticles,
 			VertexSets,
@@ -634,43 +710,80 @@ void FClothConstraints::CreateStretchConstraints(
 {
 	if (PatternData && PatternData->PatternPositions.Num() && Softs::FXPBDStretchBiasElementConstraints::IsEnabled(ConfigProperties))
 	{
-		XStretchBiasConstraints = MakeShared<Softs::FXPBDStretchBiasElementConstraints>(
-			Evolution->Particles(),
-			ParticleOffset,
-			NumParticles,
-			TriangleMesh,
-			PatternData->WeldedFaceVertexPatternPositions,
-			WeightMaps,
-			ConfigProperties,
-			/*bTrimKinematicConstraints =*/ true);
+		if (Evolution)
+		{
+			XStretchBiasConstraints = MakeShared<Softs::FXPBDStretchBiasElementConstraints>(
+				Evolution->GetSoftBodyParticles(ParticleRangeId),
+				TriangleMesh,
+				PatternData->WeldedFaceVertexPatternPositions,
+				WeightMaps,
+				ConfigProperties,
+				/*bTrimKinematicConstraints =*/ true);
+		}
+		else
+		{
+			XStretchBiasConstraints = MakeShared<Softs::FXPBDStretchBiasElementConstraints>(
+				PBDEvolution->Particles(),
+				ParticleOffset,
+				NumParticles,
+				TriangleMesh,
+				PatternData->WeldedFaceVertexPatternPositions,
+				WeightMaps,
+				ConfigProperties,
+				/*bTrimKinematicConstraints =*/ true);
+		}
 
 		++NumConstraintInits;  // Uses init to update the property tables
 		++NumConstraintRules;
 	}
 	else if (Softs::FXPBDEdgeSpringConstraints::IsEnabled(ConfigProperties))
 	{
-		XEdgeConstraints = MakeShared<Softs::FXPBDEdgeSpringConstraints>(
-			Evolution->Particles(),
-			ParticleOffset,
-			NumParticles,
-			TriangleMesh.GetSurfaceElements(),
-			WeightMaps,
-			ConfigProperties,
-			/*bTrimKinematicConstraints =*/ true);
+		if (Evolution)
+		{
+			XEdgeConstraints = MakeShared<Softs::FXPBDEdgeSpringConstraints>(
+				Evolution->GetSoftBodyParticles(ParticleRangeId),
+				TriangleMesh.GetSurfaceElements(),
+				WeightMaps,
+				ConfigProperties,
+				/*bTrimKinematicConstraints =*/ true);
+		}
+		else
+		{
+			XEdgeConstraints = MakeShared<Softs::FXPBDEdgeSpringConstraints>(
+				PBDEvolution->Particles(),
+				ParticleOffset,
+				NumParticles,
+				TriangleMesh.GetSurfaceElements(),
+				WeightMaps,
+				ConfigProperties,
+				/*bTrimKinematicConstraints =*/ true);
+		}
 
 		++NumConstraintInits;  // Uses init to update the property tables
 		++NumConstraintRules;
 	}
 	else if (Softs::FPBDEdgeSpringConstraints::IsEnabled(ConfigProperties))
 	{
-		EdgeConstraints = MakeShared<Softs::FPBDEdgeSpringConstraints>(
-			Evolution->Particles(),
-			ParticleOffset,
-			NumParticles,
-			TriangleMesh.GetSurfaceElements(),
-			WeightMaps,
-			ConfigProperties,
-			/*bTrimKinematicConstraints =*/ true);
+		if (Evolution)
+		{
+			EdgeConstraints = MakeShared<Softs::FPBDEdgeSpringConstraints>(
+				Evolution->GetSoftBodyParticles(ParticleRangeId),
+				TriangleMesh.GetSurfaceElements(),
+				WeightMaps,
+				ConfigProperties,
+				/*bTrimKinematicConstraints =*/ true);
+		}
+		else
+		{
+			EdgeConstraints = MakeShared<Softs::FPBDEdgeSpringConstraints>(
+				PBDEvolution->Particles(),
+				ParticleOffset,
+				NumParticles,
+				TriangleMesh.GetSurfaceElements(),
+				WeightMaps,
+				ConfigProperties,
+				/*bTrimKinematicConstraints =*/ true);
+		}
 
 		++NumConstraintInits;  // Uses init to update the property tables
 		++NumConstraintRules;
@@ -685,14 +798,27 @@ void FClothConstraints::CreateBendingConstraints(
 {
 	if (PatternData && PatternData->PatternPositions.Num() && Softs::FXPBDAnisotropicBendingConstraints::IsEnabled(ConfigProperties))
 	{
-		XAnisoBendingElementConstraints = MakeShared<Softs::FXPBDAnisotropicBendingConstraints>(
-			Evolution->Particles(),
-			ParticleOffset, NumParticles,
-			TriangleMesh,
-			PatternData->WeldedFaceVertexPatternPositions,
-			WeightMaps,
-			ConfigProperties,
-			/*bTrimKinematicConstraints =*/ true);
+		if (Evolution)
+		{
+			XAnisoBendingElementConstraints = MakeShared<Softs::FXPBDAnisotropicBendingConstraints>(
+				Evolution->GetSoftBodyParticles(ParticleRangeId),
+				TriangleMesh,
+				PatternData->WeldedFaceVertexPatternPositions,
+				WeightMaps,
+				ConfigProperties,
+				/*bTrimKinematicConstraints =*/ true);
+		}
+		else
+		{
+			XAnisoBendingElementConstraints = MakeShared<Softs::FXPBDAnisotropicBendingConstraints>(
+				PBDEvolution->Particles(),
+				ParticleOffset, NumParticles,
+				TriangleMesh,
+				PatternData->WeldedFaceVertexPatternPositions,
+				WeightMaps,
+				ConfigProperties,
+				/*bTrimKinematicConstraints =*/ true);
+		}
 
 		++NumConstraintInits;  // Uses init to update the property tables
 		++NumConstraintRules;
@@ -701,13 +827,25 @@ void FClothConstraints::CreateBendingConstraints(
 	{
 		TArray<Chaos::TVec4<int32>> BendingElements = TriangleMesh.GetUniqueAdjacentElements();
 
-		XBendingElementConstraints = MakeShared<Softs::FXPBDBendingConstraints>(
-			Evolution->Particles(),
-			ParticleOffset, NumParticles,
-			MoveTemp(BendingElements),
-			WeightMaps,
-			ConfigProperties,
-			/*bTrimKinematicConstraints =*/ true);
+		if (Evolution)
+		{
+			XBendingElementConstraints = MakeShared<Softs::FXPBDBendingConstraints>(
+				Evolution->GetSoftBodyParticles(ParticleRangeId),
+				MoveTemp(BendingElements),
+				WeightMaps,
+				ConfigProperties,
+				/*bTrimKinematicConstraints =*/ true);
+		}
+		else
+		{
+			XBendingElementConstraints = MakeShared<Softs::FXPBDBendingConstraints>(
+				PBDEvolution->Particles(),
+				ParticleOffset, NumParticles,
+				MoveTemp(BendingElements),
+				WeightMaps,
+				ConfigProperties,
+				/*bTrimKinematicConstraints =*/ true);
+		}
 
 		++NumConstraintInits;  // Uses init to update the property tables
 		++NumConstraintRules;
@@ -716,13 +854,25 @@ void FClothConstraints::CreateBendingConstraints(
 	{
 		TArray<Chaos::TVec4<int32>> BendingElements = TriangleMesh.GetUniqueAdjacentElements();
 
-		BendingElementConstraints = MakeShared<Softs::FPBDBendingConstraints>(
-			Evolution->Particles(),
-			ParticleOffset, NumParticles,
-			MoveTemp(BendingElements),
-			WeightMaps,
-			ConfigProperties,
-			/*bTrimKinematicConstraints =*/ true);
+		if (Evolution)
+		{
+			BendingElementConstraints = MakeShared<Softs::FPBDBendingConstraints>(
+				Evolution->GetSoftBodyParticles(ParticleRangeId),
+				MoveTemp(BendingElements),
+				WeightMaps,
+				ConfigProperties,
+				/*bTrimKinematicConstraints =*/ true);
+		}
+		else
+		{
+			BendingElementConstraints = MakeShared<Softs::FPBDBendingConstraints>(
+				PBDEvolution->Particles(),
+				ParticleOffset, NumParticles,
+				MoveTemp(BendingElements),
+				WeightMaps,
+				ConfigProperties,
+				/*bTrimKinematicConstraints =*/ true);
+		}
 
 		++NumConstraintInits;  // Uses init to update the property tables
 		++NumConstraintRules;
@@ -731,13 +881,25 @@ void FClothConstraints::CreateBendingConstraints(
 	{
 		const TArray<Chaos::TVec2<int32>> CrossEdges = TriangleMesh.GetUniqueAdjacentPoints();
 
-		XBendingConstraints = MakeShared<Softs::FXPBDBendingSpringConstraints>(
-			Evolution->Particles(),
-			ParticleOffset, NumParticles,
-			CrossEdges,
-			WeightMaps,
-			ConfigProperties,
-			/*bTrimKinematicConstraints =*/ true);
+		if (Evolution)
+		{
+			XBendingConstraints = MakeShared<Softs::FXPBDBendingSpringConstraints>(
+				Evolution->GetSoftBodyParticles(ParticleRangeId),
+				CrossEdges,
+				WeightMaps,
+				ConfigProperties,
+				/*bTrimKinematicConstraints =*/ true);
+		}
+		else
+		{
+			XBendingConstraints = MakeShared<Softs::FXPBDBendingSpringConstraints>(
+				PBDEvolution->Particles(),
+				ParticleOffset, NumParticles,
+				CrossEdges,
+				WeightMaps,
+				ConfigProperties,
+				/*bTrimKinematicConstraints =*/ true);
+		}
 
 		++NumConstraintInits;  // Uses init to update the property tables
 		++NumConstraintRules;
@@ -746,14 +908,26 @@ void FClothConstraints::CreateBendingConstraints(
 	{
 		const TArray<Chaos::TVec2<int32>> CrossEdges = TriangleMesh.GetUniqueAdjacentPoints();
 
-		BendingConstraints = MakeShared<Softs::FPBDBendingSpringConstraints>(
-			Evolution->Particles(),
-			ParticleOffset,
-			NumParticles,
-			CrossEdges,
-			WeightMaps,
-			ConfigProperties,
-			/*bTrimKinematicConstraints =*/ true);
+		if (Evolution)
+		{
+			BendingConstraints = MakeShared<Softs::FPBDBendingSpringConstraints>(
+				Evolution->GetSoftBodyParticles(ParticleRangeId),
+				CrossEdges,
+				WeightMaps,
+				ConfigProperties,
+				/*bTrimKinematicConstraints =*/ true);
+		}
+		else
+		{
+			BendingConstraints = MakeShared<Softs::FPBDBendingSpringConstraints>(
+				PBDEvolution->Particles(),
+				ParticleOffset,
+				NumParticles,
+				CrossEdges,
+				WeightMaps,
+				ConfigProperties,
+				/*bTrimKinematicConstraints =*/ true);
+		}
 
 		++NumConstraintInits;  // Uses init to update the property tables
 		++NumConstraintRules;
@@ -767,28 +941,52 @@ void FClothConstraints::CreateAreaConstraints(
 {
 	if (Softs::FXPBDAreaSpringConstraints::IsEnabled(ConfigProperties))
 	{
-		XAreaConstraints = MakeShared<Softs::FXPBDAreaSpringConstraints>(
-			Evolution->Particles(),
-			ParticleOffset,
-			NumParticles,
-			TriangleMesh.GetSurfaceElements(),
-			WeightMaps,
-			ConfigProperties,
-			/*bTrimKinematicConstraints =*/ true);
+		if (Evolution)
+		{
+			XAreaConstraints = MakeShared<Softs::FXPBDAreaSpringConstraints>(
+				Evolution->GetSoftBodyParticles(ParticleRangeId),
+				TriangleMesh.GetSurfaceElements(),
+				WeightMaps,
+				ConfigProperties,
+				/*bTrimKinematicConstraints =*/ true);
+		}
+		else
+		{
+			XAreaConstraints = MakeShared<Softs::FXPBDAreaSpringConstraints>(
+				PBDEvolution->Particles(),
+				ParticleOffset,
+				NumParticles,
+				TriangleMesh.GetSurfaceElements(),
+				WeightMaps,
+				ConfigProperties,
+				/*bTrimKinematicConstraints =*/ true);
+		}
 
 		++NumConstraintInits;  // Uses init to update the property tables
 		++NumConstraintRules;
 	}
 	else if (Softs::FPBDAreaSpringConstraints::IsEnabled(ConfigProperties))
 	{
-		AreaConstraints = MakeShared<Softs::FPBDAreaSpringConstraints>(
-			Evolution->Particles(),
-			ParticleOffset,
-			NumParticles,
-			TriangleMesh.GetSurfaceElements(),
-			WeightMaps,
-			ConfigProperties,
-			/*bTrimKinematicConstraints =*/ true);
+		if (Evolution)
+		{
+			AreaConstraints = MakeShared<Softs::FPBDAreaSpringConstraints>(
+				Evolution->GetSoftBodyParticles(ParticleRangeId),
+				TriangleMesh.GetSurfaceElements(),
+				WeightMaps,
+				ConfigProperties,
+				/*bTrimKinematicConstraints =*/ true);
+		}
+		else
+		{
+			AreaConstraints = MakeShared<Softs::FPBDAreaSpringConstraints>(
+				PBDEvolution->Particles(),
+				ParticleOffset,
+				NumParticles,
+				TriangleMesh.GetSurfaceElements(),
+				WeightMaps,
+				ConfigProperties,
+				/*bTrimKinematicConstraints =*/ true);
+		}
 
 		++NumConstraintInits;  // Uses init to update the property tables
 		++NumConstraintRules;
@@ -803,17 +1001,29 @@ void FClothConstraints::CreateLongRangeConstraints(
 {
 	if (Softs::FPBDLongRangeConstraints::IsEnabled(ConfigProperties))
 	{
-		//  Now that we're only doing a single iteration of Long range constraints, and they're more of a fake constraint to jump start our initial guess, it's not clear that using XPBD makes sense here.
-		LongRangeConstraints = MakeShared<Softs::FPBDLongRangeConstraints>(
-			Evolution->Particles(),
-			ParticleOffset,
-			NumParticles,
-			Tethers,
-			WeightMaps,
-			ConfigProperties,
-			MeshScale);
-
-		++NumConstraintInits;  // Uses init to both update the property tables and apply the constraint
+		if (Evolution)
+		{
+			//  Now that we're only doing a single iteration of Long range constraints, and they're more of a fake constraint to jump start our initial guess, it's not clear that using XPBD makes sense here.
+			LongRangeConstraints = MakeShared<Softs::FPBDLongRangeConstraints>(
+				Evolution->GetSoftBodyParticles(ParticleRangeId),
+				Tethers,
+				WeightMaps,
+				ConfigProperties,
+				MeshScale);
+			++NumPreSubstepConstraintRules;
+		}
+		else
+		{
+			LongRangeConstraints = MakeShared<Softs::FPBDLongRangeConstraints>(
+				PBDEvolution->Particles(),
+				ParticleOffset,
+				NumParticles,
+				Tethers,
+				WeightMaps,
+				ConfigProperties,
+				MeshScale);
+			++NumConstraintInits;
+		}
 	}
 }
 
@@ -903,98 +1113,500 @@ void FClothConstraints::CreateAnimDriveConstraints(
 	}
 }
 
-void FClothConstraints::CreateRules()
+void FClothConstraints::CreateVelocityAndPressureField(
+	const Softs::FCollectionPropertyConstFacade& ConfigProperties,
+	const TMap<FString, TConstArrayView<FRealSingle>>& WeightMaps,
+	const FTriangleMesh& TriangleMesh)
 {
-	check(Evolution);
-	check(ConstraintInitOffset == INDEX_NONE)
-#if !UE_BUILD_SHIPPING	
-	if (bDisplayResidual)
-	{
-		NumConstraintRules++;
-	}
-	if (bWriteFinalResiduals || bDisplayResidual)
-	{
-		NumConstraintInits++;
-	}
-#endif
-	if (NumConstraintInits)
+	if (Evolution)
 	{
-		ConstraintInitOffset = Evolution->AddConstraintInitRange(NumConstraintInits, false);
+		// Always create velocity field--we allow turning it on via blueprints
+		constexpr Softs::FSolverReal WorldScale = (Softs::FSolverReal)100.;
+		VelocityAndPressureField = MakeShared<Softs::FVelocityAndPressureField>(
+			&TriangleMesh,
+			ConfigProperties,
+			WeightMaps,
+			WorldScale
+			);
+		++NumPreSubstepInits;
+		++NumExternalForceRules;
 	}
-	check(ConstraintRuleOffset == INDEX_NONE)
+}
 
-	if (NumConstraintRules)
-	{
-		ConstraintRuleOffset = Evolution->AddConstraintRuleRange(NumConstraintRules, false);
-	}
-	check(PostCollisionConstraintRuleOffset == INDEX_NONE);
-	if (NumPostCollisionConstraintRules)
+void FClothConstraints::CreateExternalForces(
+	const Softs::FCollectionPropertyConstFacade& ConfigProperties)
+{
+	if (Evolution)
 	{
-		PostCollisionConstraintRuleOffset = Evolution->AddPostCollisionConstraintRuleRange(NumPostCollisionConstraintRules, false);
+		// Always create external forces
+		check(Normals);
+		ExternalForces = MakeShared<Softs::FExternalForces>(
+			*Normals,
+			ConfigProperties
+			);
+
+		++NumExternalForceRules;
 	}
-	check(PostprocessingConstraintRuleOffset == INDEX_NONE);
-#if !UE_BUILD_SHIPPING	
-	if (bWriteFinalResiduals)
+}
+
+void FClothConstraints::CreateCollisionConstraint(
+	const Softs::FCollectionPropertyConstFacade& ConfigProperties,
+	Softs::FSolverReal MeshScale)
+{
+	if (Evolution)
 	{
-		NumPostprocessingConstraintRules += 1;
+		// Always create collision constraint
+		check(LastSubframeCollisionTransformsCCD);
+		CollisionConstraint = MakeShared<Softs::FPBDSoftBodyCollisionConstraint>(
+			*LastSubframeCollisionTransformsCCD,
+			ConfigProperties,
+			MeshScale,
+			CollisionParticleCollided,
+			CollisionContacts,
+			CollisionNormals,
+			CollisionPhis
+			);
+		++NumCollisionConstraintRules;
 	}
-#endif
-	if (NumPostprocessingConstraintRules)
+}
+
+void FClothConstraints::CreateForceBasedRules()
+{
+	Evolution->AllocatePreSubstepParallelInitRange(ParticleRangeId, NumPreSubstepInits);
+	Evolution->AllocateExternalForceRulesRange(ParticleRangeId, NumExternalForceRules);
+	Evolution->AllocateConstraintParallelInitsRange(ParticleRangeId, NumConstraintInits);
+	Evolution->AllocatePreSubstepConstraintRulesRange(ParticleRangeId, NumPreSubstepConstraintRules);
+	Evolution->AllocatePerIterationConstraintRulesRange(ParticleRangeId, NumConstraintRules);
+	Evolution->AllocatePerIterationCollisionConstraintRulesRange(ParticleRangeId, NumCollisionConstraintRules);
+	Evolution->AllocatePerIterationPostCollisionsConstraintRulesRange(ParticleRangeId, NumPostCollisionConstraintRules);
+	Evolution->AllocatePostSubstepConstraintRulesRange(ParticleRangeId, NumPostprocessingConstraintRules);
+
+	TArrayView<Softs::FEvolution::PreSubstepParallelInitFunc> PreSubstepParallelInits =
+		Evolution->GetPreSubstepParallelInitRange(ParticleRangeId);
+	TArrayView<Softs::FEvolution::ExternalForceRuleFunc> ExternalForceRules =
+		Evolution->GetExternalForceRulesRange(ParticleRangeId);
+	TArrayView<Softs::FEvolution::ConstraintParallelInitFunc> ConstraintParallelInits =
+		Evolution->GetConstraintParallelInitsRange(ParticleRangeId);
+	TArrayView<Softs::FEvolution::ConstraintRuleFunc> PreSubstepConstraintRules =
+		Evolution->GetPreSubstepConstraintRulesRange(ParticleRangeId);
+	TArrayView<Softs::FEvolution::ConstraintRuleFunc> PerIterationConstraintRules =
+		Evolution->GetPerIterationConstraintRulesRange(ParticleRangeId);
+	TArrayView<Softs::FEvolution::CollisionConstraintRuleFunc> PerIterationCollisionConstraintRules =
+		Evolution->GetPerIterationCollisionConstraintRulesRange(ParticleRangeId);
+	TArrayView<Softs::FEvolution::ConstraintRuleFunc> PerIterationPostCollisionsConstraintRules =
+		Evolution->GetPerIterationPostCollisionsConstraintRulesRange(ParticleRangeId);
+	TArrayView<Softs::FEvolution::ConstraintRuleFunc> PostSubstepConstraintRules =
+		Evolution->GetPostSubstepConstraintRulesRange(ParticleRangeId);
+
+	int32 PreSubstepInitsIndex = 0;
+	int32 ExternalForceRulesIndex = 0;
+	int32 ConstraintInitsIndex = 0;
+	int32 PreSubstepConstraintRulesIndex = 0;
+	int32 ConstraintRuleIndex = 0;
+	int32 CollisionConstraintRulesIndex = 0;
+	int32 PostCollisionConstraintRulesIndex = 0;
+	int32 PostprocessingConstraintRulesIndex = 0;
+
+	if (ExternalForces)
 	{
-		PostprocessingConstraintRuleOffset = Evolution->AddConstraintPostprocessingsRange(NumPostprocessingConstraintRules, false);
+		ExternalForceRules[ExternalForceRulesIndex++] =
+			[this](Softs::FSolverParticlesRange& Particles, const Softs::FSolverReal Dt)
+		{
+			ExternalForces->Apply(Particles, Dt);
+		};
 	}
 
-	TFunction<void(Softs::FSolverParticles&, const Softs::FSolverReal)>* const ConstraintInits = Evolution->ConstraintInits().GetData() + ConstraintInitOffset;
-	TFunction<void(Softs::FSolverParticles&, const Softs::FSolverReal)>* const ConstraintRules = Evolution->ConstraintRules().GetData() + ConstraintRuleOffset;
-	TFunction<void(Softs::FSolverParticles&, const Softs::FSolverReal)>* const PostCollisionConstraintRules = Evolution->PostCollisionConstraintRules().GetData() + PostCollisionConstraintRuleOffset;
-	TFunction<void(Softs::FSolverParticles&, const Softs::FSolverReal)>* const PostprocessingConstraintRules = Evolution->ConstraintPostprocessings().GetData() + PostprocessingConstraintRuleOffset;
-
-	int32 ConstraintInitIndex = 0;
-	int32 ConstraintRuleIndex = 0;
-	int32 PostCollisionConstraintRuleIndex = 0;
-	int32 PostprocessingConstraintRuleIndex = 0;
-	#if !UE_BUILD_SHIPPING	
-	if (bDisplayResidual || bWriteFinalResiduals)
+	if (VelocityAndPressureField)
 	{
-		GSMasterConstraint = MakeShared<Chaos::Softs::FGaussSeidelMasterConstraint<Softs::FSolverReal, Softs::FSolverParticles>>(Evolution->Particles(), false, false, (Softs::FSolverReal)1.2, 100);
+		PreSubstepParallelInits[PreSubstepInitsIndex++] =
+			[this](const Softs::FSolverParticlesRange& Particles, const Softs::FSolverReal Dt)
+		{
+			VelocityAndPressureField->UpdateForces(Particles, Dt);
+		};
 
-		ConstraintInits[ConstraintInitIndex++] =
-			[this](Softs::FSolverParticles& InParticles, const Softs::FSolverReal Dt)
+		ExternalForceRules[ExternalForceRulesIndex++] =
+			[this](Softs::FSolverParticlesRange& Particles, const Softs::FSolverReal Dt)
 		{
-			this->GSMasterConstraint->Init(Dt, InParticles);
+			VelocityAndPressureField->Apply(Particles, Dt);
 		};
-		if (bDisplayResidual)
+	}
+
+	if (PerSolverField)
+	{
+		ExternalForceRules[ExternalForceRulesIndex++] =
+			[this](Softs::FSolverParticlesRange& Particles, const Softs::FSolverReal Dt)
 		{
+			const TArray<FVector>& LinearVelocities = PerSolverField->GetOutputResults(EFieldCommandOutputType::LinearVelocity);
+			const TArray<FVector>& LinearForces = PerSolverField->GetOutputResults(EFieldCommandOutputType::LinearForce);
+			const FVector* const LinearVelocitiesView = LinearVelocities.IsEmpty() ? nullptr :
+				Particles.GetConstArrayView(LinearVelocities).GetData();
+			const FVector* const LinearForcesView = LinearForces.IsEmpty() ? nullptr :
+				Particles.GetConstArrayView(LinearForces).GetData();
+			if (!LinearVelocitiesView && !LinearForcesView)
+			{
+				return;
+			}
 
-			ConstraintRules[ConstraintRuleIndex++] =
-				[this](Softs::FSolverParticles& InParticles, const Softs::FSolverReal Dt)
+			Softs::FSolverVec3* const Acceleration = Particles.GetAcceleration().GetData();
+			const Softs::FSolverReal* const InvM = Particles.GetInvM().GetData();
+			for (int32 Index = 0; Index < Particles.GetRangeSize(); ++Index)
 			{
-				if (this->GSMasterConstraint->DebugResidual && this->GSMasterConstraint->PassedIters < MaxResidualIters)
+				if (InvM[Index] != (Softs::FSolverReal)0.)
 				{
-					UE_LOG(LogTemp, Warning, TEXT("Current Iteration As From Evolution: %d"), this->Evolution->GetIterations());
-					this->GSMasterConstraint->ComputeNewtonResiduals(InParticles, Dt, bWriteResidual2File);
-					this->GSMasterConstraint->PassedIters += 1;
+					if (LinearForcesView)
+					{
+						Acceleration[Index] += Softs::FSolverVec3(LinearForcesView[Index]) * InvM[Index];
+					}
+					if (LinearVelocitiesView)
+					{
+						Acceleration[Index] += Softs::FSolverVec3(LinearVelocitiesView[Index]) / Dt;
+					}
 				}
-			};
-		}
-
-		if (XStretchBiasConstraints)
+			}
+		};
+	}
+	if (XStretchBiasConstraints)
+	{
+		ConstraintParallelInits[ConstraintInitsIndex++] =
+			[this](const Softs::FSolverParticlesRange& /*Particles*/, const Softs::FSolverReal Dt)
 		{
-			TArray<TArray<int32>> IncidentElements, IncidentElementsLocal;
-			GSMasterConstraint->AddStaticConstraints(XStretchBiasConstraints->GetConstraintsArray(), IncidentElements, IncidentElementsLocal);
-			int32 StaticIndex = GSMasterConstraint->AddStaticConstraintResidualAndHessianRange(1);
-
-			GSMasterConstraint->StaticConstraintResidualAndHessian()[StaticIndex] = [this](const Softs::FSolverParticles& Particles, const int32 ElementIndex, const int32 ElementIndexLocal, const Softs::FSolverReal Dt, TVec3<Softs::FSolverReal>& ParticleResidual, Chaos::PMatrix<Softs::FSolverReal, 3, 3>& ParticleHessian)
-			{
-				this->XStretchBiasConstraints->AddStretchBiasElementResidualAndHessian(Particles, ElementIndex, ElementIndexLocal, Dt, ParticleResidual, ParticleHessian);
-			};
-		}
+			XStretchBiasConstraints->Init();
+			XStretchBiasConstraints->ApplyProperties(Dt, Evolution->GetIterations());
+		};
 
-		if (XAnisoBendingElementConstraints)
+		PerIterationConstraintRules[ConstraintRuleIndex++] =
+			[this](Softs::FSolverParticlesRange& Particles, const Softs::FSolverReal Dt)
 		{
-			TArray<TArray<int32>> IncidentElements, IncidentElementsLocal;
-			GSMasterConstraint->AddStaticConstraints(XAnisoBendingElementConstraints->GetConstraintsArray(), IncidentElements, IncidentElementsLocal);
-			int32 StaticIndex = GSMasterConstraint->AddStaticConstraintResidualAndHessianRange(1);
+			XStretchBiasConstraints->Apply(Particles, Dt);
+		};
+	}
+	if (XEdgeConstraints)
+	{
+		ConstraintParallelInits[ConstraintInitsIndex++] =
+			[this](const Softs::FSolverParticlesRange& /*Particles*/, const Softs::FSolverReal Dt)
+		{
+			XEdgeConstraints->Init();
+			XEdgeConstraints->ApplyProperties(Dt, Evolution->GetIterations());
+		};
+
+		PerIterationConstraintRules[ConstraintRuleIndex++] =
+			[this](Softs::FSolverParticlesRange& Particles, const Softs::FSolverReal Dt)
+		{
+			XEdgeConstraints->Apply(Particles, Dt);
+		};
+	}
+	if (EdgeConstraints)
+	{
+		ConstraintParallelInits[ConstraintInitsIndex++] =
+			[this](const Softs::FSolverParticlesRange& /*Particles*/, const Softs::FSolverReal Dt)
+		{
+			EdgeConstraints->ApplyProperties(Dt, Evolution->GetIterations());
+		};
+		PerIterationConstraintRules[ConstraintRuleIndex++] =
+			[this](Softs::FSolverParticlesRange& Particles, const Softs::FSolverReal Dt)
+		{
+			EdgeConstraints->Apply(Particles, Dt);
+		};
+	}
+	if (XBendingConstraints)
+	{
+		ConstraintParallelInits[ConstraintInitsIndex++] =
+			[this](const Softs::FSolverParticlesRange& /*Particles*/, const Softs::FSolverReal Dt)
+		{
+			XBendingConstraints->Init();
+			XBendingConstraints->ApplyProperties(Dt, Evolution->GetIterations());
+		};
+		PerIterationConstraintRules[ConstraintRuleIndex++] =
+			[this](Softs::FSolverParticlesRange& Particles, const Softs::FSolverReal Dt)
+		{
+			XBendingConstraints->Apply(Particles, Dt);
+		};
+	}
+	if (BendingConstraints)
+	{
+		ConstraintParallelInits[ConstraintInitsIndex++] =
+			[this](const Softs::FSolverParticlesRange& /*Particles*/, const Softs::FSolverReal Dt)
+		{
+			BendingConstraints->ApplyProperties(Dt, Evolution->GetIterations());
+		};
+		PerIterationConstraintRules[ConstraintRuleIndex++] =
+			[this](Softs::FSolverParticlesRange& Particles, const Softs::FSolverReal Dt)
+		{
+			BendingConstraints->Apply(Particles, Dt);
+		};
+	}
+	if (BendingElementConstraints)
+	{
+		ConstraintParallelInits[ConstraintInitsIndex++] =
+			[this](const Softs::FSolverParticlesRange& Particles, const Softs::FSolverReal Dt)
+		{
+			BendingElementConstraints->Init(Particles);
+			BendingElementConstraints->ApplyProperties(Dt, Evolution->GetIterations());
+		};
+		PerIterationConstraintRules[ConstraintRuleIndex++] =
+			[this](Softs::FSolverParticlesRange& Particles, const Softs::FSolverReal Dt)
+		{
+			BendingElementConstraints->Apply(Particles, Dt);
+		};
+	}
+	if (XBendingElementConstraints)
+	{
+		ConstraintParallelInits[ConstraintInitsIndex++] =
+			[this](const Softs::FSolverParticlesRange& Particles, const Softs::FSolverReal Dt)
+		{
+			XBendingElementConstraints->Init(Particles);
+			XBendingElementConstraints->ApplyProperties(Dt, Evolution->GetIterations());
+		};
+		PerIterationConstraintRules[ConstraintRuleIndex++] =
+			[this](Softs::FSolverParticlesRange& Particles, const Softs::FSolverReal Dt)
+		{
+			XBendingElementConstraints->Apply(Particles, Dt);
+		};
+	}
+	if (XAnisoBendingElementConstraints)
+	{
+		ConstraintParallelInits[ConstraintInitsIndex++] =
+			[this](const Softs::FSolverParticlesRange& Particles, const Softs::FSolverReal Dt)
+		{
+			XAnisoBendingElementConstraints->Init(Particles);
+			XAnisoBendingElementConstraints->ApplyProperties(Dt, Evolution->GetIterations());
+		};
+		PerIterationConstraintRules[ConstraintRuleIndex++] =
+			[this](Softs::FSolverParticlesRange& Particles, const Softs::FSolverReal Dt)
+		{
+			XAnisoBendingElementConstraints->Apply(Particles, Dt);
+		};
+	}
+	if (XAreaConstraints)
+	{
+		ConstraintParallelInits[ConstraintInitsIndex++] =
+			[this](const Softs::FSolverParticlesRange& /*Particles*/, const Softs::FSolverReal Dt)
+		{
+			XAreaConstraints->Init();
+			XAreaConstraints->ApplyProperties(Dt, Evolution->GetIterations());
+		};
+		PerIterationConstraintRules[ConstraintRuleIndex++] =
+			[this](Softs::FSolverParticlesRange& Particles, const Softs::FSolverReal Dt)
+		{
+			XAreaConstraints->Apply(Particles, Dt);
+		};
+	}
+	if (AreaConstraints)
+	{
+		ConstraintParallelInits[ConstraintInitsIndex++] =
+			[this](const Softs::FSolverParticlesRange& /*Particles*/, const Softs::FSolverReal Dt)
+		{
+			AreaConstraints->ApplyProperties(Dt, Evolution->GetIterations());
+		};
+		PerIterationConstraintRules[ConstraintRuleIndex++] =
+			[this](Softs::FSolverParticlesRange& Particles, const Softs::FSolverReal Dt)
+		{
+			AreaConstraints->Apply(Particles, Dt);
+		};
+	}
+	if (MaximumDistanceConstraints)
+	{
+		PerIterationConstraintRules[ConstraintRuleIndex++] =
+			[this](Softs::FSolverParticlesRange& Particles, const Softs::FSolverReal Dt)
+		{
+			MaximumDistanceConstraints->Apply(Particles, Dt);
+		};
+	}
+	if (BackstopConstraints)
+	{
+		PerIterationConstraintRules[ConstraintRuleIndex++] =
+			[this](Softs::FSolverParticlesRange& Particles, const Softs::FSolverReal Dt)
+		{
+			BackstopConstraints->Apply(Particles, Dt);
+		};
+	}
+	if (AnimDriveConstraints)
+	{
+		ConstraintParallelInits[ConstraintInitsIndex++] =
+			[this](const Softs::FSolverParticlesRange& /*Particles*/, const Softs::FSolverReal Dt)
+		{
+			AnimDriveConstraints->ApplyProperties(Dt, Evolution->GetIterations());
+		};
+
+		PerIterationConstraintRules[ConstraintRuleIndex++] =
+			[this](Softs::FSolverParticlesRange& Particles, const Softs::FSolverReal Dt)
+		{
+			AnimDriveConstraints->Apply(Particles, Dt);
+		};
+	}
+	if (CollisionConstraint)
+	{
+		PerIterationCollisionConstraintRules[CollisionConstraintRulesIndex++] =
+			[this](Softs::FSolverParticlesRange& Particles, const Softs::FSolverReal Dt, const TArray<Softs::FSolverCollisionParticlesRange>& CollisionParticles)
+		{
+			CollisionConstraint->Apply(Particles, Dt, CollisionParticles);
+		};
+	}
+	if (SelfCollisionInit && SelfCollisionConstraints)
+	{
+		ConstraintParallelInits[ConstraintInitsIndex++] =
+			[this](const Softs::FSolverParticlesRange& Particles, const Softs::FSolverReal Dt)
+		{
+			// Thickness * 2 to account for collision radius for both particles
+			SelfCollisionInit->Init(Particles, SelfCollisionConstraints->GetThickness() * (Softs::FSolverReal)2.f);
+			SelfCollisionConstraints->Init(Particles, SelfCollisionInit->GetSpatialHash(), SelfCollisionInit->GetVertexGIAColors(), SelfCollisionInit->GetTriangleGIAColors());
+		};
+
+		PerIterationPostCollisionsConstraintRules[PostCollisionConstraintRulesIndex++] =
+			[this](Softs::FSolverParticlesRange& Particles, const Softs::FSolverReal Dt)
+		{
+			SelfCollisionConstraints->Apply(Particles, Dt);
+		};
+	}
+	if (SelfCollisionSphereConstraints)
+	{
+		ConstraintParallelInits[ConstraintInitsIndex++] =
+			[this](const Softs::FSolverParticlesRange& Particles, const Softs::FSolverReal /*Dt*/)
+		{
+			SelfCollisionSphereConstraints->Init(Particles);
+		};
+
+		PerIterationPostCollisionsConstraintRules[PostCollisionConstraintRulesIndex++] =
+			[this](Softs::FSolverParticlesRange& Particles, const Softs::FSolverReal Dt)
+		{
+			SelfCollisionSphereConstraints->Apply(Particles, Dt);
+		};
+	}
+
+	if (SelfCollisionInit && SelfIntersectionConstraints)
+	{
+		PreSubstepConstraintRules[PreSubstepConstraintRulesIndex++] =
+			[this](Softs::FSolverParticlesRange& Particles, const Softs::FSolverReal Dt)
+		{
+			SelfIntersectionConstraints->Apply(Particles, SelfCollisionInit->GetContourMinimizationIntersections(), Dt);
+		};
+
+		PostSubstepConstraintRules[PostprocessingConstraintRulesIndex++] =
+			[this](Softs::FSolverParticlesRange& Particles, const Softs::FSolverReal Dt)
+		{
+			const int32 NumContourIterations = SelfCollisionInit->GetNumContourMinimizationPostSteps();
+			for (int32 Iter = 0; Iter < NumContourIterations; ++Iter)
+			{
+				SelfCollisionInit->PostStepInit(Particles);
+				SelfIntersectionConstraints->Apply(Particles, SelfCollisionInit->GetPostStepContourMinimizationIntersections(), Dt);
+			}
+		};
+	}
+	if (LongRangeConstraints)
+	{
+		PreSubstepConstraintRules[PreSubstepConstraintRulesIndex++] =
+			[this](Softs::FSolverParticlesRange& Particles, const Softs::FSolverReal Dt)
+		{
+			// Only doing one iteration.
+			constexpr int32 NumLRAIterations = 1;
+			LongRangeConstraints->ApplyProperties(Dt, NumLRAIterations);
+			LongRangeConstraints->Apply(Particles, Dt);  // Run the LRA constraint only once per timestep
+		};
+	}
+
+	check(PreSubstepInitsIndex == NumPreSubstepInits);
+	check(ExternalForceRulesIndex == NumExternalForceRules);
+	check(ConstraintInitsIndex == NumConstraintInits);
+	check(PreSubstepConstraintRulesIndex == NumPreSubstepConstraintRules);
+	check(ConstraintRuleIndex == NumConstraintRules);
+	check(CollisionConstraintRulesIndex == NumCollisionConstraintRules);
+	check(PostCollisionConstraintRulesIndex == NumPostCollisionConstraintRules);
+	check(PostprocessingConstraintRulesIndex == NumPostprocessingConstraintRules);
+}
+
+
+void FClothConstraints::CreatePBDRules()
+{
+	check(PBDEvolution);
+	check(ConstraintInitOffset == INDEX_NONE)
+#if !UE_BUILD_SHIPPING
+	if (bDisplayResidual)
+	{
+		NumConstraintRules++;
+	}
+	if (bWriteFinalResiduals || bDisplayResidual)
+	{
+		NumConstraintInits++;
+	}
+#endif
+	if (NumConstraintInits)
+	{
+		ConstraintInitOffset = PBDEvolution->AddConstraintInitRange(NumConstraintInits, false);
+	}
+	check(ConstraintRuleOffset == INDEX_NONE)
+
+	if (NumConstraintRules)
+	{
+		ConstraintRuleOffset = PBDEvolution->AddConstraintRuleRange(NumConstraintRules, false);
+	}
+	check(PostCollisionConstraintRuleOffset == INDEX_NONE);
+	if (NumPostCollisionConstraintRules)
+	{
+		PostCollisionConstraintRuleOffset = PBDEvolution->AddPostCollisionConstraintRuleRange(NumPostCollisionConstraintRules, false);
+	}
+	check(PostprocessingConstraintRuleOffset == INDEX_NONE);
+#if !UE_BUILD_SHIPPING	
+	if (bWriteFinalResiduals)
+	{
+		NumPostprocessingConstraintRules += 1;
+	}
+#endif
+	if (NumPostprocessingConstraintRules)
+	{
+		PostprocessingConstraintRuleOffset = PBDEvolution->AddConstraintPostprocessingsRange(NumPostprocessingConstraintRules, false);
+	}
+
+	TFunction<void(Softs::FSolverParticles&, const Softs::FSolverReal)>* const ConstraintInits = PBDEvolution->ConstraintInits().GetData() + ConstraintInitOffset;
+	TFunction<void(Softs::FSolverParticles&, const Softs::FSolverReal)>* const ConstraintRules = PBDEvolution->ConstraintRules().GetData() + ConstraintRuleOffset;
+	TFunction<void(Softs::FSolverParticles&, const Softs::FSolverReal)>* const PostCollisionConstraintRules = PBDEvolution->PostCollisionConstraintRules().GetData() + PostCollisionConstraintRuleOffset;
+	TFunction<void(Softs::FSolverParticles&, const Softs::FSolverReal)>* const PostprocessingConstraintRules = PBDEvolution->ConstraintPostprocessings().GetData() + PostprocessingConstraintRuleOffset;
+
+	int32 ConstraintInitIndex = 0;
+	int32 ConstraintRuleIndex = 0;
+	int32 PostCollisionConstraintRuleIndex = 0;
+	int32 PostprocessingConstraintRuleIndex = 0;
+	#if !UE_BUILD_SHIPPING	
+	if (bDisplayResidual || bWriteFinalResiduals)
+	{
+		GSMasterConstraint = MakeShared<Chaos::Softs::FGaussSeidelMasterConstraint<Softs::FSolverReal, Softs::FSolverParticles>>(PBDEvolution->Particles(), false, false, (Softs::FSolverReal)1.2, 100);
+
+		ConstraintInits[ConstraintInitIndex++] =
+			[this](Softs::FSolverParticles& InParticles, const Softs::FSolverReal Dt)
+		{
+			this->GSMasterConstraint->Init(Dt, InParticles);
+		};
+		if (bDisplayResidual)
+		{
+
+			ConstraintRules[ConstraintRuleIndex++] =
+				[this](Softs::FSolverParticles& InParticles, const Softs::FSolverReal Dt)
+			{
+				if (this->GSMasterConstraint->DebugResidual && this->GSMasterConstraint->PassedIters < MaxResidualIters)
+				{
+					UE_LOG(LogTemp, Warning, TEXT("Current Iteration As From Evolution: %d"), this->Evolution->GetIterations());
+					this->GSMasterConstraint->ComputeNewtonResiduals(InParticles, Dt, bWriteResidual2File);
+					this->GSMasterConstraint->PassedIters += 1;
+				}
+			};
+		}
+
+		if (XStretchBiasConstraints)
+		{
+			TArray<TArray<int32>> IncidentElements, IncidentElementsLocal;
+			GSMasterConstraint->AddStaticConstraints(XStretchBiasConstraints->GetConstraintsArray(), IncidentElements, IncidentElementsLocal);
+			int32 StaticIndex = GSMasterConstraint->AddStaticConstraintResidualAndHessianRange(1);
+
+			GSMasterConstraint->StaticConstraintResidualAndHessian()[StaticIndex] = [this](const Softs::FSolverParticles& Particles, const int32 ElementIndex, const int32 ElementIndexLocal, const Softs::FSolverReal Dt, TVec3<Softs::FSolverReal>& ParticleResidual, Chaos::PMatrix<Softs::FSolverReal, 3, 3>& ParticleHessian)
+			{
+				this->XStretchBiasConstraints->AddStretchBiasElementResidualAndHessian(Particles, ElementIndex, ElementIndexLocal, Dt, ParticleResidual, ParticleHessian);
+			};
+		}
+
+		if (XAnisoBendingElementConstraints)
+		{
+			TArray<TArray<int32>> IncidentElements, IncidentElementsLocal;
+			GSMasterConstraint->AddStaticConstraints(XAnisoBendingElementConstraints->GetConstraintsArray(), IncidentElements, IncidentElementsLocal);
+			int32 StaticIndex = GSMasterConstraint->AddStaticConstraintResidualAndHessianRange(1);
 
 			GSMasterConstraint->StaticConstraintResidualAndHessian()[StaticIndex] = [this](const Softs::FSolverParticles& Particles, const int32 ElementIndex, const int32 ElementIndexLocal, const Softs::FSolverReal Dt, TVec3<Softs::FSolverReal>& ParticleResidual, Chaos::PMatrix<Softs::FSolverReal, 3, 3>& ParticleHessian)
 			{
@@ -1015,7 +1627,7 @@ void FClothConstraints::CreateRules()
 			};
 		}
 
-		GSMasterConstraint->InitStaticColor(Evolution->Particles());
+		GSMasterConstraint->InitStaticColor(PBDEvolution->Particles());
 
 		if (bWriteFinalResiduals)
 		{
@@ -1037,7 +1649,7 @@ void FClothConstraints::CreateRules()
 			[this](Softs::FSolverParticles& /*Particles*/, const Softs::FSolverReal Dt)
 		{
 			XStretchBiasConstraints->Init();
-			XStretchBiasConstraints->ApplyProperties(Dt, Evolution->GetIterations());
+			XStretchBiasConstraints->ApplyProperties(Dt, PBDEvolution->GetIterations());
 		};
 
 		ConstraintRules[ConstraintRuleIndex++] =
@@ -1052,7 +1664,7 @@ void FClothConstraints::CreateRules()
 			[this](Softs::FSolverParticles& /*Particles*/, const Softs::FSolverReal Dt)
 			{
 				XEdgeConstraints->Init();
-				XEdgeConstraints->ApplyProperties(Dt, Evolution->GetIterations());
+				XEdgeConstraints->ApplyProperties(Dt, PBDEvolution->GetIterations());
 			};
 
 		ConstraintRules[ConstraintRuleIndex++] = 
@@ -1061,27 +1673,12 @@ void FClothConstraints::CreateRules()
 				XEdgeConstraints->Apply(Particles, Dt);
 			};
 	}
-	if (XEdgeConstraints_Deprecated)  // TODO: Remove for 5.4
-	{
-		ConstraintInits[ConstraintInitIndex++] =
-			[this](Softs::FSolverParticles& /*Particles*/, const Softs::FSolverReal Dt)
-		{
-			XEdgeConstraints_Deprecated->Init();
-			XEdgeConstraints_Deprecated->ApplyProperties(Dt, Evolution->GetIterations());
-		};
-
-		ConstraintRules[ConstraintRuleIndex++] =
-			[this](Softs::FSolverParticles& Particles, const Softs::FSolverReal Dt)
-		{
-			XEdgeConstraints_Deprecated->Apply(Particles, Dt);
-		};
-	}
 	if (EdgeConstraints)
 	{
 		ConstraintInits[ConstraintInitIndex++] =
 			[this](Softs::FSolverParticles& /*Particles*/, const Softs::FSolverReal Dt)
 			{
-				EdgeConstraints->ApplyProperties(Dt, Evolution->GetIterations());
+				EdgeConstraints->ApplyProperties(Dt, PBDEvolution->GetIterations());
 			};
 		ConstraintRules[ConstraintRuleIndex++] =
 			[this](Softs::FSolverParticles& Particles, const Softs::FSolverReal Dt)
@@ -1089,26 +1686,13 @@ void FClothConstraints::CreateRules()
 				EdgeConstraints->Apply(Particles, Dt);
 			};
 	}
-	if (EdgeConstraints_Deprecated)  // TODO: Remove for 5.4
-	{
-		ConstraintInits[ConstraintInitIndex++] =
-			[this](Softs::FSolverParticles& /*Particles*/, const Softs::FSolverReal Dt)
-		{
-			EdgeConstraints_Deprecated->ApplyProperties(Dt, Evolution->GetIterations());
-		};
-		ConstraintRules[ConstraintRuleIndex++] =
-			[this](Softs::FSolverParticles& Particles, const Softs::FSolverReal Dt)
-		{
-			EdgeConstraints_Deprecated->Apply(Particles, Dt);
-		};
-	}
 	if (XBendingConstraints)
 	{
 		ConstraintInits[ConstraintInitIndex++] =
 			[this](Softs::FSolverParticles& /*Particles*/, const Softs::FSolverReal Dt)
 			{
 				XBendingConstraints->Init();
-				XBendingConstraints->ApplyProperties(Dt, Evolution->GetIterations());
+				XBendingConstraints->ApplyProperties(Dt, PBDEvolution->GetIterations());
 			};
 		ConstraintRules[ConstraintRuleIndex++] =
 			[this](Softs::FSolverParticles& Particles, const Softs::FSolverReal Dt)
@@ -1116,26 +1700,12 @@ void FClothConstraints::CreateRules()
 				XBendingConstraints->Apply(Particles, Dt);
 			};
 	}
-	if (XBendingConstraints_Deprecated)  // TODO: Remove for 5.4
-	{
-		ConstraintInits[ConstraintInitIndex++] =
-			[this](Softs::FSolverParticles& /*Particles*/, const Softs::FSolverReal Dt)
-		{
-			XBendingConstraints_Deprecated->Init();
-			XBendingConstraints_Deprecated->ApplyProperties(Dt, Evolution->GetIterations());
-		};
-		ConstraintRules[ConstraintRuleIndex++] =
-			[this](Softs::FSolverParticles& Particles, const Softs::FSolverReal Dt)
-		{
-			XBendingConstraints_Deprecated->Apply(Particles, Dt);
-		};
-	}
 	if (BendingConstraints)
 	{
 		ConstraintInits[ConstraintInitIndex++] =
 			[this](Softs::FSolverParticles& /*Particles*/, const Softs::FSolverReal Dt)
 			{
-				BendingConstraints->ApplyProperties(Dt, Evolution->GetIterations());
+				BendingConstraints->ApplyProperties(Dt, PBDEvolution->GetIterations());
 			};
 		ConstraintRules[ConstraintRuleIndex++] =
 			[this](Softs::FSolverParticles& Particles, const Softs::FSolverReal Dt)
@@ -1143,26 +1713,13 @@ void FClothConstraints::CreateRules()
 				BendingConstraints->Apply(Particles, Dt);
 			};
 	}
-	if (BendingConstraints_Deprecated)  // TODO: Remove for 5.4
-	{
-		ConstraintInits[ConstraintInitIndex++] =
-			[this](Softs::FSolverParticles& /*Particles*/, const Softs::FSolverReal Dt)
-		{
-			BendingConstraints_Deprecated->ApplyProperties(Dt, Evolution->GetIterations());
-		};
-		ConstraintRules[ConstraintRuleIndex++] =
-			[this](Softs::FSolverParticles& Particles, const Softs::FSolverReal Dt)
-		{
-			BendingConstraints_Deprecated->Apply(Particles, Dt);
-		};
-	}
 	if (BendingElementConstraints)
 	{
 		ConstraintInits[ConstraintInitIndex++] =
 			[this](Softs::FSolverParticles& Particles, const Softs::FSolverReal Dt)
 			{
 				BendingElementConstraints->Init(Particles);
-				BendingElementConstraints->ApplyProperties(Dt, Evolution->GetIterations());
+				BendingElementConstraints->ApplyProperties(Dt, PBDEvolution->GetIterations());
 			};
 		ConstraintRules[ConstraintRuleIndex++] =
 			[this](Softs::FSolverParticles& Particles, const Softs::FSolverReal Dt)
@@ -1176,7 +1733,7 @@ void FClothConstraints::CreateRules()
 			[this](Softs::FSolverParticles& Particles, const Softs::FSolverReal Dt)
 		{
 			XBendingElementConstraints->Init(Particles);
-			XBendingElementConstraints->ApplyProperties(Dt, Evolution->GetIterations());
+			XBendingElementConstraints->ApplyProperties(Dt, PBDEvolution->GetIterations());
 		};
 		ConstraintRules[ConstraintRuleIndex++] =
 			[this](Softs::FSolverParticles& Particles, const Softs::FSolverReal Dt)
@@ -1190,7 +1747,7 @@ void FClothConstraints::CreateRules()
 			[this](Softs::FSolverParticles& Particles, const Softs::FSolverReal Dt)
 		{
 			XAnisoBendingElementConstraints->Init(Particles);
-			XAnisoBendingElementConstraints->ApplyProperties(Dt, Evolution->GetIterations());
+			XAnisoBendingElementConstraints->ApplyProperties(Dt, PBDEvolution->GetIterations());
 		};
 		ConstraintRules[ConstraintRuleIndex++] =
 			[this](Softs::FSolverParticles& Particles, const Softs::FSolverReal Dt)
@@ -1204,7 +1761,7 @@ void FClothConstraints::CreateRules()
 			[this](Softs::FSolverParticles& /*Particles*/, const Softs::FSolverReal Dt)
 			{
 				XAreaConstraints->Init();
-				XAreaConstraints->ApplyProperties(Dt, Evolution->GetIterations());
+				XAreaConstraints->ApplyProperties(Dt, PBDEvolution->GetIterations());
 			};
 		ConstraintRules[ConstraintRuleIndex++] =
 			[this](Softs::FSolverParticles& Particles, const Softs::FSolverReal Dt)
@@ -1212,26 +1769,12 @@ void FClothConstraints::CreateRules()
 				XAreaConstraints->Apply(Particles, Dt);
 			};
 	}
-	if (XAreaConstraints_Deprecated)  // TODO: Remove for 5.4
-	{
-		ConstraintInits[ConstraintInitIndex++] =
-			[this](Softs::FSolverParticles& /*Particles*/, const Softs::FSolverReal Dt)
-		{
-			XAreaConstraints_Deprecated->Init();
-			XAreaConstraints_Deprecated->ApplyProperties(Dt, Evolution->GetIterations());
-		};
-		ConstraintRules[ConstraintRuleIndex++] =
-			[this](Softs::FSolverParticles& Particles, const Softs::FSolverReal Dt)
-		{
-			XAreaConstraints_Deprecated->Apply(Particles, Dt);
-		};
-	}
 	if (AreaConstraints)
 	{
 		ConstraintInits[ConstraintInitIndex++] =
 			[this](Softs::FSolverParticles& /*Particles*/, const Softs::FSolverReal Dt)
 			{
-				AreaConstraints->ApplyProperties(Dt, Evolution->GetIterations());
+				AreaConstraints->ApplyProperties(Dt, PBDEvolution->GetIterations());
 			};
 		ConstraintRules[ConstraintRuleIndex++] =
 			[this](Softs::FSolverParticles& Particles, const Softs::FSolverReal Dt)
@@ -1239,35 +1782,6 @@ void FClothConstraints::CreateRules()
 				AreaConstraints->Apply(Particles, Dt);
 			};
 	}
-	if (AreaConstraints_Deprecated)  // TODO: Remove for 5.4
-	{
-		ConstraintInits[ConstraintInitIndex++] =
-			[this](Softs::FSolverParticles& /*Particles*/, const Softs::FSolverReal Dt)
-		{
-			AreaConstraints_Deprecated->ApplyProperties(Dt, Evolution->GetIterations());
-		};
-		ConstraintRules[ConstraintRuleIndex++] =
-			[this](Softs::FSolverParticles& Particles, const Softs::FSolverReal Dt)
-		{
-			AreaConstraints_Deprecated->Apply(Particles, Dt);
-		};
-	}
-	if (ThinShellVolumeConstraints_Deprecated)  // TODO: Remove for 5.4
-	{
-		ConstraintRules[ConstraintRuleIndex++] =
-			[this](Softs::FSolverParticles& Particles, const Softs::FSolverReal Dt)
-		{
-			ThinShellVolumeConstraints_Deprecated->Apply(Particles, Dt);
-		};
-	}
-	if (VolumeConstraints_Deprecated)  // TODO: Remove for 5.4
-	{
-		ConstraintRules[ConstraintRuleIndex++] =
-			[this](Softs::FSolverParticles& Particles, const Softs::FSolverReal Dt)
-		{
-			VolumeConstraints_Deprecated->Apply(Particles, Dt);
-		};
-	}
 	if (MaximumDistanceConstraints)
 	{
 		ConstraintRules[ConstraintRuleIndex++] =
@@ -1289,7 +1803,7 @@ void FClothConstraints::CreateRules()
 		ConstraintInits[ConstraintInitIndex++] =
 			[this](Softs::FSolverParticles& /*Particles*/, const Softs::FSolverReal Dt)
 			{
-				AnimDriveConstraints->ApplyProperties(Dt, Evolution->GetIterations());
+				AnimDriveConstraints->ApplyProperties(Dt, PBDEvolution->GetIterations());
 			};
 
 		ConstraintRules[ConstraintRuleIndex++] =
@@ -1298,12 +1812,12 @@ void FClothConstraints::CreateRules()
 				AnimDriveConstraints->Apply(Particles, Dt);
 			};
 	}
-	if (ShapeConstraints)
+	if (ShapeConstraints_Deprecated) // TODO: Remove 5.6
 	{
 		ConstraintRules[ConstraintRuleIndex++] =
 			[this](Softs::FSolverParticles& Particles, const Softs::FSolverReal Dt)
 			{
-				ShapeConstraints->Apply(Particles, Dt);
+				ShapeConstraints_Deprecated->Apply(Particles, Dt);
 			};
 	}
 
@@ -1382,338 +1896,18 @@ void FClothConstraints::CreateRules()
 	check(PostCollisionConstraintRuleIndex == NumPostCollisionConstraintRules);
 }
 
-void FClothConstraints::SetEdgeConstraints(const TArray<TVec3<int32>>& SurfaceElements, const TConstArrayView<FRealSingle>& StiffnessMultipliers, bool bUseXPBDConstraints)
-{
-	check(Evolution);
-
-	if (bUseXPBDConstraints)
-	{
-		const TConstArrayView<FRealSingle> DampingMultipliers;
-
-		XEdgeConstraints_Deprecated = MakeShared<Softs::FXPBDSpringConstraints>(
-			Evolution->Particles(),
-			ParticleOffset, NumParticles,
-			SurfaceElements,
-			StiffnessMultipliers,
-			DampingMultipliers,
-			/*InStiffness =*/ Softs::FSolverVec2::UnitVector,
-			/*bTrimKinematicConstraints =*/ true);
-	}
-	else
-	{
-		EdgeConstraints_Deprecated = MakeShared<Softs::FPBDSpringConstraints>(
-			Evolution->Particles(),
-			ParticleOffset,
-			NumParticles,
-			SurfaceElements,
-			StiffnessMultipliers,
-			/*InStiffness =*/ Softs::FSolverVec2::UnitVector,
-			/*bTrimKinematicConstraints =*/ true);
-	}
-	++NumConstraintInits;  // Uses init to update the property tables
-	++NumConstraintRules;
-}
-
-void FClothConstraints::SetXPBDEdgeConstraints(const TArray<TVec3<int32>>& SurfaceElements, const TConstArrayView<FRealSingle>& StiffnessMultipliers, const TConstArrayView<FRealSingle>& DampingRatioMultipliers)
-{
-	check(Evolution);
-	XEdgeConstraints_Deprecated = MakeShared<Softs::FXPBDSpringConstraints>(
-		Evolution->Particles(),
-		ParticleOffset, NumParticles,
-		SurfaceElements,
-		StiffnessMultipliers,
-		DampingRatioMultipliers,
-		/*InStiffness =*/ Softs::FSolverVec2::UnitVector,
-		/*InDampingRatio =*/ Softs::FSolverVec2::ZeroVector,
-		/*bTrimKinematicConstraints =*/ true);
-	++NumConstraintInits;  // Uses init to update the property tables
-	++NumConstraintRules;
-}
-
-void FClothConstraints::SetBendingConstraints(const TArray<TVec2<int32>>& Edges, const TConstArrayView<FRealSingle>& StiffnessMultipliers, bool bUseXPBDConstraints)
-{
-	check(Evolution);
-
-	if (bUseXPBDConstraints)
-	{
-		const TConstArrayView<FRealSingle> DampingMultipliers;
-
-		XBendingConstraints_Deprecated = MakeShared<Softs::FXPBDSpringConstraints>(
-			Evolution->Particles(),
-			ParticleOffset, NumParticles,
-			Edges,
-			StiffnessMultipliers,
-			DampingMultipliers,
-			/*InStiffness =*/ Softs::FSolverVec2::UnitVector,
-			/*bTrimKinematicConstraints =*/ true);
-	}
-	else
-	{
-		BendingConstraints_Deprecated = MakeShared<Softs::FPBDSpringConstraints>(
-			Evolution->Particles(),
-			ParticleOffset,
-			NumParticles,
-			Edges,
-			StiffnessMultipliers,
-			/*InStiffness =*/ Softs::FSolverVec2::UnitVector,
-			/*bTrimKinematicConstraints =*/ true);
-	}
-	++NumConstraintInits;  // Uses init to update the property tables
-	++NumConstraintRules;
-}
-
-void FClothConstraints::SetBendingConstraints(TArray<TVec4<int32>>&& BendingElements, const TConstArrayView<FRealSingle>& StiffnessMultipliers, const TConstArrayView<FRealSingle>& BucklingStiffnessMultipliers, bool bUseXPBDConstraints)
-{
-	check(Evolution);
-
-	if (bUseXPBDConstraints)
-	{
-		const TConstArrayView<FRealSingle> DampingMultipliers;
-
-		XBendingElementConstraints = MakeShared<Softs::FXPBDBendingConstraints>(
-			Evolution->Particles(),
-			ParticleOffset, NumParticles,
-			MoveTemp(BendingElements),
-			StiffnessMultipliers,
-			BucklingStiffnessMultipliers,
-			DampingMultipliers,
-			/*InStiffness =*/ Softs::FSolverVec2::UnitVector,
-			/*InBucklingRatio=*/ (Softs::FSolverReal)0.f,
-			/*InBucklingStiffness =*/ Softs::FSolverVec2::UnitVector,
-			/*InDampingRatio =*/ Softs::FSolverVec2::ZeroVector,
-			/*bTrimKinematicConstraints =*/ true);
-	}
-	else
-	{
-		BendingElementConstraints = MakeShared<Softs::FPBDBendingConstraints>(
-			Evolution->Particles(),
-			ParticleOffset, NumParticles,
-			MoveTemp(BendingElements),
-			StiffnessMultipliers,
-			BucklingStiffnessMultipliers,
-			/*InStiffness =*/ Softs::FSolverVec2::UnitVector,
-			/*InBucklingRatio=*/ (Softs::FSolverReal)0.f,
-			/*InBucklingStiffness =*/ Softs::FSolverVec2::UnitVector,
-			/*bTrimKinematicConstraints =*/ true);
-
-	}
-	++NumConstraintInits;  // Uses init to update the property tables
-	++NumConstraintRules;
-}
-
-void FClothConstraints::SetXPBDBendingConstraints(TArray<TVec4<int32>>&& BendingElements, const TConstArrayView<FRealSingle>& StiffnessMultipliers, const TConstArrayView<FRealSingle>& BucklingStiffnessMultipliers, const TConstArrayView<FRealSingle>& DampingRatioMultipliers)
-{
-	check(Evolution);
-	XBendingElementConstraints = MakeShared<Softs::FXPBDBendingConstraints>(
-		Evolution->Particles(),
-		ParticleOffset, NumParticles,
-		MoveTemp(BendingElements),
-		StiffnessMultipliers,
-		BucklingStiffnessMultipliers,
-		DampingRatioMultipliers,
-		/*InStiffness =*/ Softs::FSolverVec2::UnitVector,
-		/*InBucklingRatio=*/ (Softs::FSolverReal)0.f,
-		/*InBucklingStiffness =*/ Softs::FSolverVec2::UnitVector,
-		/*InDampingRatio =*/ Softs::FSolverVec2::ZeroVector,
-		/*bTrimKinematicConstraints =*/ true);
-	++NumConstraintInits;  // Uses init to update the property tables
-	++NumConstraintRules;
-}
-
-void FClothConstraints::SetBendingConstraints(TArray<TVec4<int32>>&& BendingElements, Softs::FSolverReal BendingStiffness)
-{
-	// Deprecated 5.1
-	check(Evolution);
-
-	BendingElementConstraints = MakeShared<Softs::FPBDBendingConstraints>(
-		Evolution->Particles(),
-		ParticleOffset,
-		NumParticles,
-		MoveTemp(BendingElements),
-		/*StiffnessMultipliers =*/ TConstArrayView<FRealSingle>(),
-		/*BucklingStiffnessMultipliers =*/ TConstArrayView<FRealSingle>(),
-		/*InStiffness =*/ BendingStiffness,
-		/*InBucklingRatio=*/ (Softs::FSolverReal)0.f,
-		/*InBucklingStiffness =*/ Softs::FSolverVec2::UnitVector,
-		/*bTrimKinematicConstraints =*/ true);
-
-	++NumConstraintInits;  // Uses init to update the property tables
-	++NumConstraintRules;
-}
-
-void FClothConstraints::SetAreaConstraints(const TArray<TVec3<int32>>& SurfaceElements, const TConstArrayView<FRealSingle>& StiffnessMultipliers, bool bUseXPBDConstraints)
-{
-	check(Evolution);
-
-	if (bUseXPBDConstraints)
-	{
-		XAreaConstraints_Deprecated = MakeShared<Softs::FXPBDAxialSpringConstraints>(
-			Evolution->Particles(),
-			ParticleOffset,
-			NumParticles,
-			SurfaceElements,
-			StiffnessMultipliers,
-			/*InStiffness =*/ Softs::FSolverVec2::UnitVector,
-			/*bTrimKinematicConstraints =*/ true);
-	}
-	else
-	{
-		AreaConstraints_Deprecated = MakeShared<Softs::FPBDAxialSpringConstraints>(
-			Evolution->Particles(),
-			ParticleOffset,
-			NumParticles,
-			SurfaceElements,
-			StiffnessMultipliers,
-			/*InStiffness =*/ Softs::FSolverVec2::UnitVector,
-			/*bTrimKinematicConstraints =*/ true);
-	}
-	++NumConstraintInits;  // Uses init to update the property tables
-	++NumConstraintRules;
-}
-
-void FClothConstraints::SetVolumeConstraints(const TArray<TVec2<int32>>& DoubleBendingEdges, Softs::FSolverReal VolumeStiffness)
-{
-	check(Evolution);
-
-	ThinShellVolumeConstraints_Deprecated = MakeShared<Softs::FPBDSpringConstraints>(
-		Evolution->Particles(),
-		ParticleOffset,
-		NumParticles,
-		DoubleBendingEdges,
-		TConstArrayView<FRealSingle>(),
-		VolumeStiffness,
-		/*bTrimKinematicConstraints =*/ true);
-	++NumConstraintRules;
-}
-
-void FClothConstraints::SetVolumeConstraints(TArray<TVec3<int32>>&& SurfaceElements, Softs::FSolverReal VolumeStiffness)
-{
-	check(Evolution);
-	check(VolumeStiffness > 0.f && VolumeStiffness <= 1.f);
-
-	VolumeConstraints_Deprecated = MakeShared<Softs::FPBDVolumeConstraint>(Evolution->Particles(), MoveTemp(SurfaceElements), VolumeStiffness);
-	++NumConstraintRules;
-}
-
-void FClothConstraints::SetLongRangeConstraints(
-	const TArray<TConstArrayView<TTuple<int32, int32, FRealSingle>>>& Tethers,
-	const TConstArrayView<FRealSingle>& TetherStiffnessMultipliers,
-	const TConstArrayView<FRealSingle>& TetherScaleMultipliers,
-	const Softs::FSolverVec2& TetherScale,
-	Softs::FSolverReal MeshScale)
-{
-	check(Evolution);
-	//  Now that we're only doing a single iteration of Long range constraints, and they're more of a fake constraint to jump start our initial guess, it's not clear that using XPBD makes sense here.
-	LongRangeConstraints = MakeShared<Softs::FPBDLongRangeConstraints>(
-		Evolution->Particles(),
-		ParticleOffset,
-		NumParticles,
-		Tethers,
-		TetherStiffnessMultipliers,
-		TetherScaleMultipliers,
-		/*InStiffness =*/ Softs::FSolverVec2::UnitVector,
-		TetherScale,
-		MeshScale);
-	++NumConstraintInits;  // Uses init to both update the property tables and apply the constraint
-}
-
-// Deprecated in 5.1
-void FClothConstraints::SetLongRangeConstraints(
-	const TArray<TConstArrayView<TTuple<int32,
-	int32, FRealSingle>>>& Tethers,
-	const TConstArrayView<FRealSingle>& TetherStiffnessMultipliers,
-	const TConstArrayView<FRealSingle>& TetherScaleMultipliers,
-	const Softs::FSolverVec2& TetherScale,
-	bool bUseXPBDConstraints,
-	Softs::FSolverReal MeshScale)
-{
-PRAGMA_DISABLE_DEPRECATION_WARNINGS
-	return SetLongRangeConstraints(Tethers, TetherStiffnessMultipliers, TetherScaleMultipliers, TetherScale, MeshScale);
-PRAGMA_ENABLE_DEPRECATION_WARNINGS
-}
-
-void FClothConstraints::SetMaximumDistanceConstraints(const TConstArrayView<FRealSingle>& MaxDistances)
-{
-	MaximumDistanceConstraints = MakeShared<Softs::FPBDSphericalConstraint>(
-		ParticleOffset,
-		NumParticles,
-		*AnimationPositions,
-		MaxDistances);
-	++NumConstraintRules;
-}
-
-void FClothConstraints::SetBackstopConstraints(const TConstArrayView<FRealSingle>& BackstopDistances, const TConstArrayView<FRealSingle>& BackstopRadiuses, bool bUseLegacyBackstop)
+void FClothConstraints::UpdateFromSolver(const FSolverVec3& SolverGravity, bool bPerClothGravityOverrideEnabled,
+	const FSolverVec3& FictitiousAngularDisplacement, const FSolverVec3& ReferenceSpaceLocation,
+	const FSolverVec3& InSolverWindVelocity, const FSolverReal LegacyWindAdaptation)
 {
-	BackstopConstraints = MakeShared<Softs::FPBDSphericalBackstopConstraint>(
-		ParticleOffset,
-		NumParticles,
-		*AnimationPositions,
-		*AnimationNormals,
-		BackstopRadiuses,
-		BackstopDistances,
-		bUseLegacyBackstop);
-	++NumConstraintRules;
-}
-
-void FClothConstraints::SetAnimDriveConstraints(const TConstArrayView<FRealSingle>& AnimDriveStiffnessMultipliers, const TConstArrayView<FRealSingle>& AnimDriveDampingMultipliers)
-{
-	if (AnimationVelocities)
+	if (ExternalForces)
 	{
-		AnimDriveConstraints = MakeShared<Softs::FPBDAnimDriveConstraint>(
-			ParticleOffset,
-			NumParticles,
-			*AnimationPositions,
-			TArray<Softs::FSolverVec3>(),
-			*AnimationVelocities,
-			AnimDriveStiffnessMultipliers,
-			AnimDriveDampingMultipliers);
+		ExternalForces->SetWorldGravityMultiplier((FSolverReal)ClothingSimulationClothConsoleVariables::CVarGravityMultiplier.GetValueOnAnyThread());
+		ExternalForces->SetSolverGravityProperties(SolverGravity, bPerClothGravityOverrideEnabled);
+		ExternalForces->SetFictitiousForcesData(FictitiousAngularDisplacement, ReferenceSpaceLocation);
+		ExternalForces->SetSolverWind(InSolverWindVelocity, LegacyWindAdaptation);
 	}
-
-	++NumConstraintInits;  // Uses init to update the property tables
-	++NumConstraintRules;
-}
-
-void FClothConstraints::SetShapeTargetConstraints(Softs::FSolverReal ShapeTargetStiffness)
-{
-	// TODO: Review this constraint. Currently does nothing more than the anim drive with less controls
-	check(ShapeTargetStiffness > 0.f && ShapeTargetStiffness <= 1.f);
-
-	ShapeConstraints = MakeShared<Softs::FPBDShapeConstraints>(
-		ParticleOffset,
-		NumParticles,
-		*AnimationPositions,
-		*AnimationPositions,
-		ShapeTargetStiffness);
-	++NumConstraintRules;
-}
-
-void FClothConstraints::SetSelfCollisionConstraints(const FTriangleMesh& TriangleMesh, TSet<TVec2<int32>>&& DisabledCollisionElements, Softs::FSolverReal SelfCollisionThickness, Softs::FSolverReal SelfCollisionFrictionCoefficient, bool bGlobalIntersectionAnalysis, bool bContourMinimization)
-{
-	SelfCollisionInit = MakeShared<Softs::FPBDTriangleMeshCollisions>(
-		ParticleOffset,
-		NumParticles,
-		TriangleMesh,
-		bGlobalIntersectionAnalysis,
-		bContourMinimization);
-
-	SelfCollisionConstraints = MakeShared<Softs::FPBDCollisionSpringConstraints>(
-		ParticleOffset,
-		NumParticles,
-		TriangleMesh,
-		AnimationPositions,
-		MoveTemp(DisabledCollisionElements),
-		SelfCollisionThickness,
-		Softs::FPBDCollisionSpringConstraintsBase::BackCompatStiffness,
-		SelfCollisionFrictionCoefficient);
-
-	++NumConstraintInits;
-	++NumPostCollisionConstraintRules;
-
-	SelfIntersectionConstraints = MakeShared<Softs::FPBDTriangleMeshIntersections>(
-		ParticleOffset,
-		NumParticles,
-		TriangleMesh);
-	++NumConstraintInits;
+	SolverWindVelocity = InSolverWindVelocity;
 }
 
 void FClothConstraints::Update(
@@ -1791,6 +1985,33 @@ void FClothConstraints::Update(
 	{
 		SelfCollisionSphereConstraints->SetProperties(ConfigProperties, VertexSets);
 	}
+
+	bool bUsePointBasedWindModel = false;
+	if (ExternalForces)
+	{
+		ExternalForces->SetProperties(ConfigProperties);
+		bUsePointBasedWindModel = ExternalForces->UsePointBasedWindModel();
+	}
+	if (VelocityAndPressureField)
+	{
+		constexpr FSolverReal WorldScale = 100.f;
+		const bool bPointBasedWindDisablesAccurateWind = ClothingSimulationClothConsoleVariables::CVarLegacyDisablesAccurateWind.GetValueOnAnyThread();
+		const bool bEnableAerodynamics = !(bUsePointBasedWindModel && bPointBasedWindDisablesAccurateWind);
+		VelocityAndPressureField->SetPropertiesAndWind(
+			ConfigProperties,
+			WeightMaps,
+			WorldScale,
+			bEnableAerodynamics,
+			SolverWindVelocity
+		);
+	}
+	if (CollisionConstraint)
+	{
+		static IConsoleVariable* const WriteCCDContacts = IConsoleManager::Get().FindConsoleVariable(TEXT("p.Chaos.PBDEvolution.WriteCCDContacts"));
+		const bool bWriteCCDContacts = WriteCCDContacts ? WriteCCDContacts->GetBool() : false;
+		CollisionConstraint->SetProperties(ConfigProperties);
+		CollisionConstraint->SetWriteDebugContacts(bWriteCCDContacts);
+	}
 }
 
 // Deprecated
@@ -1812,114 +2033,4 @@ void FClothConstraints::Update(
 	Update(ConfigProperties, TMap<FString, TConstArrayView<FRealSingle>>(), TMap<FString, const TSet<int32>*>(), MeshScale, MaxDistancesScale);
 }
 
-void FClothConstraints::SetEdgeProperties(const Softs::FSolverVec2& EdgeStiffness, const Softs::FSolverVec2& DampingRatio)
-{
-	if (EdgeConstraints)
-	{
-		static_cast<Softs::FPBDSpringConstraints*>(EdgeConstraints.Get())->SetProperties(EdgeStiffness);
-	}
-	if (XEdgeConstraints)
-	{
-		static_cast<Softs::FXPBDSpringConstraints*>(XEdgeConstraints.Get())->SetProperties(EdgeStiffness, DampingRatio);
-	}
-}
-
-void FClothConstraints::SetBendingProperties(const Softs::FSolverVec2& BendingStiffness, Softs::FSolverReal BucklingRatio, const Softs::FSolverVec2& BucklingStiffness, const Softs::FSolverVec2& BendingDampingRatio)
-{
-	if (BendingConstraints)
-	{
-		static_cast<Softs::FPBDSpringConstraints*>(BendingConstraints.Get())->SetProperties(BendingStiffness);
-	}
-	if (XBendingConstraints)
-	{
-		static_cast<Softs::FXPBDSpringConstraints*>(XBendingConstraints.Get())->SetProperties(BendingStiffness);
-	}
-	if (BendingElementConstraints)
-	{
-		BendingElementConstraints->SetProperties(BendingStiffness, BucklingRatio, BucklingStiffness);
-	}
-	if (XBendingElementConstraints)
-	{
-		XBendingElementConstraints->SetProperties(BendingStiffness, BucklingRatio, BucklingStiffness, BendingDampingRatio);
-	}
-}
-
-void FClothConstraints::SetAreaProperties(const Softs::FSolverVec2& AreaStiffness)
-{
-	if (AreaConstraints)
-	{
-		static_cast<Softs::FPBDAxialSpringConstraints*>(AreaConstraints.Get())->SetProperties(AreaStiffness);
-	}
-	if (XAreaConstraints)
-	{
-		static_cast<Softs::FXPBDAxialSpringConstraints*>(XAreaConstraints.Get())->SetProperties(AreaStiffness);
-	}
-}
-
-void FClothConstraints::SetThinShellVolumeProperties(Softs::FSolverReal VolumeStiffness)
-{
-	if (ThinShellVolumeConstraints_Deprecated)
-	{
-		ThinShellVolumeConstraints_Deprecated->SetProperties(VolumeStiffness);
-	}
-}
-
-void FClothConstraints::SetVolumeProperties(Softs::FSolverReal VolumeStiffness)
-{
-	if (VolumeConstraints_Deprecated)
-	{
-		VolumeConstraints_Deprecated->SetStiffness(VolumeStiffness);
-	}
-}
-
-void FClothConstraints::SetLongRangeAttachmentProperties(
-	const Softs::FSolverVec2& TetherStiffness,
-	const Softs::FSolverVec2& TetherScale,
-	Softs::FSolverReal MeshScale)
-{
-	if (LongRangeConstraints)
-	{
-		LongRangeConstraints->SetProperties(TetherStiffness, TetherScale, MeshScale);
-	}
-}
-
-void FClothConstraints::SetMaximumDistanceProperties(Softs::FSolverReal MeshScale)
-{
-	if (MaximumDistanceConstraints)
-	{
-		MaximumDistanceConstraints->SetScale(MeshScale);
-	}
-}
-
-void FClothConstraints::SetAnimDriveProperties(const Softs::FSolverVec2& AnimDriveStiffness, const Softs::FSolverVec2& AnimDriveDamping)
-{
-	if (AnimDriveConstraints)
-	{
-		AnimDriveConstraints->SetProperties(AnimDriveStiffness, AnimDriveDamping);
-	}
-}
-
-void FClothConstraints::SetSelfCollisionProperties(Softs::FSolverReal SelfCollisionThickness, Softs::FSolverReal SelfCollisionFrictionCoefficient, bool bGlobalIntersectionAnalysis, bool bContourMinimization)
-{
-	if (SelfCollisionInit)
-	{
-		SelfCollisionInit->SetGlobalIntersectionAnalysis(bGlobalIntersectionAnalysis);
-		SelfCollisionInit->SetContourMinimization(bContourMinimization);
-	}
-	if (SelfCollisionConstraints)
-	{
-		SelfCollisionConstraints->SetThickness(SelfCollisionThickness);
-		SelfCollisionConstraints->SetFrictionCoefficient(SelfCollisionFrictionCoefficient);
-	}
-}
-
-void FClothConstraints::SetBackstopProperties(bool bEnabled, Softs::FSolverReal MeshScale)
-{
-	if (BackstopConstraints)
-	{
-		BackstopConstraints->SetEnabled(bEnabled);
-		BackstopConstraints->SetScale(MeshScale);
-	}
-}
-
 }  // End namespace Chaos
diff --git a/Engine/Plugins/Experimental/ChaosCloth/Source/Chaos/Private/ChaosCloth/ChaosClothVisualization.cpp b/Engine/Plugins/Experimental/ChaosCloth/Source/Chaos/Private/ChaosCloth/ChaosClothVisualization.cpp
--- a/Engine/Plugins/Experimental/ChaosCloth/Source/Chaos/Private/ChaosCloth/ChaosClothVisualization.cpp
+++ b/Engine/Plugins/Experimental/ChaosCloth/Source/Chaos/Private/ChaosCloth/ChaosClothVisualization.cpp
@@ -120,11 +120,13 @@ static FAutoConsoleVariableRef CVarClothVizWeightMapName(TEXT("p.ChaosClothVisua
 
 		for (const FClothingSimulationCloth* const Cloth : Solver->GetCloths())
 		{
-			const int32 Offset = Cloth->GetOffset(Solver);
-			if (Offset == INDEX_NONE)
+			const int32 ParticleRangeId = Cloth->GetParticleRangeId(Solver);
+			if (ParticleRangeId == INDEX_NONE)
 			{
 				continue;
 			}
+			// Elements are local indexed for force based solver
+			const int32 Offset = Solver->IsForceBasedSolver() ? 0 : ParticleRangeId;
 
 			const TConstArrayView<TVec3<int32>> Elements = Cloth->GetTriangleMesh(Solver).GetElements();
 			const TConstArrayView<Softs::FSolverVec3> Positions = Cloth->GetParticlePositions(Solver);
@@ -184,11 +186,12 @@ static FAutoConsoleVariableRef CVarClothVizWeightMapName(TEXT("p.ChaosClothVisua
 
 		for (const FClothingSimulationCloth* const Cloth : Solver->GetCloths())
 		{
-			const int32 Offset = Cloth->GetOffset(Solver);
-			if (Offset == INDEX_NONE)
+			const int32 ParticleRangeId = Cloth->GetParticleRangeId(Solver);
+			if (ParticleRangeId == INDEX_NONE)
 			{
 				continue;
 			}
+			const int32 Offset = Solver->GetGlobalParticleOffset(ParticleRangeId);
 
 			const TConstArrayView<Softs::FSolverVec3> Positions = Cloth->GetParticlePositions(Solver);
 			const TConstArrayView<Softs::FSolverReal> InvMasses = Cloth->GetParticleInvMasses(Solver);
@@ -218,11 +221,13 @@ static FAutoConsoleVariableRef CVarClothVizWeightMapName(TEXT("p.ChaosClothVisua
 
 		for (const FClothingSimulationCloth* const Cloth : Solver->GetCloths())
 		{
-			const int32 Offset = Cloth->GetOffset(Solver);
-			if (Offset == INDEX_NONE)
+			const int32 ParticleRangeId = Cloth->GetParticleRangeId(Solver);
+			if (ParticleRangeId == INDEX_NONE)
 			{
 				continue;
 			}
+			// Elements are local indexed for force based solver
+			const int32 Offset = Solver->IsForceBasedSolver() ? 0 : ParticleRangeId;
 
 			const TArray<TVec3<int32>>& Elements = Cloth->GetTriangleMesh(Solver).GetElements();
 			const TConstArrayView<Softs::FSolverVec3> Positions = Cloth->GetParticlePositions(Solver);
@@ -270,13 +275,11 @@ static FAutoConsoleVariableRef CVarClothVizWeightMapName(TEXT("p.ChaosClothVisua
 
 		for (const FClothingSimulationCloth* const Cloth : Solver->GetCloths())
 		{
-			const int32 Offset = Cloth->GetOffset(Solver);
-			if (Offset == INDEX_NONE)
+			if (Cloth->GetParticleRangeId(Solver) == INDEX_NONE)
 			{
 				continue;
 			}
 
-
 			const TConstArrayView<FRealSingle>& MaxDistances = Cloth->GetWeightMapByProperty(Solver, TEXT("MaxDistance"));
 			if (!MaxDistances.Num())
 			{
@@ -311,12 +314,13 @@ static FAutoConsoleVariableRef CVarClothVizWeightMapName(TEXT("p.ChaosClothVisua
 
 		for (const FClothingSimulationCloth* const Cloth : Solver->GetCloths())
 		{
-			const int32 Offset = Cloth->GetOffset(Solver);
-			if (Offset == INDEX_NONE)
+			const int32 ParticleRangeId = Cloth->GetParticleRangeId(Solver);
+			if (ParticleRangeId == INDEX_NONE)
 			{
 				continue;
 			}
-
+			// Elements are local indexed for force based solver
+			const int32 Offset = Solver->IsForceBasedSolver() ? 0 : ParticleRangeId;
 			const TConstArrayView<TVec3<int32>> Elements = Cloth->GetTriangleMesh(Solver).GetElements();
 			const TConstArrayView<Softs::FSolverVec3> Positions = Cloth->GetParticlePositions(Solver);
 			const TConstArrayView<FRealSingle>& WeightMap = Cloth->GetWeightMapByName(Solver, Name);
@@ -677,7 +681,8 @@ static FAutoConsoleVariableRef CVarClothVizWeightMapName(TEXT("p.ChaosClothVisua
 		static const FLinearColor Color = FLinearColor(FColor::Purple).Desaturate(0.5);
 		for (const FClothingSimulationCloth* const Cloth : Solver->GetCloths())
 		{
-			if (Cloth->GetOffset(Solver) == INDEX_NONE)
+			const int32 ParticleRangeId = Cloth->GetParticleRangeId(Solver);
+			if (ParticleRangeId == INDEX_NONE)
 			{
 				continue;
 			}
@@ -697,7 +702,7 @@ static FAutoConsoleVariableRef CVarClothVizWeightMapName(TEXT("p.ChaosClothVisua
 		// Draw gravity
 		for (const FClothingSimulationCloth* const Cloth : Solver->GetCloths())
 		{
-			if (Cloth->GetOffset(Solver) == INDEX_NONE)
+			if (Cloth->GetParticleRangeId(Solver) == INDEX_NONE)
 			{
 				continue;
 			}
@@ -724,12 +729,13 @@ static FAutoConsoleVariableRef CVarClothVizWeightMapName(TEXT("p.ChaosClothVisua
 
 		for (const FClothingSimulationCloth* const Cloth : Solver->GetCloths())
 		{
-			const int32 Offset = Cloth->GetOffset(Solver);
-			if (Offset == INDEX_NONE)
+			const int32 ParticleRangeId = Cloth->GetParticleRangeId(Solver);
+			if (ParticleRangeId == INDEX_NONE)
 			{
 				continue;
 			}
-
+			// Elements are local indexed for force based solver
+			const int32 Offset = Solver->IsForceBasedSolver() ? 0 : ParticleRangeId;
 			const TConstArrayView<TVec3<int32>> Elements = Cloth->GetTriangleMesh(Solver).GetElements();
 			const TConstArrayView<Softs::FSolverVec3> Positions = Cloth->GetParticlePositions(Solver);
 			const TConstArrayView<Softs::FSolverReal> InvMasses = Cloth->GetParticleInvMasses(Solver);
@@ -767,11 +773,13 @@ static FAutoConsoleVariableRef CVarClothVizWeightMapName(TEXT("p.ChaosClothVisua
 
 		for (const FClothingSimulationCloth* const Cloth : Solver->GetCloths())
 		{
-			const int32 Offset = Cloth->GetOffset(Solver);
-			if (Offset == INDEX_NONE)
+			const int32 ParticleRangeId = Cloth->GetParticleRangeId(Solver);
+			if (ParticleRangeId == INDEX_NONE)
 			{
 				continue;
 			}
+			// Elements are local indexed for force based solver
+			const int32 Offset = Solver->IsForceBasedSolver() ? 0 : ParticleRangeId;
 
 			const TConstArrayView<TVec3<int32>> Elements = Cloth->GetTriangleMesh(Solver).GetElements();
 			const TConstArrayView<Softs::FSolverVec3> Positions = Cloth->GetAnimationPositions(Solver);
@@ -831,11 +839,13 @@ static FAutoConsoleVariableRef CVarClothVizWeightMapName(TEXT("p.ChaosClothVisua
 
 		for (const FClothingSimulationCloth* const Cloth : Solver->GetCloths())
 		{
-			const int32 Offset = Cloth->GetOffset(Solver);
-			if (Offset == INDEX_NONE)
+			const int32 ParticleRangeId = Cloth->GetParticleRangeId(Solver);
+			if (ParticleRangeId == INDEX_NONE)
 			{
 				continue;
 			}
+			// Elements are local indexed for force based solver
+			const int32 Offset = Solver->IsForceBasedSolver() ? 0 : ParticleRangeId;
 
 			const TConstArrayView<TVec3<int32>> Elements = Cloth->GetTriangleMesh(Solver).GetElements();
 			const TConstArrayView<Softs::FSolverVec3> Positions = Cloth->GetAnimationPositions(Solver);
@@ -871,8 +881,7 @@ static FAutoConsoleVariableRef CVarClothVizWeightMapName(TEXT("p.ChaosClothVisua
 
 		for (const FClothingSimulationCloth* const Cloth : Solver->GetCloths())
 		{
-			const int32 Offset = Cloth->GetOffset(Solver);
-			if (Offset == INDEX_NONE)
+			if (Cloth->GetParticleRangeId(Solver) == INDEX_NONE)
 			{
 				continue;
 			}
@@ -906,8 +915,7 @@ static FAutoConsoleVariableRef CVarClothVizWeightMapName(TEXT("p.ChaosClothVisua
 
 		for (const FClothingSimulationCloth* const Cloth : Solver->GetCloths())
 		{
-			const int32 Offset = Cloth->GetOffset(Solver);
-			if (Offset == INDEX_NONE)
+			if (Cloth->GetParticleRangeId(Solver) == INDEX_NONE)
 			{
 				continue;
 			}
@@ -940,8 +948,7 @@ static FAutoConsoleVariableRef CVarClothVizWeightMapName(TEXT("p.ChaosClothVisua
 
 		for (const FClothingSimulationCloth* const Cloth : Solver->GetCloths())
 		{
-			const int32 Offset = Cloth->GetOffset(Solver);
-			if (Offset == INDEX_NONE)
+			if (Cloth->GetParticleRangeId(Solver) == INDEX_NONE)
 			{
 				continue;
 			}
@@ -1160,13 +1167,13 @@ static FAutoConsoleVariableRef CVarClothVizWeightMapName(TEXT("p.ChaosClothVisua
 
 		for (const FClothingSimulationCloth* const Cloth : Solver->GetCloths())
 		{
-			const int32 Offset = Cloth->GetOffset(Solver);
-			if (Offset == INDEX_NONE)
+			const int32 ParticleRangeId = Cloth->GetParticleRangeId(Solver);
+			if (ParticleRangeId == INDEX_NONE)
 			{
 				continue;
 			}
 
-			const FClothConstraints& ClothConstraints = Solver->GetClothConstraints(Offset);
+			const FClothConstraints& ClothConstraints = Solver->GetClothConstraints(ParticleRangeId);
 			if (const Softs::FPBDSphericalBackstopConstraint* const BackstopConstraint = ClothConstraints.GetBackstopConstraints().Get())
 			{
 				const bool bUseLegacyBackstop = BackstopConstraint->UseLegacyBackstop();
@@ -1220,13 +1227,13 @@ static FAutoConsoleVariableRef CVarClothVizWeightMapName(TEXT("p.ChaosClothVisua
 
 		for (const FClothingSimulationCloth* const Cloth : Solver->GetCloths())
 		{
-			const int32 Offset = Cloth->GetOffset(Solver);
-			if (Offset == INDEX_NONE)
+			const int32 ParticleRangeId = Cloth->GetParticleRangeId(Solver);
+			if (ParticleRangeId == INDEX_NONE)
 			{
 				continue;
 			}
 
-			const FClothConstraints& ClothConstraints = Solver->GetClothConstraints(Offset);
+			const FClothConstraints& ClothConstraints = Solver->GetClothConstraints(ParticleRangeId);
 			if (const Softs::FPBDSphericalBackstopConstraint* const BackstopConstraint = ClothConstraints.GetBackstopConstraints().Get())
 			{
 				const bool bUseLegacyBackstop = BackstopConstraint->UseLegacyBackstop();
@@ -1267,8 +1274,8 @@ static FAutoConsoleVariableRef CVarClothVizWeightMapName(TEXT("p.ChaosClothVisua
 
 		for (const FClothingSimulationCloth* const Cloth : Solver->GetCloths())
 		{
-			const int32 Offset = Cloth->GetOffset(Solver);
-			if (Offset == INDEX_NONE)
+			const int32 ParticleRangeId = Cloth->GetParticleRangeId(Solver);
+			if (ParticleRangeId == INDEX_NONE)
 			{
 				continue;
 			}
@@ -1317,13 +1324,13 @@ static FAutoConsoleVariableRef CVarClothVizWeightMapName(TEXT("p.ChaosClothVisua
 
 		for (const FClothingSimulationCloth* const Cloth : Solver->GetCloths())
 		{
-			const int32 Offset = Cloth->GetOffset(Solver);
-			if (Offset == INDEX_NONE)
+			const int32 ParticleRangeId = Cloth->GetParticleRangeId(Solver);
+			if (ParticleRangeId == INDEX_NONE)
 			{
 				continue;
 			}
 
-			const FClothConstraints& ClothConstraints = Solver->GetClothConstraints(Offset);
+			const FClothConstraints& ClothConstraints = Solver->GetClothConstraints(ParticleRangeId);
 			if (const Softs::FPBDAnimDriveConstraint* const AnimDriveConstraint = ClothConstraints.GetAnimDriveConstraints().Get())
 			{
 				const TConstArrayView<FRealSingle>& AnimDriveStiffnessMultipliers = Cloth->GetWeightMapByProperty(Solver, TEXT("AnimDriveStiffness"));
@@ -1592,16 +1599,17 @@ static FAutoConsoleVariableRef CVarClothVizWeightMapName(TEXT("p.ChaosClothVisua
 
 		for (const FClothingSimulationCloth* const Cloth : Solver->GetCloths())
 		{
-			const int32 Offset = Cloth->GetOffset(Solver);
-			if (Offset == INDEX_NONE)
+			const int32 ParticleRangeId = Cloth->GetParticleRangeId(Solver);
+			if (ParticleRangeId == INDEX_NONE)
 			{
 				continue;
 			}
 
-			// Draw constraints
-			const FClothConstraints& ClothConstraints = Solver->GetClothConstraints(Offset);
+			const FClothConstraints& ClothConstraints = Solver->GetClothConstraints(ParticleRangeId);
 
-			const TArray<Softs::FSolverVec3>& Positions = Solver->GetParticleXs();
+			// Constraints are locally indexed for force based solver
+			const TConstArrayView<Softs::FSolverVec3> Positions = Solver->IsForceBasedSolver() ? Solver->GetParticleXsView(ParticleRangeId)
+				: TConstArrayView<Softs::FSolverVec3>(Solver->GetParticleXs());
 
 			if (const Softs::FPBDEdgeSpringConstraints* const EdgeConstraints = ClothConstraints.GetEdgeSpringConstraints().Get())
 			{
@@ -1701,16 +1709,18 @@ static FAutoConsoleVariableRef CVarClothVizWeightMapName(TEXT("p.ChaosClothVisua
 
 		for (const FClothingSimulationCloth* const Cloth : Solver->GetCloths())
 		{
-			const int32 Offset = Cloth->GetOffset(Solver);
-			if (Offset == INDEX_NONE)
+			const int32 ParticleRangeId = Cloth->GetParticleRangeId(Solver);
+			if (ParticleRangeId == INDEX_NONE)
 			{
 				continue;
 			}
 
-			// Draw constraints
-			const FClothConstraints& ClothConstraints = Solver->GetClothConstraints(Offset);
+			const FClothConstraints& ClothConstraints = Solver->GetClothConstraints(ParticleRangeId);
+
 
-			const TArray<Softs::FSolverVec3>& Positions = Solver->GetParticleXs();
+			// Constraints are locally indexed for force based solver
+			const TConstArrayView<Softs::FSolverVec3> Positions = Solver->IsForceBasedSolver() ? Solver->GetParticleXsView(ParticleRangeId)
+				: TConstArrayView<Softs::FSolverVec3>(Solver->GetParticleXs());
 
 			if (const Softs::FPBDBendingSpringConstraints* const BendingConstraints = ClothConstraints.GetBendingSpringConstraints().Get())
 			{
@@ -1811,14 +1821,13 @@ static FAutoConsoleVariableRef CVarClothVizWeightMapName(TEXT("p.ChaosClothVisua
 
 		for (const FClothingSimulationCloth* const Cloth : Solver->GetCloths())
 		{
-			const int32 Offset = Cloth->GetOffset(Solver);
-			if (Offset == INDEX_NONE)
+			const int32 ParticleRangeId = Cloth->GetParticleRangeId(Solver);
+			if (ParticleRangeId == INDEX_NONE)
 			{
 				continue;
 			}
 
-			// Draw constraints
-			const FClothConstraints& ClothConstraints = Solver->GetClothConstraints(Offset);
+			const FClothConstraints& ClothConstraints = Solver->GetClothConstraints(ParticleRangeId);
 			const TConstArrayView<Softs::FSolverVec3> Positions = Cloth->GetParticlePositions(Solver);
 
 			if (const Softs::FPBDLongRangeConstraints* const LongRangeConstraints = ClothConstraints.GetLongRangeConstraints().Get())
@@ -1876,18 +1885,35 @@ static FAutoConsoleVariableRef CVarClothVizWeightMapName(TEXT("p.ChaosClothVisua
 
 		for (const FClothingSimulationCloth* const Cloth : Solver->GetCloths())
 		{
-			const int32 Offset = Cloth->GetOffset(Solver);
-			if (Offset == INDEX_NONE)
+			const int32 ParticleRangeId = Cloth->GetParticleRangeId(Solver);
+			if (ParticleRangeId == INDEX_NONE)
 			{
 				continue;
 			}
 
-			const Softs::FVelocityAndPressureField& VelocityField = Solver->GetWindVelocityAndPressureField(Cloth->GetGroupId());
+			const Softs::FVelocityAndPressureField* VelocityField = nullptr;
+			if (Solver->IsForceBasedSolver())
+			{
+				const FClothConstraints& ClothConstraints = Solver->GetClothConstraints(ParticleRangeId);
+				VelocityField = ClothConstraints.GetVelocityAndPressureField().Get();
+				if (!VelocityField)
+				{
+					continue;
+				}
+			}
+			else
+			{
+				VelocityField = &Solver->GetWindVelocityAndPressureField(Cloth->GetGroupId());
+			}
+
+			// Constraints are locally indexed for force based solver
+			const TConstArrayView<Softs::FSolverVec3> Positions = Solver->IsForceBasedSolver() ? Solver->GetParticleXsView(ParticleRangeId)
+				: TConstArrayView<Softs::FSolverVec3>(Solver->GetParticleXs());
+			const TConstArrayView<Softs::FSolverReal> InvMasses = Solver->IsForceBasedSolver() ? Solver->GetParticleInvMassesView(ParticleRangeId)
+				: TConstArrayView<Softs::FSolverReal>(Solver->GetParticleInvMasses());
 
-			const TConstArrayView<TVec3<int32>>& Elements = VelocityField.GetElements();
-			const TConstArrayView<Softs::FSolverVec3> Forces = VelocityField.GetForces();
-			const TArray<Softs::FSolverVec3>& Positions = Solver->GetParticleXs();
-			const TArray<Softs::FSolverReal>& InvMasses = Solver->GetParticleInvMasses();
+			const TConstArrayView<TVec3<int32>>& Elements = VelocityField->GetElements();
+			const TConstArrayView<Softs::FSolverVec3> Forces = VelocityField->GetForces();
 			check(InvMasses.Num() == Positions.Num());
 
 			for (int32 ElementIndex = 0; ElementIndex < Elements.Num(); ++ElementIndex)
@@ -1922,7 +1948,8 @@ static FAutoConsoleVariableRef CVarClothVizWeightMapName(TEXT("p.ChaosClothVisua
 		// Draw reference spaces
 		for (const FClothingSimulationCloth* const Cloth : Solver->GetCloths())
 		{
-			if (Cloth->GetOffset(Solver) == INDEX_NONE)
+			const int32 ParticleRangeId = Cloth->GetParticleRangeId(Solver);
+			if (ParticleRangeId == INDEX_NONE)
 			{
 				continue;
 			}
@@ -1942,18 +1969,20 @@ static FAutoConsoleVariableRef CVarClothVizWeightMapName(TEXT("p.ChaosClothVisua
 
 		for (const FClothingSimulationCloth* const Cloth : Solver->GetCloths())
 		{
-			const int32 Offset = Cloth->GetOffset(Solver);
-			if (Offset == INDEX_NONE)
+			const int32 ParticleRangeId = Cloth->GetParticleRangeId(Solver);
+			if (ParticleRangeId == INDEX_NONE)
 			{
 				continue;
 			}
 
-			// Draw constraints
-			const FClothConstraints& ClothConstraints = Solver->GetClothConstraints(Offset);
+			const FClothConstraints& ClothConstraints = Solver->GetClothConstraints(ParticleRangeId);
 
 			if (const Softs::FPBDCollisionSpringConstraints* const SelfCollisionConstraints = ClothConstraints.GetSelfCollisionConstraints().Get())
 			{
-				const TArray<Softs::FSolverVec3>& Positions = Solver->GetParticleXs();
+				// Constraints are locally indexed for force based solver
+				const TConstArrayView<Softs::FSolverVec3> Positions = Solver->IsForceBasedSolver() ? Solver->GetParticleXsView(ParticleRangeId)
+					: TConstArrayView<Softs::FSolverVec3>(Solver->GetParticleXs());
+				const int32 Offset = Solver->IsForceBasedSolver() ? 0 : ParticleRangeId;
 				const TArray<TVec4<int32>>& Constraints = SelfCollisionConstraints->GetConstraints();
 				const TArray<Softs::FSolverVec3>& Barys = SelfCollisionConstraints->GetBarys();
 				const FReal Thickness = (FReal)SelfCollisionConstraints->GetThickness();
@@ -1992,7 +2021,10 @@ static FAutoConsoleVariableRef CVarClothVizWeightMapName(TEXT("p.ChaosClothVisua
 			if (const Softs::FPBDSelfCollisionSphereConstraints* const SelfCollisionSphereConstraints =
 				ClothConstraints.GetSelfCollisionSphereConstraints().Get())
 			{
-				const TArray<Softs::FSolverVec3>& Positions = Solver->GetParticleXs();
+				// Constraints are locally indexed for force based solver
+				const TConstArrayView<Softs::FSolverVec3> Positions = Solver->IsForceBasedSolver() ? Solver->GetParticleXsView(ParticleRangeId)
+					: TConstArrayView<Softs::FSolverVec3>(Solver->GetParticleXs());
+				const int32 Offset = Solver->IsForceBasedSolver() ? 0 : ParticleRangeId;
 				const TArray<TVec2<int32>>& Constraints = SelfCollisionSphereConstraints->GetConstraints();
 				for (int32 Index = 0; Index < Constraints.Num(); ++Index)
 				{
@@ -2028,13 +2060,13 @@ static FAutoConsoleVariableRef CVarClothVizWeightMapName(TEXT("p.ChaosClothVisua
 
 		for (const FClothingSimulationCloth* const Cloth : Solver->GetCloths())
 		{
-			const int32 Offset = Cloth->GetOffset(Solver);
-			if (Offset == INDEX_NONE)
+			const int32 ParticleRangeId = Cloth->GetParticleRangeId(Solver);
+			if (ParticleRangeId == INDEX_NONE)
 			{
 				continue;
 			}
 
-			const FClothConstraints& ClothConstraints = Solver->GetClothConstraints(Offset);
+			const FClothConstraints& ClothConstraints = Solver->GetClothConstraints(ParticleRangeId);
 
 			static const FLinearColor Red(1.f, 0.f, 0.f);
 			static const FLinearColor White(1.f, 1.f, 1.f);
@@ -2047,7 +2079,12 @@ static FAutoConsoleVariableRef CVarClothVizWeightMapName(TEXT("p.ChaosClothVisua
 
 			if (const Softs::FPBDTriangleMeshCollisions* const SelfCollisionInit = ClothConstraints.GetSelfCollisionInit().Get())
 			{
-				const TArray<Softs::FSolverVec3>& Positions = Solver->GetParticleXs();
+				// Constraints are locally indexed for force based solver
+				const TConstArrayView<Softs::FSolverVec3> Positions = Solver->IsForceBasedSolver() ? Solver->GetParticleXsView(ParticleRangeId)
+					: TConstArrayView<Softs::FSolverVec3>(Solver->GetParticleXs());
+
+				const int32 Offset = Solver->IsForceBasedSolver() ? 0 : ParticleRangeId;
+
 				const FTriangleMesh& TriangleMesh = Cloth->GetTriangleMesh(Solver);
 
 				// Draw contours
