Update the Cloth solver plugin to optionally use the new Softs::FEvolution.

diff --git a/Engine/Plugins/Experimental/ChaosCloth/Source/Chaos/Private/ChaosCloth/ChaosClothingCacheSchema.cpp b/Engine/Plugins/Experimental/ChaosCloth/Source/Chaos/Private/ChaosCloth/ChaosClothingCacheSchema.cpp
--- a/Engine/Plugins/Experimental/ChaosCloth/Source/Chaos/Private/ChaosCloth/ChaosClothingCacheSchema.cpp
+++ b/Engine/Plugins/Experimental/ChaosCloth/Source/Chaos/Private/ChaosCloth/ChaosClothingCacheSchema.cpp
@@ -12,11 +12,11 @@ namespace Chaos
 
 void FClothingCacheSchema::RecordPostSolve(const FClothingSimulationSolver& ClothSolver, FPendingFrameWrite& OutFrame, FReal InTime)
 {
-	const int32 NumParticles = ClothSolver.GetNumParticles();
+	const int32 NumParticles = ClothSolver.GetNumActiveParticles();
 	if (NumParticles > 0)
 	{
-		const Softs::FSolverVec3* ParticleXs = ClothSolver.GetParticleXs(0);
-		const Softs::FSolverVec3* ParticleVs = ClothSolver.GetParticleVs(0);
+		const Softs::FSolverVec3* const ParticleXs = ClothSolver.GetParticleXs(0);
+		const Softs::FSolverVec3* const ParticleVs = ClothSolver.GetParticleVs(0);
 
 		TArray<float> PendingVX, PendingVY, PendingVZ, PendingPX, PendingPY, PendingPZ;
 		TArray<int32>& PendingID = OutFrame.PendingChannelsIndices;
@@ -37,16 +37,16 @@ void FClothingCacheSchema::RecordPostSolve(const FClothingSimulationSolver& Clot
 			const FName ReferenceTransformNameAndIndex(ReferenceTransformsName, Cloth->GetGroupId());
 			OutFrame.PendingNamedTransformData.Add(ReferenceTransformNameAndIndex, ReferenceSpaceTransform);
 
-			const int32 ClothOffset = Cloth->GetOffset(&ClothSolver);
+			const int32 ParticleRangeId = Cloth->GetParticleRangeId(&ClothSolver);
 			const int32 NumClothParticles = Cloth->GetNumParticles(&ClothSolver);
-			check(ClothOffset + NumClothParticles <= NumParticles);
-			for (int32 ParticleIndex = ClothOffset; ParticleIndex < ClothOffset + NumClothParticles; ++ParticleIndex)
+			const int32 GlobalOffset = ClothSolver.GetGlobalParticleOffset(ParticleRangeId);
+			for (int32 ParticleIndex = GlobalOffset; ParticleIndex < GlobalOffset + NumClothParticles; ++ParticleIndex)
 			{
 				const Softs::FSolverVec3& ParticleV = ParticleVs[ParticleIndex];
 				const Softs::FSolverVec3& ParticleX = ParticleXs[ParticleIndex];
 
 				// Adding the vertices relative position to the particle write datas
-				PendingID.Add(ParticleIndex);
+				PendingID.Add(ParticleIndex + GlobalOffset);
 				PendingVX.Add(ParticleV.X);
 				PendingVY.Add(ParticleV.Y);
 				PendingVZ.Add(ParticleV.Z);
@@ -89,8 +89,8 @@ void FClothingCacheSchema::PlaybackPreSolve(UChaosCache& InCache, FReal InTime,
 	if (PendingVX && PendingVY && PendingVZ && PendingPX && PendingPY && PendingPZ && NumCachedParticles)
 	{
 		// Directly set the result of the cache into the solver particles
-		Softs::FSolverVec3* ParticleXs = ClothSolver.GetParticleXs(0);
-		Softs::FSolverVec3* ParticleVs = ClothSolver.GetParticleVs(0);
+		Softs::FSolverVec3* const ParticleXs = ClothSolver.GetParticleXs(0);
+		Softs::FSolverVec3* const ParticleVs = ClothSolver.GetParticleVs(0);
 		const int32 NumParticles = ClothSolver.GetNumParticles();
 		for (int32 CachedIndex = 0; CachedIndex < NumCachedParticles; ++CachedIndex)
 		{
@@ -132,11 +132,13 @@ void FClothingCacheSchema::PlaybackPreSolve(UChaosCache& InCache, FReal InTime,
 					const FTransform& CachedTransform = NamedTransform.Value;
 					if (!CurrentTransform.Equals(CachedTransform))
 					{
+						const int32 ParticleRangeId = Cloth->GetParticleRangeId(&ClothSolver);
+						const int32 GlobalOffset = ClothSolver.GetGlobalParticleOffset(ParticleRangeId);
+
 						const FTransform RelativeTransform = CurrentTransform.GetRelativeTransform(CachedTransform);
-						const int32 ClothOffset = Cloth->GetOffset(&ClothSolver);
 						const int32 NumClothParticles = Cloth->GetNumParticles(&ClothSolver);
-						check(ClothOffset + NumClothParticles <= NumParticles);
-						for (int32 ParticleIndex = ClothOffset; ParticleIndex < ClothOffset + NumClothParticles; ++ParticleIndex)
+						check(GlobalOffset + NumClothParticles <= NumParticles);
+						for (int32 ParticleIndex = GlobalOffset; ParticleIndex < GlobalOffset + NumClothParticles; ++ParticleIndex)
 						{
 							Softs::FSolverVec3& ParticleX = ParticleXs[ParticleIndex];
 							ParticleX = Softs::FSolverVec3(RelativeTransform.TransformPosition(FVec3(ParticleX + LocalSpaceLocation)) - LocalSpaceLocation);
diff --git a/Engine/Plugins/Experimental/ChaosCloth/Source/Chaos/Private/ChaosCloth/ChaosClothingSimulation.cpp b/Engine/Plugins/Experimental/ChaosCloth/Source/Chaos/Private/ChaosCloth/ChaosClothingSimulation.cpp
--- a/Engine/Plugins/Experimental/ChaosCloth/Source/Chaos/Private/ChaosCloth/ChaosClothingSimulation.cpp
+++ b/Engine/Plugins/Experimental/ChaosCloth/Source/Chaos/Private/ChaosCloth/ChaosClothingSimulation.cpp
@@ -466,7 +466,7 @@ bool FClothingSimulation::ShouldSimulate() const
 {
 	for (const TUniquePtr<FClothingSimulationCloth>& Cloth : Cloths)
 	{
-		if (Cloth->GetLODIndex(Solver.Get()) != INDEX_NONE && Cloth->GetOffset(Solver.Get()) != INDEX_NONE)
+		if (Cloth->GetLODIndex(Solver.Get()) != INDEX_NONE && Cloth->GetParticleRangeId(Solver.Get()) != INDEX_NONE)
 		{
 			return true;
 		}
@@ -646,7 +646,7 @@ PRAGMA_ENABLE_DEPRECATION_WARNINGS
 			Solver->Update(FSolverReal(0.));  // Update for LOD switching, but do not simulate
 		}
 
-		if (Cloth->GetOffset(Solver.Get()) == INDEX_NONE || Cloth->GetLODIndex(Solver.Get()) == INDEX_NONE)
+		if (Cloth->GetParticleRangeId(Solver.Get()) == INDEX_NONE || Cloth->GetLODIndex(Solver.Get()) == INDEX_NONE)
 		{
 			OutData.Remove(AssetIndex);  // Ensures that the cloth vertex factory won't run unnecessarily
 			continue;  // No valid LOD, there's nothing to write out
diff --git a/Engine/Plugins/Experimental/ChaosCloth/Source/Chaos/Private/ChaosCloth/ChaosClothingSimulationCloth.cpp b/Engine/Plugins/Experimental/ChaosCloth/Source/Chaos/Private/ChaosCloth/ChaosClothingSimulationCloth.cpp
--- a/Engine/Plugins/Experimental/ChaosCloth/Source/Chaos/Private/ChaosCloth/ChaosClothingSimulationCloth.cpp
+++ b/Engine/Plugins/Experimental/ChaosCloth/Source/Chaos/Private/ChaosCloth/ChaosClothingSimulationCloth.cpp
@@ -9,6 +9,8 @@
 #include "ChaosCloth/ChaosWeightMapTarget.h"
 #include "ChaosCloth/ChaosClothPrivate.h"
 #include "Chaos/CollectionPropertyFacade.h"
+#include "Chaos/PBDSoftBodyCollisionConstraint.h"
+#include "Chaos/SoftsExternalForces.h"
 #include "Containers/ArrayView.h"
 #include "GeometryCollection/ManagedArrayCollection.h"
 #include "HAL/IConsoleManager.h"
@@ -23,17 +25,17 @@ namespace ClothingSimulationClothDefault
 	constexpr bool bUseGeodesicTethers = true;
 	constexpr float MassValue = 0.35f;
 	constexpr float MinPerParticleMass = 0.0001f;
-	constexpr float CollisionThickness = 1.0f;
-	constexpr float FrictionCoefficient = 0.8f;
+	constexpr float CollisionThickness = Softs::FPBDSoftBodyCollisionConstraint::DefaultCollisionThickness;
+	constexpr float FrictionCoefficient = Softs::FPBDSoftBodyCollisionConstraint::DefaultFrictionCoefficient;
 	constexpr float DampingCoefficient = 0.01f;
 	constexpr float Drag = 0.035f;
 	constexpr float Lift = 0.035f;
 	constexpr float Pressure = 0.f;
 	constexpr float AirDensity = 1.225f;  // Air density in kg/m^3
-	constexpr float GravityScale = 1.f;
-	constexpr float GravityZOverride = -980.665f;
+	constexpr float GravityScale = Softs::FExternalForces::DefaultGravityScale; // 1.f;
+	constexpr float GravityZOverride = Softs::FExternalForces::DefaultGravityZOverride; // -980.665f;
 	constexpr float VelocityScale = 0.75f;
-	constexpr float FictitiousAngularScale = 1.f;
+	constexpr float FictitiousAngularScale = Softs::FExternalForces::DefaultFictitiousAngularScale; // 1.f;
 }
 
 namespace ClothingSimulationClothConsoleVariables
@@ -53,19 +55,19 @@ struct FClothingSimulationCloth::FLODData
 {
 	// Input mesh
 	const int32 NumParticles;
-	const TConstArrayView<uint32> Indices;
 	const TMap<FString, TConstArrayView<FRealSingle>> WeightMaps;
 	const TMap<FString, const TSet<int32>*> VertexSets;
 	const TArray<TConstArrayView<TTuple<int32, int32, FRealSingle>>> Tethers;
 
 	const FClothingPatternData PatternData;
+	const FTriangleMesh NoOffsetTriangleMesh;
 
 	// Per Solver data
 	struct FSolverData
 	{
 		int32 LODIndex;
-		int32 Offset;
-		FTriangleMesh TriangleMesh;  // TODO: Triangle Mesh shouldn't really be solver dependent (ie not use an offset)
+		int32 ParticleRangeId;
+		FTriangleMesh OffsetTriangleMesh; // Only used if using PBD solver
 	};
 	TMap<FClothingSimulationSolver*, FSolverData> SolverData;
 
@@ -83,6 +85,8 @@ struct FClothingSimulationCloth::FLODData
 		TMap<FString, const TSet<int32>*>&& InVertexSets,
 		const TArray<TConstArrayView<TTuple<int32, int32, FRealSingle>>>& InTethers);
 
+	static FTriangleMesh BuildTriangleMesh(const TConstArrayView<uint32>& Indices, const int32 NumParticles);
+
 	void Add(FClothingSimulationSolver* Solver, FClothingSimulationCloth* Cloth, int32 LODIndex);
 	void Remove(FClothingSimulationSolver* Solver);
 
@@ -105,14 +109,36 @@ FClothingSimulationCloth::FLODData::FLODData(
 	TMap<FString, const TSet<int32>*>&& InVertexSets,
 	const TArray<TConstArrayView<TTuple<int32, int32, FRealSingle>>>& InTethers)
 	: NumParticles(InNumParticles)
-	, Indices(InIndices)
 	, WeightMaps(MoveTemp(InWeightMaps))
 	, VertexSets(MoveTemp(InVertexSets))
 	, Tethers(InTethers)
-	, PatternData(NumParticles, Indices, InPatternPositions, InPatternIndices, InPatternToWeldedIndices)
+	, PatternData(NumParticles, InIndices, InPatternPositions, InPatternIndices, InPatternToWeldedIndices)
+	, NoOffsetTriangleMesh(BuildTriangleMesh(InIndices, NumParticles))
 {
 }
 
+FTriangleMesh FClothingSimulationCloth::FLODData::BuildTriangleMesh(const TConstArrayView<uint32>& Indices, const int32 InNumParticles)
+{
+	FTriangleMesh OutTriangleMesh;
+	// Build a sim friendly triangle mesh including the solver particle's Offset
+	const int32 NumElements = Indices.Num() / 3;
+	TArray<TVec3<int32>> Elements;
+	Elements.Reserve(NumElements);
+
+	for (int32 i = 0; i < NumElements; ++i)
+	{
+		const int32 Index = 3 * i;
+		Elements.Add(
+			{ static_cast<int32>(Indices[Index]),
+			 static_cast<int32>(Indices[Index + 1]),
+			 static_cast<int32>(Indices[Index + 2]) });
+	}
+
+	OutTriangleMesh.Init(MoveTemp(Elements), 0, InNumParticles - 1);
+	OutTriangleMesh.GetPointToTriangleMap(); // Builds map for later use by GetPointNormals(), and the velocity fields
+	return OutTriangleMesh;
+}
+
 void FClothingSimulationCloth::FLODData::Add(FClothingSimulationSolver* Solver, FClothingSimulationCloth* Cloth, int32 InLODIndex)
 {
 	check(Solver);
@@ -123,11 +149,9 @@ void FClothingSimulationCloth::FLODData::Add(FClothingSimulationSolver* Solver,
 	check(!SolverData.Find(Solver));
 	FSolverData& SolverDatum = SolverData.Add(Solver);
 	SolverDatum.LODIndex = InLODIndex;
-	int32& Offset = SolverDatum.Offset;
-	FTriangleMesh& TriangleMesh = SolverDatum.TriangleMesh;
-
+	int32& ParticleRangeId = SolverDatum.ParticleRangeId;
 	// Add particles
-	Offset = Solver->AddParticles(NumParticles, Cloth->GroupId);  // TODO: Have a per solver map of offset
+	ParticleRangeId = Solver->AddParticles(NumParticles, Cloth->GroupId);
 
 	if (!NumParticles)
 	{
@@ -136,7 +160,7 @@ void FClothingSimulationCloth::FLODData::Add(FClothingSimulationSolver* Solver,
 
 PRAGMA_DISABLE_DEPRECATION_WARNINGS  // TODO: CHAOS_IS_CLOTHINGSIMULATIONMESH_ABSTRACT
 	// Update source mesh for this LOD, this is required prior to reset the start pose
-	Cloth->Mesh->Update(Solver, INDEX_NONE, InLODIndex, 0, Offset);
+	Cloth->Mesh->Update(Solver, INDEX_NONE, InLODIndex, 0, ParticleRangeId);
 	
 	// Retrieve the component's scale
 	const Softs::FSolverReal MeshScale = Cloth->Mesh->GetScale();
@@ -145,22 +169,24 @@ PRAGMA_ENABLE_DEPRECATION_WARNINGS
 	// Reset the particles start pose before setting up mass and constraints
 	ResetStartPose(Solver);
 
-	// Build a sim friendly triangle mesh including the solver particle's Offset
-	const int32 NumElements = Indices.Num() / 3;
-	TArray<TVec3<int32>> Elements;
-	Elements.Reserve(NumElements);
-
-	for (int32 i = 0; i < NumElements; ++i)
+	if (!Solver->IsForceBasedSolver())
 	{
-		const int32 Index = 3 * i;
-		Elements.Add(
-			{ static_cast<int32>(Offset + Indices[Index]),
-			 static_cast<int32>(Offset + Indices[Index + 1]),
-			 static_cast<int32>(Offset + Indices[Index + 2]) });
+		const TArray<TVec3<int32>>& Elements = NoOffsetTriangleMesh.GetElements();
+		TArray<TVec3<int32>> OffsetElements;
+		OffsetElements.Reserve(Elements.Num());
+		for (const TVec3<int32>& Element : Elements)
+		{
+			OffsetElements.Add({
+				Element[0] + ParticleRangeId,
+				Element[1] + ParticleRangeId,
+				Element[2] + ParticleRangeId });
+		}
+
+		SolverDatum.OffsetTriangleMesh.Init(MoveTemp(OffsetElements), ParticleRangeId, ParticleRangeId + NumParticles - 1);
+		SolverDatum.OffsetTriangleMesh.GetPointToTriangleMap(); // Builds map for later use by GetPointNormals(), and the velocity fields
 	}
 
-	TriangleMesh.Init(MoveTemp(Elements), Offset, Offset + NumParticles - 1);  // Init with the Offset to avoid discrepancies since the Triangle Mesh only relies on the used indices
-	TriangleMesh.GetPointToTriangleMap(); // Builds map for later use by GetPointNormals(), and the velocity fields
+	const FTriangleMesh& TriangleMesh = Solver->IsForceBasedSolver() ? NoOffsetTriangleMesh : SolverDatum.OffsetTriangleMesh;
 
 	// Initialize the normals, in case the sim data is queried before the simulation steps
 	UpdateNormals(Solver);
@@ -208,25 +234,25 @@ PRAGMA_ENABLE_DEPRECATION_WARNINGS
 	default:
 		check(false);
 	case EMassMode::UniformMass:
-		Solver->SetParticleMassUniform(Offset, MassValue, MinPerParticleMass, TriangleMesh, KinematicPredicate);
+		Solver->SetParticleMassUniform(ParticleRangeId, MassValue, MinPerParticleMass, TriangleMesh, KinematicPredicate);
 		break;
 	case EMassMode::TotalMass:
-		Solver->SetParticleMassFromTotalMass(Offset, MassValue, MinPerParticleMass, TriangleMesh, KinematicPredicate);
+		Solver->SetParticleMassFromTotalMass(ParticleRangeId, MassValue, MinPerParticleMass, TriangleMesh, KinematicPredicate);
 		break;
 	case EMassMode::Density:
-		Solver->SetParticleMassFromDensity(Offset, MassValue, MinPerParticleMass, TriangleMesh, KinematicPredicate);
+		Solver->SetParticleMassFromDensity(ParticleRangeId, MassValue, MinPerParticleMass, TriangleMesh, KinematicPredicate);
 		break;
 	}
 
 	// Setup solver constraints
-	FClothConstraints& ClothConstraints = Solver->GetClothConstraints(Offset);
+	FClothConstraints& ClothConstraints = Solver->GetClothConstraints(ParticleRangeId);
 
 	// Create constraints
 	const bool bEnabled = false;  // Set constraint disabled by default
 	ClothConstraints.AddRules(ConfigProperties, TriangleMesh, &PatternData, WeightMaps, VertexSets, Tethers, MeshScale, bEnabled);
 
 	// Update LOD stats
-	const TConstArrayView<Softs::FSolverReal> InvMasses(Solver->GetParticleInvMasses(Offset), NumParticles);
+	const TConstArrayView<Softs::FSolverReal> InvMasses(Solver->GetParticleInvMasses(ParticleRangeId), NumParticles);
 	NumKinematicParticles = 0;
 	NumDynamicParticles = 0;
 	for (int32 Index = 0; Index < NumParticles; ++Index)
@@ -253,11 +279,11 @@ void FClothingSimulationCloth::FLODData::Update(FClothingSimulationSolver* Solve
 	check(Cloth); 
 	const FSolverData& SolverDatum = SolverData.FindChecked(Solver);
 
-	const int32 Offset = SolverDatum.Offset;
-	check(Offset != INDEX_NONE);
+	const int32 ParticleRangeId = SolverDatum.ParticleRangeId;
+	check(ParticleRangeId != INDEX_NONE);
 
 	// Update the animatable constraint parameters
-	FClothConstraints& ClothConstraints = Solver->GetClothConstraints(Offset);
+	FClothConstraints& ClothConstraints = Solver->GetClothConstraints(ParticleRangeId);
 
 	check(Cloth->Config);
 PRAGMA_DISABLE_DEPRECATION_WARNINGS  // TODO: CHAOS_IS_CLOTHINGSIMULATIONMESH_ABSTRACT
@@ -270,20 +296,20 @@ PRAGMA_ENABLE_DEPRECATION_WARNINGS
 void FClothingSimulationCloth::FLODData::Enable(FClothingSimulationSolver* Solver, bool bEnable) const
 {
 	check(Solver);
-	const int32 Offset = SolverData.FindChecked(Solver).Offset;
-	check(Offset != INDEX_NONE);
+	const int32 ParticleRangeId = SolverData.FindChecked(Solver).ParticleRangeId;
+	check(ParticleRangeId != INDEX_NONE);
 	
 	// Enable particles (and related constraints)
-	Solver->EnableParticles(Offset, bEnable);
+	Solver->EnableParticles(ParticleRangeId, bEnable);
 }
 
 void FClothingSimulationCloth::FLODData::ResetStartPose(FClothingSimulationSolver* Solver) const
 {
 	check(Solver);
-	const int32 Offset = SolverData.FindChecked(Solver).Offset;
-	check(Offset != INDEX_NONE);
+	const int32 ParticleRangeId = SolverData.FindChecked(Solver).ParticleRangeId;
+	check(ParticleRangeId != INDEX_NONE);
 
-	Solver->ResetStartPose(Offset, NumParticles);
+	Solver->ResetStartPose(ParticleRangeId, NumParticles);
 }
 
 void FClothingSimulationCloth::FLODData::UpdateNormals(FClothingSimulationSolver* Solver) const
@@ -291,17 +317,15 @@ void FClothingSimulationCloth::FLODData::UpdateNormals(FClothingSimulationSolver
 	check(Solver);
 
 	const FSolverData& SolverDatum = SolverData.FindChecked(Solver);
-	const int32 Offset = SolverDatum.Offset;
-	const FTriangleMesh& TriangleMesh = SolverDatum.TriangleMesh;
+	const int32 ParticleRangeId = SolverDatum.ParticleRangeId;
 
-	if (Offset != INDEX_NONE)
+	if (ParticleRangeId != INDEX_NONE)
 	{
-		TConstArrayView<Softs::FSolverVec3> Points(Solver->GetParticleXs(Offset) - Offset, Offset + NumParticles);  // TODO: TriangleMesh still uses global array
+		TConstArrayView<Softs::FSolverVec3> Points = Solver->GetParticleXsView(ParticleRangeId);
+		TArrayView<Softs::FSolverVec3> Normals = Solver->GetNormalsView(ParticleRangeId);
 		TArray<Softs::FSolverVec3> FaceNormals;
-		TriangleMesh.GetFaceNormals(FaceNormals, Points, /*ReturnEmptyOnError =*/ false);
-
-		TArrayView<Softs::FSolverVec3> Normals(Solver->GetNormals(Offset), NumParticles);
-		TriangleMesh.GetPointNormals(Normals, TConstArrayView<Softs::FSolverVec3>(FaceNormals), /*bUseGlobalArray =*/ false);
+		NoOffsetTriangleMesh.GetFaceNormals(FaceNormals, Points, /*ReturnEmptyOnError =*/ false);
+		NoOffsetTriangleMesh.GetPointNormals(Normals, TConstArrayView<Softs::FSolverVec3>(FaceNormals), /*bUseGlobalArray =*/ false);
 	}
 }
 
@@ -764,23 +788,38 @@ int32 FClothingSimulationCloth::GetNumParticles(int32 InLODIndex) const
 	return LODData.IsValidIndex(InLODIndex) ? LODData[InLODIndex]->NumParticles : 0;
 }
 
-int32 FClothingSimulationCloth::GetOffset(const FClothingSimulationSolver* Solver, int32 InLODIndex) const
+int32 FClothingSimulationCloth::GetParticleRangeId(const FClothingSimulationSolver* Solver, int32 InLODIndex) const
 {
-	return LODData.IsValidIndex(InLODIndex) ? LODData[InLODIndex]->SolverData.FindChecked(Solver).Offset : 0;
+	return LODData.IsValidIndex(InLODIndex) ? LODData[InLODIndex]->SolverData.FindChecked(Solver).ParticleRangeId : 0;
 }
 
 TVec3<FRealSingle> FClothingSimulationCloth::GetGravity(const FClothingSimulationSolver* Solver) const
 {
 	check(Solver);
-	check(Config);
-	const Softs::FCollectionPropertyFacade& ConfigProperties = Config->GetProperties(GetLODIndex(Solver));
+	if (Solver->IsForceBasedSolver())
+	{
+		const int32 ParticleRangeId = GetParticleRangeId(Solver);
+		if (const Softs::FExternalForces* const ExternalForces = Solver->GetClothConstraints(ParticleRangeId).GetExternalForces().Get())
+		{
+			return ExternalForces->GetGravity();
+		}
+		else
+		{
+			return TVec3<FRealSingle>(0.f, 0.f, Softs::FExternalForces::DefaultGravityZOverride);
+		}
+	}
+	else
+	{
+		check(Config);
+		const Softs::FCollectionPropertyFacade& ConfigProperties = Config->GetProperties(GetLODIndex(Solver));
 
-	const bool bUseGravityOverride = ConfigProperties.GetValue<bool>(TEXT("UseGravityOverride"));
-	const TVec3<FRealSingle> GravityOverride = (TVec3<FRealSingle>)ConfigProperties.GetValue<FVector3f>(TEXT("GravityOverride"), FVector3f(0.f, 0.f, ClothingSimulationClothDefault::GravityZOverride));
-	const FRealSingle GravityScale = (FRealSingle)ConfigProperties.GetValue<float>(TEXT("GravityScale"), 1.f);
-	const FRealSingle GravityMultiplier = (FRealSingle)ClothingSimulationClothConsoleVariables::CVarGravityMultiplier.GetValueOnAnyThread();
+		const bool bUseGravityOverride = ConfigProperties.GetValue<bool>(TEXT("UseGravityOverride"));
+		const TVec3<FRealSingle> GravityOverride = (TVec3<FRealSingle>)ConfigProperties.GetValue<FVector3f>(TEXT("GravityOverride"), FVector3f(0.f, 0.f, ClothingSimulationClothDefault::GravityZOverride));
+		const FRealSingle GravityScale = (FRealSingle)ConfigProperties.GetValue<float>(TEXT("GravityScale"), 1.f);
+		const FRealSingle GravityMultiplier = (FRealSingle)ClothingSimulationClothConsoleVariables::CVarGravityMultiplier.GetValueOnAnyThread();
 
-	return (Solver->IsClothGravityOverrideEnabled() && bUseGravityOverride ? GravityOverride : Solver->GetGravity() * GravityScale) * GravityMultiplier;
+		return (Solver->IsClothGravityOverrideEnabled() && bUseGravityOverride ? GravityOverride : Solver->GetGravity() * GravityScale) * GravityMultiplier;
+	}
 }
 
 FAABB3 FClothingSimulationCloth::CalculateBoundingBox(const FClothingSimulationSolver* Solver) const
@@ -800,10 +839,10 @@ FAABB3 FClothingSimulationCloth::CalculateBoundingBox(const FClothingSimulationS
 	return FAABB3(BoundingBox).TransformedAABB(FRigidTransform3(Solver->GetLocalSpaceLocation(), FRotation3::Identity));
 }
 
-int32 FClothingSimulationCloth::GetOffset(const FClothingSimulationSolver* Solver) const
+int32 FClothingSimulationCloth::GetParticleRangeId(const FClothingSimulationSolver* Solver) const
 {
 	const int32 LODIndex = LODIndices.FindChecked(Solver);
-	return LODData.IsValidIndex(LODIndex) ? GetOffset(Solver, LODIndex) : INDEX_NONE;
+	return LODData.IsValidIndex(LODIndex) ? GetParticleRangeId(Solver, LODIndex) : INDEX_NONE;
 }
 
 int32 FClothingSimulationCloth::GetNumParticles(const FClothingSimulationSolver* Solver) const
@@ -816,7 +855,8 @@ const FTriangleMesh& FClothingSimulationCloth::GetTriangleMesh(const FClothingSi
 {
 	const int32 LODIndex = LODIndices.FindChecked(Solver);
 	static const FTriangleMesh EmptyTriangleMesh;
-	return LODData.IsValidIndex(LODIndex) ? LODData[LODIndex]->SolverData.FindChecked(Solver).TriangleMesh : EmptyTriangleMesh;
+
+	return LODData.IsValidIndex(LODIndex) ? (Solver->IsForceBasedSolver() ? LODData[LODIndex]->NoOffsetTriangleMesh : LODData[LODIndex]->SolverData.FindChecked(Solver).OffsetTriangleMesh) : EmptyTriangleMesh;
 }
 
 // Deprecated for 5.3
@@ -908,26 +948,17 @@ PRAGMA_ENABLE_DEPRECATION_WARNINGS
 	}
 
 	// Update the source mesh skinned positions
-	const int32 PrevOffset = GetOffset(Solver, PrevLODIndex);
-	const int32 Offset = GetOffset(Solver, LODIndex);
-	if (PrevOffset == INDEX_NONE || Offset == INDEX_NONE)
+	const int32 PrevParticleRangeId = GetParticleRangeId(Solver, PrevLODIndex);
+	const int32 ParticleRangeId = GetParticleRangeId(Solver, LODIndex);
+	if (PrevParticleRangeId == INDEX_NONE || ParticleRangeId == INDEX_NONE)
 	{
 		return;
 	}
 
 PRAGMA_DISABLE_DEPRECATION_WARNINGS  // TODO: CHAOS_IS_CLOTHINGSIMULATIONMESH_ABSTRACT
-	Mesh->Update(Solver, PrevLODIndex, LODIndex, PrevOffset, Offset);
+	Mesh->Update(Solver, PrevLODIndex, LODIndex, PrevParticleRangeId, ParticleRangeId);
 PRAGMA_ENABLE_DEPRECATION_WARNINGS
 
-	// Update all LODs dirty properties, since it is easier done than re-updating all properties when switching LODs
-	if (Config->IsLegacySingleLOD())
-	{
-		for (TUniquePtr<FLODData>& LODDatum : LODData)
-		{
-			LODDatum->Update(Solver, this);
-		}
-	}
-
 	// Retrieve config
 	check(Config);
 	const Softs::FCollectionPropertyFacade& ConfigProperties = Config->GetProperties(LODIndex);
@@ -953,15 +984,18 @@ PRAGMA_DISABLE_DEPRECATION_WARNINGS  // TODO: CHAOS_IS_CLOTHINGSIMULATIONMESH_AB
 PRAGMA_ENABLE_DEPRECATION_WARNINGS
 				PrevLODIndex,
 				LODIndex,
-				Solver->GetNormals(PrevOffset),
-				Solver->GetParticlePandInvMs(PrevOffset),
-				Solver->GetParticleVs(PrevOffset),
-				Solver->GetParticlePandInvMs(Offset),
-				Solver->GetParticleXs(Offset),
-				Solver->GetParticleVs(Offset));
-
-			// Update the wind velocity field for the new LOD mesh
-			Solver->SetWindAndPressureGeometry(GroupId, GetTriangleMesh(Solver), ConfigProperties, LODData[LODIndex]->WeightMaps);
+				Solver->GetNormals(PrevParticleRangeId),
+				Solver->GetParticlePandInvMs(PrevParticleRangeId),
+				Solver->GetParticleVs(PrevParticleRangeId),
+				Solver->GetParticlePandInvMs(ParticleRangeId),
+				Solver->GetParticleXs(ParticleRangeId),
+				Solver->GetParticleVs(ParticleRangeId));
+
+			if(!Solver->IsForceBasedSolver())
+			{
+				// Update the wind velocity field for the new LOD mesh
+				Solver->SetWindAndPressureGeometry(GroupId, GetTriangleMesh(Solver), ConfigProperties, LODData[LODIndex]->WeightMaps);
+			}
 		}
 		else
 		{
@@ -970,16 +1004,9 @@ PRAGMA_ENABLE_DEPRECATION_WARNINGS
 		}
 	}
 
-	// Update Cloth group parameters
+	// Update Cloth group parameters  TODO: Cloth groups should exist as their own node object so that they can be used by several cloth objects
 	if (LODIndex != INDEX_NONE)
 	{
-		// Update Cloth group parameters  TODO: Cloth groups should exist as their own node object so that they can be used by several cloth objects
-		// This was updated above for all LODs if single lod config.
-		if (!Config->IsLegacySingleLOD())
-		{
-			LODData[LODIndex]->Update(Solver, this);
-		}
-
 		// TODO: Move all groupID updates out of the cloth update to allow to use of the same GroupId with different cloths
 
 		// Set the reference input velocity and deal with teleport & reset; external forces depends on these values, so they must be initialized before then
@@ -1014,8 +1041,8 @@ PRAGMA_ENABLE_DEPRECATION_WARNINGS
 			OutAngularVelocityScale = ConfigProperties.GetValue<float>(TEXT("AngularVelocityScale"), ClothingSimulationClothDefault::VelocityScale);
 		}
 
+		// NOTE: Force-based solver doesn't actually use FictitiousAngularScale here. It gets it from the property collection directly.
 		const FRealSingle FictitiousAngularScale = ConfigProperties.GetValue<float>(TEXT("FictitiousAngularScale"), ClothingSimulationClothDefault::FictitiousAngularScale);
-
 		Solver->SetReferenceVelocityScale(
 			GroupId,
 			OldReferenceSpaceTransform,
@@ -1023,39 +1050,60 @@ PRAGMA_ENABLE_DEPRECATION_WARNINGS
 			OutLinearVelocityScale,
 			OutAngularVelocityScale,
 			FictitiousAngularScale);
+		if (Solver->IsForceBasedSolver())
+		{
+			Solver->SetProperties(ParticleRangeId, ConfigProperties);
+		}
+		else
+		{
+			// Update gravity
+			// This code relies on the solver gravity property being already set.
+			// In order to use a cloth gravity override, it must first be enabled by the solver so that an override at solver level can still take precedence if needed.
+			// In all cases apart from when the cloth override is used, the gravity scale must be combined to the solver gravity value.
+			Solver->SetGravity(GroupId, GetGravity(Solver));
+
+			// External forces (legacy wind+field)
+			const bool bUsePointBasedWindModel = ConfigProperties.GetValue<bool>(TEXT("UsePointBasedWindModel"));
+			Solver->AddExternalForces(GroupId, bUsePointBasedWindModel);
+
+			const bool bPointBasedWindDisablesAccurateWind = ClothingSimulationClothConsoleVariables::CVarLegacyDisablesAccurateWind.GetValueOnAnyThread();
+			const bool bEnableAerodynamics = !(bUsePointBasedWindModel && bPointBasedWindDisablesAccurateWind);
+			Solver->SetWindAndPressureProperties(GroupId, ConfigProperties, LODData[LODIndex]->WeightMaps, bEnableAerodynamics);
+
+			constexpr float WorldScale = 100.f;  // VelocityField wind is in m/s in the config (same as the wind unit), but cm/s in the solver  TODO: Cleanup the Solver SetWindVelocity functions to be consistent with the unit
+			const FVec3f WindVelocity = ConfigProperties.GetValue<FVector3f>(TEXT("WindVelocity")) * WorldScale;
+			Solver->SetWindVelocity(GroupId, WindVelocity + Solver->GetWindVelocity());
+
+			// Update general solver properties
+			PRAGMA_DISABLE_DEPRECATION_WARNINGS  // TODO: CHAOS_IS_CLOTHINGSIMULATIONMESH_ABSTRACT
+				const Softs::FSolverReal MeshScale = Mesh->GetScale();
+			PRAGMA_ENABLE_DEPRECATION_WARNINGS
+
+				const FRealSingle DampingCoefficient = ConfigProperties.GetValue<float>(TEXT("DampingCoefficient"), ClothingSimulationClothDefault::DampingCoefficient);
+			const FRealSingle LocalDampingCoefficient = ConfigProperties.GetValue<float>(TEXT("LocalDampingCoefficient"));
+			const FRealSingle CollisionThickness = ConfigProperties.GetValue<float>(TEXT("CollisionThickness"), ClothingSimulationClothDefault::CollisionThickness);
+			const FRealSingle FrictionCoefficient = ConfigProperties.GetValue<float>(TEXT("FrictionCoefficient"), ClothingSimulationClothDefault::FrictionCoefficient);
+			Solver->SetProperties(GroupId, DampingCoefficient, LocalDampingCoefficient, CollisionThickness * MeshScale, FrictionCoefficient);
+
+			// Update use of continuous collision detection
+			const bool bUseCCD = ConfigProperties.GetValue<bool>(TEXT("UseCCD"));
+			Solver->SetUseCCD(GroupId, bUseCCD);
+		}
 
-		// Update gravity
-		// This code relies on the solver gravity property being already set.
-		// In order to use a cloth gravity override, it must first be enabled by the solver so that an override at solver level can still take precedence if needed.
-		// In all cases apart from when the cloth override is used, the gravity scale must be combined to the solver gravity value.
-		Solver->SetGravity(GroupId, GetGravity(Solver));
-
-		// External forces (legacy wind+field)
-		const bool bUsePointBasedWindModel = ConfigProperties.GetValue<bool>(TEXT("UsePointBasedWindModel"));
-		Solver->AddExternalForces(GroupId, bUsePointBasedWindModel);
-
-		const bool bPointBasedWindDisablesAccurateWind = ClothingSimulationClothConsoleVariables::CVarLegacyDisablesAccurateWind.GetValueOnAnyThread();
-		const bool bEnableAerodynamics = !(bUsePointBasedWindModel && bPointBasedWindDisablesAccurateWind);
-		Solver->SetWindAndPressureProperties(GroupId, ConfigProperties, LODData[LODIndex]->WeightMaps, bEnableAerodynamics);
-
-		constexpr float WorldScale = 100.f;  // VelocityField wind is in m/s in the config (same as the wind unit), but cm/s in the solver  TODO: Cleanup the Solver SetWindVelocity functions to be consistent with the unit
-		const FVec3f WindVelocity = ConfigProperties.GetValue<FVector3f>(TEXT("WindVelocity")) * WorldScale;
-		Solver->SetWindVelocity(GroupId, WindVelocity + Solver->GetWindVelocity());
-
-		// Update general solver properties
-PRAGMA_DISABLE_DEPRECATION_WARNINGS  // TODO: CHAOS_IS_CLOTHINGSIMULATIONMESH_ABSTRACT
-		const Softs::FSolverReal MeshScale = Mesh->GetScale();
-PRAGMA_ENABLE_DEPRECATION_WARNINGS
-
-		const FRealSingle DampingCoefficient = ConfigProperties.GetValue<float>(TEXT("DampingCoefficient"), ClothingSimulationClothDefault::DampingCoefficient);
-		const FRealSingle LocalDampingCoefficient = ConfigProperties.GetValue<float>(TEXT("LocalDampingCoefficient"));
-		const FRealSingle CollisionThickness = ConfigProperties.GetValue<float>(TEXT("CollisionThickness"), ClothingSimulationClothDefault::CollisionThickness);
-		const FRealSingle FrictionCoefficient = ConfigProperties.GetValue<float>(TEXT("FrictionCoefficient"), ClothingSimulationClothDefault::FrictionCoefficient);
-		Solver->SetProperties(GroupId, DampingCoefficient, LocalDampingCoefficient, CollisionThickness * MeshScale, FrictionCoefficient);
+		// This will be updated below if single legacy lod
+		if (!Config->IsLegacySingleLOD())
+		{
+			LODData[LODIndex]->Update(Solver, this);
+		}
+	}
 
-		// Update use of continuous collision detection
-		const bool bUseCCD = ConfigProperties.GetValue<bool>(TEXT("UseCCD"));
-		Solver->SetUseCCD(GroupId, bUseCCD);
+	// Update all LODs dirty properties, since it is easier done than re-updating all properties when switching LODs
+	if (Config->IsLegacySingleLOD())
+	{
+		for (TUniquePtr<FLODData>& LODDatum : LODData)
+		{
+			LODDatum->Update(Solver, this);
+		}
 	}
 
 	// Reset trigger flags
@@ -1088,48 +1136,48 @@ TConstArrayView<Softs::FSolverVec3> FClothingSimulationCloth::GetAnimationPositi
 {
 	check(Solver);
 	const int32 LODIndex = LODIndices.FindChecked(Solver);
-	check(GetOffset(Solver, LODIndex) != INDEX_NONE);
-	return TConstArrayView<Softs::FSolverVec3>(Solver->GetAnimationPositions(GetOffset(Solver, LODIndex)), GetNumParticles(LODIndex));
+	check(GetParticleRangeId(Solver, LODIndex) != INDEX_NONE);
+	return TConstArrayView<Softs::FSolverVec3>(Solver->GetAnimationPositions(GetParticleRangeId(Solver, LODIndex)), GetNumParticles(LODIndex));
 }
 
 TConstArrayView<Softs::FSolverVec3> FClothingSimulationCloth::GetAnimationNormals(const FClothingSimulationSolver* Solver) const
 {
 	check(Solver);
 	const int32 LODIndex = LODIndices.FindChecked(Solver);
-	check(GetOffset(Solver, LODIndex) != INDEX_NONE);
-	return TConstArrayView<Softs::FSolverVec3>(Solver->GetAnimationNormals(GetOffset(Solver, LODIndex)), GetNumParticles(LODIndex));
+	check(GetParticleRangeId(Solver, LODIndex) != INDEX_NONE);
+	return TConstArrayView<Softs::FSolverVec3>(Solver->GetAnimationNormals(GetParticleRangeId(Solver, LODIndex)), GetNumParticles(LODIndex));
 }
 
 TConstArrayView<Softs::FSolverVec3> FClothingSimulationCloth::GetParticlePositions(const FClothingSimulationSolver* Solver) const
 {
 	check(Solver);
 	const int32 LODIndex = LODIndices.FindChecked(Solver);
-	check(GetOffset(Solver, LODIndex) != INDEX_NONE);
-	return TConstArrayView<Softs::FSolverVec3>(Solver->GetParticleXs(GetOffset(Solver, LODIndex)), GetNumParticles(LODIndex));
+	check(GetParticleRangeId(Solver, LODIndex) != INDEX_NONE);
+	return TConstArrayView<Softs::FSolverVec3>(Solver->GetParticleXs(GetParticleRangeId(Solver, LODIndex)), GetNumParticles(LODIndex));
 }
 
 TConstArrayView<Softs::FSolverVec3> FClothingSimulationCloth::GetParticleNormals(const FClothingSimulationSolver* Solver) const
 {
 	check(Solver);
 	const int32 LODIndex = LODIndices.FindChecked(Solver);
-	check(GetOffset(Solver, LODIndex) != INDEX_NONE);
-	return TConstArrayView<Softs::FSolverVec3>(Solver->GetNormals(GetOffset(Solver, LODIndex)), GetNumParticles(LODIndex));
+	check(GetParticleRangeId(Solver, LODIndex) != INDEX_NONE);
+	return TConstArrayView<Softs::FSolverVec3>(Solver->GetNormals(GetParticleRangeId(Solver, LODIndex)), GetNumParticles(LODIndex));
 }
 
 TConstArrayView<Softs::FSolverVec3> FClothingSimulationCloth::GetParticleVelocities(const FClothingSimulationSolver* Solver) const
 {
 	check(Solver);
 	const int32 LODIndex = LODIndices.FindChecked(Solver);
-	check(GetOffset(Solver, LODIndex) != INDEX_NONE);
-	return TConstArrayView<Softs::FSolverVec3>(Solver->GetParticleVs(GetOffset(Solver, LODIndex)), GetNumParticles(LODIndex));
+	check(GetParticleRangeId(Solver, LODIndex) != INDEX_NONE);
+	return TConstArrayView<Softs::FSolverVec3>(Solver->GetParticleVs(GetParticleRangeId(Solver, LODIndex)), GetNumParticles(LODIndex));
 }
 
 TConstArrayView<Softs::FSolverReal> FClothingSimulationCloth::GetParticleInvMasses(const FClothingSimulationSolver* Solver) const
 {
 	check(Solver);
 	const int32 LODIndex = LODIndices.FindChecked(Solver);
-	check(GetOffset(Solver, LODIndex) != INDEX_NONE);
-	return TConstArrayView<Softs::FSolverReal>(Solver->GetParticleInvMasses(GetOffset(Solver, LODIndex)), GetNumParticles(LODIndex));
+	check(GetParticleRangeId(Solver, LODIndex) != INDEX_NONE);
+	return TConstArrayView<Softs::FSolverReal>(Solver->GetParticleInvMasses(GetParticleRangeId(Solver, LODIndex)), GetNumParticles(LODIndex));
 }
 
 void FClothingSimulationCloth::SetMaterialProperties(const TVec2<FRealSingle>& InEdgeStiffness, const TVec2<FRealSingle>& InBendingStiffness, const TVec2<FRealSingle>& InAreaStiffness)
diff --git a/Engine/Plugins/Experimental/ChaosCloth/Source/Chaos/Private/ChaosCloth/ChaosClothingSimulationCollider.cpp b/Engine/Plugins/Experimental/ChaosCloth/Source/Chaos/Private/ChaosCloth/ChaosClothingSimulationCollider.cpp
--- a/Engine/Plugins/Experimental/ChaosCloth/Source/Chaos/Private/ChaosCloth/ChaosClothingSimulationCollider.cpp
+++ b/Engine/Plugins/Experimental/ChaosCloth/Source/Chaos/Private/ChaosCloth/ChaosClothingSimulationCollider.cpp
@@ -35,7 +35,7 @@ struct FClothingSimulationCollider::FLODData
 {
 	FClothCollisionData ClothCollisionData;
 	int32 NumGeometries;  // Number of collision bodies
-	TMap<FSolverClothPair, int32> Offsets;  // Solver particle offset
+	TMap<FSolverClothPair, int32> CollisionRangeIds;  // Solver particle range ids
 
 	FLODData() : NumGeometries(0) {}
 
@@ -95,7 +95,7 @@ void FClothingSimulationCollider::FLODData::Add(
 		}
 	}
 
-	const uint32 NumConvexes = ClothCollisionData.Convexes.Num();
+	const int32 NumConvexes = ClothCollisionData.Convexes.Num();
 	const int32 NumBoxes = ClothCollisionData.Boxes.Num();
 	const int32 NumLevelSets = InLevelSetCollisionData.Num();
 	const int32 NumSkinnedLevelSets = InSkinnedLevelSetCollisionData.Num();
@@ -106,25 +106,29 @@ void FClothingSimulationCollider::FLODData::Add(
 	}
 
 	NumGeometries = NumSpheres + NumCapsules + NumConvexes + NumBoxes + NumLevelSets + NumSkinnedLevelSets + NumSkinnedLevelSetBones;
+	if(NumGeometries == 0)
+	{
+		return;
+	}
 
 	// Retrieve cloth group Id, or use INDEX_NONE if this collider applies to all cloths (when Cloth == nullptr)
 	const uint32 GroupId = Cloth ? Cloth->GetGroupId() : INDEX_NONE;
 
 	// The offset will be set to the first collision particle's index
 	// Try to reuse existing offsets when Add is called during the collider update (ie Offset isn't INDEX_NONE)
-	int32* const OffsetPtr = Offsets.Find(FSolverClothPair(Solver, Cloth));
-	const bool bIsNewCollider = !OffsetPtr;
-	int32& Offset = bIsNewCollider ? Offsets.Add(FSolverClothPair(Solver, Cloth)) : *OffsetPtr;
-	Offset = Solver->AddCollisionParticles(NumGeometries, GroupId, bIsNewCollider ? INDEX_NONE : Offset);
+	int32* const CollisionRangeIdPtr = CollisionRangeIds.Find(FSolverClothPair(Solver, Cloth));
+	const bool bIsNewCollider = !CollisionRangeIdPtr;
+	int32& CollisionRangeId = bIsNewCollider ? CollisionRangeIds.Add(FSolverClothPair(Solver, Cloth)) : *CollisionRangeIdPtr;
+	CollisionRangeId = Solver->AddCollisionParticles(NumGeometries, GroupId, bIsNewCollider ? INDEX_NONE : CollisionRangeId);
+
+	TArrayView<int32> BoneIndices = Solver->GetCollisionBoneIndicesView(CollisionRangeId);
+	TArrayView<Softs::FSolverRigidTransform3> BaseTransforms = Solver->GetCollisionBaseTransformsView(CollisionRangeId);
 
 	// Capsules
-	const int32 CapsuleOffset = Offset;
+	const int32 CapsuleIndexOffset = 0;
 	if (NumCapsules)
 	{
-		int32* const BoneIndices = Solver->GetCollisionBoneIndices(CapsuleOffset);
-		Softs::FSolverRigidTransform3* const BaseTransforms = Solver->GetCollisionBaseTransforms(CapsuleOffset);
-
-		for (int32 Index = 0; Index < NumCapsules; ++Index)
+		for (int32 Index = CapsuleIndexOffset; Index < CapsuleIndexOffset + NumCapsules; ++Index)
 		{
 			const FClothCollisionPrim_SphereConnection& Connection = ClothCollisionData.SphereConnections[Index];
 
@@ -158,20 +162,20 @@ void FClothingSimulationCollider::FLODData::Add(
 			if (Axis.SizeSquared() < SMALL_NUMBER)
 			{
 				// Sphere
-				Solver->SetCollisionGeometry(CapsuleOffset, Index,
+				Solver->SetCollisionGeometry(CollisionRangeId, Index,
 					MakeImplicitObjectPtr<FSphere>(Center, MaxRadius));
 			}
 			else if (MaxRadius - MinRadius < KINDA_SMALL_NUMBER)
 			{
 				// Capsule
-				Solver->SetCollisionGeometry(CapsuleOffset, Index, 
+				Solver->SetCollisionGeometry(CollisionRangeId, Index,
 					MakeImplicitObjectPtr<FCapsule>(P0, P1, MaxRadius));
 			}
 			else
 			{
 				if (ClothingSimulationColliderConsoleVariables::CVarUseOptimizedTaperedCapsule.GetValueOnAnyThread())
 				{
-					Solver->SetCollisionGeometry(CapsuleOffset, Index,
+					Solver->SetCollisionGeometry(CollisionRangeId, Index,
 						MakeImplicitObjectPtr<FTaperedCapsule>(P0, P1, Radius0, Radius1));
 				}
 				else
@@ -182,20 +186,17 @@ void FClothingSimulationCollider::FLODData::Add(
 					Objects.Add(MakeImplicitObjectPtr<FTaperedCylinder>(P0, P1, Radius0, Radius1));
 					Objects.Add(MakeImplicitObjectPtr<FSphere>(P0, Radius0));
 					Objects.Add(MakeImplicitObjectPtr<FSphere>(P1, Radius1));
-					Solver->SetCollisionGeometry(CapsuleOffset, Index, MakeImplicitObjectPtr<FImplicitObjectUnion>(MoveTemp(Objects)));
+					Solver->SetCollisionGeometry(CollisionRangeId, Index, MakeImplicitObjectPtr<FImplicitObjectUnion>(MoveTemp(Objects)));
 				}
 			}
 		}
 	}
 
 	// Spheres
-	const int32 SphereOffset = CapsuleOffset + NumCapsules;
+	const int32 SphereIndexOffset = CapsuleIndexOffset + NumCapsules;
 	if (NumSpheres != 0)
 	{
-		int32* const BoneIndices = Solver->GetCollisionBoneIndices(SphereOffset);
-		Softs::FSolverRigidTransform3* const BaseTransforms = Solver->GetCollisionBaseTransforms(SphereOffset);
-
-		for (int32 Index = 0, SphereIndex = 0; SphereIndex < ClothCollisionData.Spheres.Num(); ++SphereIndex)
+		for (int32 Index = SphereIndexOffset, SphereIndex = 0; SphereIndex < ClothCollisionData.Spheres.Num(); ++SphereIndex)
 		{
 			// Skip spheres that are the end caps of capsules.
 			if (CapsuleEnds[SphereIndex])
@@ -211,7 +212,7 @@ void FClothingSimulationCollider::FLODData::Add(
 
 			BaseTransforms[Index] = Softs::FSolverRigidTransform3::Identity;
 
-			Solver->SetCollisionGeometry(SphereOffset, Index,
+			Solver->SetCollisionGeometry(CollisionRangeId, Index,
 				MakeImplicitObjectPtr<FSphere>(
 					Sphere.LocalPosition * InScale,
 					Sphere.Radius * InScale));
@@ -221,15 +222,12 @@ void FClothingSimulationCollider::FLODData::Add(
 	}
 
 	// Convexes
-	const int32 ConvexOffset = SphereOffset + NumSpheres;
+	const int32 ConvexIndexOffset = SphereIndexOffset + NumSpheres;
 	if (NumConvexes != 0)
 	{
-		int32* const BoneIndices = Solver->GetCollisionBoneIndices(ConvexOffset);
-		Softs::FSolverRigidTransform3* const BaseTransforms = Solver->GetCollisionBaseTransforms(ConvexOffset);
-
-		for (uint32 Index = 0; Index < NumConvexes; ++Index)
+		for (int32 Index = ConvexIndexOffset, ConvexIndex = 0; ConvexIndex < NumConvexes; ++Index, ++ConvexIndex)
 		{
-			const FClothCollisionPrim_Convex& Convex = ClothCollisionData.Convexes[Index];
+			const FClothCollisionPrim_Convex& Convex = ClothCollisionData.Convexes[ConvexIndex];
 
 			// Always initialize the collision particle transforms before setting any geometry as otherwise NaNs gets detected during the bounding box updates
 			BaseTransforms[Index] = Softs::FSolverRigidTransform3::Identity;
@@ -290,26 +288,23 @@ void FClothingSimulationCollider::FLODData::Add(
 				}
 
 				// Setup the collision particle geometry
-				Solver->SetCollisionGeometry(ConvexOffset, Index, MakeImplicitObjectPtr<FConvex>(MoveTemp(Planes), MoveTemp(FaceIndices), MoveTemp(Vertices)));
+				Solver->SetCollisionGeometry(CollisionRangeId, Index, MakeImplicitObjectPtr<FConvex>(MoveTemp(Planes), MoveTemp(FaceIndices), MoveTemp(Vertices)));
 			}
 			else
 			{
 				UE_LOG(LogChaosCloth, Warning, TEXT("Replacing invalid convex collision by a default unit sphere."));
-				Solver->SetCollisionGeometry(ConvexOffset, Index, MakeImplicitObjectPtr<FSphere>(FVec3(0.0f), 1.0f));  // Default to a unit sphere to replace the faulty convex
+				Solver->SetCollisionGeometry(CollisionRangeId, Index, MakeImplicitObjectPtr<FSphere>(FVec3(0.0f), 1.0f));  // Default to a unit sphere to replace the faulty convex
 			}
 		}
 	}
 
 	// Boxes
-	const int32 BoxOffset = ConvexOffset + NumConvexes;
+	const int32 BoxIndexOffset = ConvexIndexOffset + NumConvexes;
 	if (NumBoxes != 0)
-	{
-		int32* const BoneIndices = Solver->GetCollisionBoneIndices(BoxOffset);
-		Softs::FSolverRigidTransform3* const BaseTransforms = Solver->GetCollisionBaseTransforms(BoxOffset);
-		
-		for (int32 Index = 0; Index < NumBoxes; ++Index)
+	{		
+		for (int32 Index = BoxIndexOffset, BoxIndex = 0; BoxIndex < NumBoxes; ++Index, ++BoxIndex)
 		{
-			const FClothCollisionPrim_Box& Box = ClothCollisionData.Boxes[Index];
+			const FClothCollisionPrim_Box& Box = ClothCollisionData.Boxes[BoxIndex];
 			
 			BaseTransforms[Index] = Softs::FSolverRigidTransform3(Box.LocalPosition, Box.LocalRotation);
 			
@@ -317,43 +312,38 @@ void FClothingSimulationCollider::FLODData::Add(
 			UE_LOG(LogChaosCloth, VeryVerbose, TEXT("Found collision box on bone index %d."), BoneIndices[Index]);
 
 			const FVec3 HalfExtents = Box.HalfExtents * InScale;
-			Solver->SetCollisionGeometry(BoxOffset, Index, MakeImplicitObjectPtr<TBox<FReal, 3>>(-HalfExtents, HalfExtents));
+			Solver->SetCollisionGeometry(CollisionRangeId, Index, MakeImplicitObjectPtr<TBox<FReal, 3>>(-HalfExtents, HalfExtents));
 		}
 	}
 
-	const int32 LevelSetOffset = BoxOffset + NumBoxes;
+	const int32 LevelSetIndexOffset = BoxIndexOffset + NumBoxes;
 	if (NumLevelSets != 0)
 	{
-		int32* const BoneIndices = Solver->GetCollisionBoneIndices(LevelSetOffset);
-		Softs::FSolverRigidTransform3* const BaseTransforms = Solver->GetCollisionBaseTransforms(LevelSetOffset);
-
-		for (int32 Index = 0; Index < NumLevelSets; ++Index)
+		for (int32 Index = LevelSetIndexOffset, LevelSetIndex = 0; LevelSetIndex < NumLevelSets; ++Index, ++LevelSetIndex)
 		{
 			// Always initialize the collision particle transforms before setting any geometry as otherwise NaNs gets detected during the bounding box updates
 			BaseTransforms[Index] = Softs::FSolverRigidTransform3::Identity;
 
-			BoneIndices[Index] = GetMappedBoneIndex(UsedBoneIndices, InLevelSetCollisionData[Index].BoneIndex);
+			BoneIndices[Index] = GetMappedBoneIndex(UsedBoneIndices, InLevelSetCollisionData[LevelSetIndex].BoneIndex);
 			UE_LOG(LogChaosCloth, VeryVerbose, TEXT("Found collision level set on bone index %d."), BoneIndices[Index]);
 
 			// Setup the collision particle geometry
 			FImplicitObjectPtr TransformedLevelSet = MakeImplicitObjectPtr<TImplicitObjectTransformed<FReal, 3>>(InLevelSetCollisionData[Index].LevelSet->DeepCopyGeometry(), 
-				TRigidTransform<FReal, 3>(InLevelSetCollisionData[Index].Transform));
-			Solver->SetCollisionGeometry(LevelSetOffset, Index, MoveTemp(TransformedLevelSet));
+				TRigidTransform<FReal, 3>(InLevelSetCollisionData[LevelSetIndex].Transform));
+			Solver->SetCollisionGeometry(CollisionRangeId, Index, MoveTemp(TransformedLevelSet));
 		}
 	}
 
-	const int32 SkinnedLevelSetOffset = LevelSetOffset + NumLevelSets;
+	const int32 SkinnedLevelSetIndexOffset = LevelSetIndexOffset + NumLevelSets;
 	if (NumSkinnedLevelSets != 0)
 	{
-		int32* const BoneIndices = Solver->GetCollisionBoneIndices(SkinnedLevelSetOffset);
-		Softs::FSolverRigidTransform3* const BaseTransforms = Solver->GetCollisionBaseTransforms(SkinnedLevelSetOffset);
-
-		int32 Index = 0;
+		int32 Index = SkinnedLevelSetIndexOffset;
 		for (const FSkinnedLevelSetCollisionData& SkinnedCollisionData : InSkinnedLevelSetCollisionData)
 		{
 			// Add bone proxies first so they get updated first?
 			TArray<int32> SolverBoneIndices;
 			SolverBoneIndices.Reserve(SkinnedCollisionData.MappedSkinnedBones.Num());
+			const int32 GlobalBoneOffset = Solver->IsForceBasedSolver() ? 0 : CollisionRangeId;
 			for (int32 MappedSubBoneIndex : SkinnedCollisionData.MappedSkinnedBones)
 			{
 				// Always initialize the collision particle transforms before setting any geometry as otherwise NaNs gets detected during the bounding box updates
@@ -363,9 +353,9 @@ void FClothingSimulationCollider::FLODData::Add(
 				UE_LOG(LogChaosCloth, VeryVerbose, TEXT("Found collision skinned level set sub-bone on bone index %d."), BoneIndices[Index]);
 
 				FImplicitObjectPtr BoneProxy = MakeImplicitObjectPtr<FWeightedLatticeBoneProxy>();
-				Solver->SetCollisionGeometry(SkinnedLevelSetOffset, Index, MoveTemp(BoneProxy));
+				Solver->SetCollisionGeometry(CollisionRangeId, Index, MoveTemp(BoneProxy));
 
-				SolverBoneIndices.Add(Index + SkinnedLevelSetOffset);
+				SolverBoneIndices.Add(GlobalBoneOffset + Index);
 
 				++Index;
 			}
@@ -379,18 +369,17 @@ void FClothingSimulationCollider::FLODData::Add(
 			FImplicitObjectPtr SkinnedLevelSet = SkinnedCollisionData.WeightedLevelSet->DeepCopyGeometry();
 			SkinnedLevelSet->GetObjectChecked<TWeightedLatticeImplicitObject<FLevelSet>>().SetSolverBoneIndices(MoveTemp(SolverBoneIndices));
 
-			Solver->SetCollisionGeometry(SkinnedLevelSetOffset, Index, MoveTemp(SkinnedLevelSet));
+			Solver->SetCollisionGeometry(CollisionRangeId, Index, MoveTemp(SkinnedLevelSet));
 
 			++Index;
 		}
 	}
-
 	UE_LOG(LogChaosCloth, VeryVerbose, TEXT("Added collisions: %d spheres, %d capsules, %d convexes, %d boxes, %d level sets."), NumSpheres, NumCapsules, NumConvexes, NumBoxes, NumLevelSets);
 }
 
 void FClothingSimulationCollider::FLODData::Remove(FClothingSimulationSolver* Solver, FClothingSimulationCloth* Cloth)
 {
-	Offsets.Remove(FSolverClothPair(Solver, Cloth));
+	CollisionRangeIds.Remove(FSolverClothPair(Solver, Cloth));
 }
 
 void FClothingSimulationCollider::FLODData::Update(
@@ -403,10 +392,10 @@ void FClothingSimulationCollider::FLODData::Update(
 	check(Cloth);
 	if (NumGeometries)
 	{
-		const int32 Offset = Offsets.FindChecked(FSolverClothPair(Solver, Cloth));
-		const int32* const BoneIndices = Solver->GetCollisionBoneIndices(Offset);
-		const Softs::FSolverRigidTransform3* BaseTransforms = Solver->GetCollisionBaseTransforms(Offset);
-		Softs::FSolverRigidTransform3* const CollisionTransforms = Solver->GetCollisionTransforms(Offset);
+		const int32 CollisionRangeId = CollisionRangeIds.FindChecked(FSolverClothPair(Solver, Cloth));
+		TConstArrayView<int32> BoneIndices = Solver->GetCollisionBoneIndicesView(CollisionRangeId);
+		TConstArrayView<Softs::FSolverRigidTransform3> BaseTransforms = Solver->GetCollisionBaseTransformsView(CollisionRangeId);
+		TArrayView<Softs::FSolverRigidTransform3> CollisionTransforms = Solver->GetCollisionTransformsView(CollisionRangeId);
 
 		FTransform ComponentToLocalSpaceReal = ComponentTransform;
 		ComponentToLocalSpaceReal.AddToTranslation(-Solver->GetLocalSpaceLocation());
@@ -429,8 +418,8 @@ void FClothingSimulationCollider::FLODData::Enable(FClothingSimulationSolver* So
 	check(Cloth);
 	if (NumGeometries)
 	{
-		const int32 Offset = Offsets.FindChecked(FSolverClothPair(Solver, Cloth));
-		Solver->EnableCollisionParticles(Offset, bEnable);
+		const int32 CollisionRangeId = CollisionRangeIds.FindChecked(FSolverClothPair(Solver, Cloth));
+		Solver->EnableCollisionParticles(CollisionRangeId, bEnable);
 	}
 }
 
@@ -440,8 +429,8 @@ void FClothingSimulationCollider::FLODData::ResetStartPose(FClothingSimulationSo
 	check(Cloth);
 	if (NumGeometries)
 	{
-		const int32 Offset = Offsets.FindChecked(FSolverClothPair(Solver, Cloth));
-		Solver->ResetCollisionStartPose(Offset, NumGeometries);
+		const int32 CollisionRangeId = CollisionRangeIds.FindChecked(FSolverClothPair(Solver, Cloth));
+		Solver->ResetCollisionStartPose(CollisionRangeId, NumGeometries);
 	}
 }
 
@@ -695,15 +684,25 @@ int32 FClothingSimulationCollider::GetNumGeometries(int32 InSlotIndex) const
 	return LODData.IsValidIndex(InSlotIndex) ? LODData[InSlotIndex]->NumGeometries : 0;
 }
 
-int32 FClothingSimulationCollider::GetOffset(const FClothingSimulationSolver* Solver, const FClothingSimulationCloth* Cloth, int32 InSlotIndex) const
+int32 FClothingSimulationCollider::GetCollisionRangeId(const FClothingSimulationSolver* Solver, const FClothingSimulationCloth* Cloth, int32 InSlotIndex) const
+{
+	const int32* const CollisionRangeId = LODData.IsValidIndex(InSlotIndex) ? LODData[InSlotIndex]->CollisionRangeIds.Find(FSolverClothPair(Solver, Cloth)) : nullptr;
+	return CollisionRangeId ? *CollisionRangeId : INDEX_NONE;
+}
+
+int32 FClothingSimulationCollider::GetCollisionRangeId(const FClothingSimulationSolver* Solver, const FClothingSimulationCloth* Cloth, ECollisionDataType CollisionDataType) const
 {
-	const int32* const Offset = LODData.IsValidIndex(InSlotIndex) ? LODData[InSlotIndex]->Offsets.Find(FSolverClothPair(Solver, Cloth)) : nullptr;
-	return Offset ? *Offset : INDEX_NONE;
+	const int32 LODIndex = LODIndices.FindChecked(FSolverClothPair(Solver, Cloth));
+	const int32 SlotIndex =
+		(CollisionDataType < ECollisionDataType::LODs) ? (int32)CollisionDataType :
+		(LODIndex >= (int32)ECollisionDataType::LODs) ? LODIndex : INDEX_NONE;
+
+	return GetCollisionRangeId(Solver, Cloth, SlotIndex);
 }
 
-bool FClothingSimulationCollider::GetOffsetAndNumGeometries(const FClothingSimulationSolver* Solver, const FClothingSimulationCloth* Cloth, ECollisionDataType CollisionDataType, int32& OutOffset, int32& OutNumGeometries) const
+bool FClothingSimulationCollider::GetCollisionRangeIdAndNumGeometries(const FClothingSimulationSolver* Solver, const FClothingSimulationCloth* Cloth, ECollisionDataType CollisionDataType, int32& OutCollisionRangeId, int32& OutNumGeometries) const
 {
-	OutOffset = INDEX_NONE;
+	OutCollisionRangeId = INDEX_NONE;
 	OutNumGeometries = 0;
 
 	const int32 LODIndex = LODIndices.FindChecked(FSolverClothPair(Solver, Cloth));
@@ -713,11 +712,11 @@ bool FClothingSimulationCollider::GetOffsetAndNumGeometries(const FClothingSimul
 
 	if (LODData.IsValidIndex(SlotIndex))
 	{
-		OutOffset = LODData[SlotIndex]->Offsets.FindChecked(FSolverClothPair(Solver, Cloth));
+		OutCollisionRangeId = LODData[SlotIndex]->CollisionRangeIds.FindChecked(FSolverClothPair(Solver, Cloth));
 		OutNumGeometries = LODData[SlotIndex]->NumGeometries;
 	}
 
-	return OutOffset != INDEX_NONE && OutNumGeometries > 0;
+	return OutCollisionRangeId != INDEX_NONE && OutNumGeometries > 0;
 }
 
 void FClothingSimulationCollider::Add(FClothingSimulationSolver* Solver, FClothingSimulationCloth* Cloth)
@@ -775,7 +774,7 @@ void FClothingSimulationCollider::PreUpdate(FClothingSimulationSolver* Solver, F
 
 	// Add or re-add the external collision particles
 	const int32 ExternalCollisionNumGeometries = GetNumGeometries((int32)ECollisionDataType::External);
-	const int32 ExternalCollisionOffset = GetOffset(Solver, Cloth, (int32)ECollisionDataType::External);
+	const int32 ExternalCollisionCollisionRangeId = GetCollisionRangeId(Solver, Cloth, (int32)ECollisionDataType::External);
 
 	// TODO: Get level sets?
 	const TArray<FLevelSetCollisionData> LevelSetCollisions;
@@ -786,7 +785,7 @@ void FClothingSimulationCollider::PreUpdate(FClothingSimulationSolver* Solver, F
 	// TODO: Find a better way in case the same number but different collisions are being re-added (hash collision data? Provide user dirty function?)
 	bHasExternalCollisionChanged =
 		ExternalCollisionNumGeometries != GetNumGeometries((int32)ECollisionDataType::External) ||
-		ExternalCollisionOffset != GetOffset(Solver, Cloth, (int32)ECollisionDataType::External);
+		ExternalCollisionCollisionRangeId != GetCollisionRangeId(Solver, Cloth, (int32)ECollisionDataType::External);
 }
 
 void FClothingSimulationCollider::Update(FClothingSimulationSolver* Solver, FClothingSimulationCloth* Cloth)
@@ -869,10 +868,8 @@ TConstArrayView<FSolverVec3> FClothingSimulationCollider::GetCollisionTranslatio
 	check(Solver);
 	check(Cloth);
 
-	int32 Offset, NumGeometries;
-	return GetOffsetAndNumGeometries(Solver, Cloth, CollisionDataType, Offset, NumGeometries) ?
-		TConstArrayView<FSolverVec3>(Solver->GetCollisionParticleXs(Offset), NumGeometries) :
-		TConstArrayView<FSolverVec3>();
+	const int32 CollisionRangeId = GetCollisionRangeId(Solver, Cloth, CollisionDataType);
+	return CollisionRangeId != INDEX_NONE ? Solver->GetCollisionParticleXsView(CollisionRangeId) : TConstArrayView<FSolverVec3>();
 }
 
 TConstArrayView<Softs::FSolverRotation3> FClothingSimulationCollider::GetCollisionRotations(const FClothingSimulationSolver* Solver, const FClothingSimulationCloth* Cloth, ECollisionDataType CollisionDataType) const
@@ -880,10 +877,8 @@ TConstArrayView<Softs::FSolverRotation3> FClothingSimulationCollider::GetCollisi
 	check(Solver);
 	check(Cloth);
 
-	int32 Offset, NumGeometries;
-	return GetOffsetAndNumGeometries(Solver, Cloth, CollisionDataType, Offset, NumGeometries) ?
-		TConstArrayView<Softs::FSolverRotation3>(Solver->GetCollisionParticleRs(Offset), NumGeometries) :
-		TConstArrayView<Softs::FSolverRotation3>();
+	const int32 CollisionRangeId = GetCollisionRangeId(Solver, Cloth, CollisionDataType);
+	return CollisionRangeId != INDEX_NONE ? Solver->GetCollisionParticleRsView(CollisionRangeId) : TConstArrayView<Softs::FSolverRotation3>();
 }
 
 TConstArrayView<Softs::FSolverRigidTransform3> FClothingSimulationCollider::GetOldCollisionTransforms(const FClothingSimulationSolver* Solver, const FClothingSimulationCloth* Cloth, ECollisionDataType CollisionDataType) const
@@ -891,10 +886,8 @@ TConstArrayView<Softs::FSolverRigidTransform3> FClothingSimulationCollider::GetO
 	check(Solver);
 	check(Cloth);
 
-	int32 Offset, NumGeometries;
-	return GetOffsetAndNumGeometries(Solver, Cloth, CollisionDataType, Offset, NumGeometries) ?
-		TConstArrayView<Softs::FSolverRigidTransform3>(Solver->GetOldCollisionTransforms(Offset), NumGeometries) :
-		TConstArrayView<Softs::FSolverRigidTransform3>();
+	const int32 CollisionRangeId = GetCollisionRangeId(Solver, Cloth, CollisionDataType);
+	return CollisionRangeId != INDEX_NONE ? Solver->GetOldCollisionTransformsView(CollisionRangeId) : TConstArrayView<Softs::FSolverRigidTransform3>();
 }
 
 TConstArrayView<FImplicitObjectPtr> FClothingSimulationCollider::GetCollisionGeometry(const FClothingSimulationSolver* Solver, const FClothingSimulationCloth* Cloth, ECollisionDataType CollisionDataType) const
@@ -902,10 +895,8 @@ TConstArrayView<FImplicitObjectPtr> FClothingSimulationCollider::GetCollisionGeo
 	check(Solver);
 	check(Cloth);
 
-	int32 Offset, NumGeometries;
-	return GetOffsetAndNumGeometries(Solver, Cloth, CollisionDataType, Offset, NumGeometries) ?
-		TConstArrayView<FImplicitObjectPtr>(Solver->GetCollisionGeometry(Offset), NumGeometries) :
-		TConstArrayView<FImplicitObjectPtr>();
+	const int32 CollisionRangeId = GetCollisionRangeId(Solver, Cloth, CollisionDataType);
+	return CollisionRangeId != INDEX_NONE ? Solver->GetCollisionGeometryView(CollisionRangeId) : TConstArrayView<FImplicitObjectPtr>();
 }
 
 TConstArrayView<TUniquePtr<FImplicitObject>> FClothingSimulationCollider::GetCollisionGeometries(const FClothingSimulationSolver* Solver, const FClothingSimulationCloth* Cloth, ECollisionDataType CollisionDataType) const
@@ -919,10 +910,8 @@ TConstArrayView<bool> FClothingSimulationCollider::GetCollisionStatus(const FClo
 	check(Solver);
 	check(Cloth);
 
-	int32 Offset, NumGeometries;
-	return GetOffsetAndNumGeometries(Solver, Cloth, CollisionDataType, Offset, NumGeometries) ?
-		TConstArrayView<bool>(Solver->GetCollisionStatus(Offset), NumGeometries) :
-		TConstArrayView<bool>();
+	const int32 CollisionRangeId = GetCollisionRangeId(Solver, Cloth, CollisionDataType);
+	return CollisionRangeId != INDEX_NONE ? Solver->GetCollisionStatusView(CollisionRangeId) : TConstArrayView<bool>();
 }
 
 }  // End namespace Chaos
diff --git a/Engine/Plugins/Experimental/ChaosCloth/Source/Chaos/Private/ChaosCloth/ChaosClothingSimulationMesh.cpp b/Engine/Plugins/Experimental/ChaosCloth/Source/Chaos/Private/ChaosCloth/ChaosClothingSimulationMesh.cpp
--- a/Engine/Plugins/Experimental/ChaosCloth/Source/Chaos/Private/ChaosCloth/ChaosClothingSimulationMesh.cpp
+++ b/Engine/Plugins/Experimental/ChaosCloth/Source/Chaos/Private/ChaosCloth/ChaosClothingSimulationMesh.cpp
@@ -480,8 +480,8 @@ void FClothingSimulationMesh::Update(
 	FClothingSimulationSolver* Solver,
 	int32 PrevLODIndex,
 	int32 LODIndex,
-	int32 PrevOffset,
-	int32 Offset)
+	int32 PrevParticleRangeId,
+	int32 ParticleRangeId)
 {
 	check(Solver);
 
@@ -493,8 +493,8 @@ void FClothingSimulationMesh::Update(
 
 	// Skin current LOD positions
 	const FVec3& LocalSpaceLocation = Solver->GetLocalSpaceLocation();
-	FSolverVec3* const OutPositions = Solver->GetAnimationPositions(Offset);
-	FSolverVec3* const OutNormals = Solver->GetAnimationNormals(Offset);
+	FSolverVec3* const OutPositions = Solver->GetAnimationPositions(ParticleRangeId);
+	FSolverVec3* const OutNormals = Solver->GetAnimationNormals(ParticleRangeId);
 	
 	SkinPhysicsMesh(LODIndex, LocalSpaceLocation, OutPositions, OutNormals);
 
@@ -502,9 +502,9 @@ void FClothingSimulationMesh::Update(
 	if (LODIndex != PrevLODIndex)
 	{
 		// TODO: Using the more accurate skinning method here would require double buffering the context at the skeletal mesh level
-		const FSolverVec3* const SrcWrapNormals = Solver->GetAnimationNormals(PrevOffset);  // No need to keep an old normals array around, since the LOD has just changed
-		const FSolverVec3* const SrcWrapPositions = Solver->GetOldAnimationPositions(PrevOffset);
-		FSolverVec3* const OutOldPositions = Solver->GetOldAnimationPositions(Offset);
+		const FSolverVec3* const SrcWrapNormals = Solver->GetAnimationNormals(PrevParticleRangeId);  // No need to keep an old normals array around, since the LOD has just changed
+		const FSolverVec3* const SrcWrapPositions = Solver->GetOldAnimationPositions(PrevParticleRangeId);
+		FSolverVec3* const OutOldPositions = Solver->GetOldAnimationPositions(ParticleRangeId);
 
 		const bool bValidWrap = WrapDeformLOD(PrevLODIndex, LODIndex, SrcWrapNormals, SrcWrapPositions, OutOldPositions);
 	
