Update the Cloth solver plugin to optionally use the new Softs::FEvolution.

diff --git a/Engine/Plugins/Experimental/ChaosCloth/Source/Chaos/Private/ChaosCloth/ChaosClothingSimulationSolver.cpp b/Engine/Plugins/Experimental/ChaosCloth/Source/Chaos/Private/ChaosCloth/ChaosClothingSimulationSolver.cpp
--- a/Engine/Plugins/Experimental/ChaosCloth/Source/Chaos/Private/ChaosCloth/ChaosClothingSimulationSolver.cpp
+++ b/Engine/Plugins/Experimental/ChaosCloth/Source/Chaos/Private/ChaosCloth/ChaosClothingSimulationSolver.cpp
@@ -6,6 +6,7 @@
 #include "ChaosCloth/ChaosClothingSimulationConfig.h"
 #include "ChaosCloth/ChaosClothingSimulation.h"
 #include "ChaosCloth/ChaosClothPrivate.h"
+#include "Chaos/SoftsEvolution.h"
 #include "Chaos/PBDEvolution.h"
 #include "Chaos/WeightedLatticeImplicitObject.h"
 #include "Chaos/Levelset.h"
@@ -96,8 +97,68 @@ namespace ClothingSimulationSolverConstant
 	static const Softs::FSolverReal StartDeltaTime = (Softs::FSolverReal)(1. / 30.);  // Initialize filtered timestep at 30fps
 }
 
-FClothingSimulationSolver::FClothingSimulationSolver()
-	: OldLocalSpaceLocation(0.)
+namespace Private
+{
+template<typename T>
+TConstArrayView<T> GetPBDParticleConstArrayView(const TUniquePtr<Softs::FPBDEvolution>& PBDEvolution, const int32 ParticleRangeId, const TArray<T>& Array)
+{
+	return TConstArrayView<T>(Array.GetData() + ParticleRangeId, PBDEvolution->GetParticleRangeSize(ParticleRangeId));
+}
+
+template<typename T>
+TArrayView<T> GetPBDParticleArrayView(TUniquePtr<Softs::FPBDEvolution>& PBDEvolution, const int32 ParticleRangeId, TArray<T>& Array)
+{
+	return TArrayView<T>(Array.GetData() + ParticleRangeId, PBDEvolution->GetParticleRangeSize(ParticleRangeId));
+}
+
+template<typename T>
+TConstArrayView<T> GetParticleConstArrayView(const TUniquePtr<Softs::FEvolution>& Evolution, const TUniquePtr<Softs::FPBDEvolution>& PBDEvolution,
+	const int32 ParticleRangeId, const TArray<T>& Array)
+{
+	return Evolution ? Evolution->GetSoftBodyParticles(ParticleRangeId).GetConstArrayView(Array) :
+		TConstArrayView<T>(Array.GetData() + ParticleRangeId, PBDEvolution->GetParticleRangeSize(ParticleRangeId));
+}
+
+template<typename T>
+TArrayView<T> GetParticleArrayView(TUniquePtr<Softs::FEvolution>& Evolution, TUniquePtr<Softs::FPBDEvolution>& PBDEvolution,
+	const int32 ParticleRangeId, TArray<T>& Array)
+{
+	return Evolution ? Evolution->GetSoftBodyParticles(ParticleRangeId).GetArrayView(Array) :
+		TArrayView<T>(Array.GetData() + ParticleRangeId, PBDEvolution->GetParticleRangeSize(ParticleRangeId));
+}
+
+template<typename T>
+TConstArrayView<T> GetPBDCollisionParticleConstArrayView(const TUniquePtr<Softs::FPBDEvolution>& PBDEvolution, const int32 ParticleRangeId, const TArray<T>& Array)
+{
+	return TConstArrayView<T>(Array.GetData() + ParticleRangeId, PBDEvolution->GetCollisionParticleRangeSize(ParticleRangeId));
+}
+
+template<typename T>
+TArrayView<T> GetPBDCollisionParticleArrayView(TUniquePtr<Softs::FPBDEvolution>& PBDEvolution, const int32 ParticleRangeId, TArray<T>& Array)
+{
+	return TArrayView<T>(Array.GetData() + ParticleRangeId, PBDEvolution->GetCollisionParticleRangeSize(ParticleRangeId));
+}
+
+template<typename T>
+TConstArrayView<T> GetCollisionParticleConstArrayView(const TUniquePtr<Softs::FEvolution>& Evolution, const TUniquePtr<Softs::FPBDEvolution>& PBDEvolution,
+	const int32 ParticleRangeId, const TArray<T>& Array)
+{
+	return Evolution ? Evolution->GetCollisionParticleRange(ParticleRangeId).GetConstArrayView(Array) :
+		TConstArrayView<T>(Array.GetData() + ParticleRangeId, PBDEvolution->GetCollisionParticleRangeSize(ParticleRangeId));
+}
+
+template<typename T>
+TArrayView<T> GetCollisionParticleArrayView(TUniquePtr<Softs::FEvolution>& Evolution, TUniquePtr<Softs::FPBDEvolution>& PBDEvolution,
+	const int32 ParticleRangeId, TArray<T>& Array)
+{
+	return Evolution ? Evolution->GetCollisionParticleRange(ParticleRangeId).GetArrayView(Array) :
+		TArrayView<T>(Array.GetData() + ParticleRangeId, PBDEvolution->GetCollisionParticleRangeSize(ParticleRangeId));
+}
+}
+
+FClothingSimulationSolver::FClothingSimulationSolver(bool bForceBasedSolver, FClothingSimulationConfig* InConfig)
+	: Evolution(nullptr), PBDEvolution(nullptr)
+	, OldLocalSpaceLocation(0.)
 	, LocalSpaceLocation(0.)
 	, LocalSpaceRotation(FRotation3::Identity)
 	, VelocityScale(1.)
@@ -111,50 +172,145 @@ FClothingSimulationSolver::FClothingSimulationSolver()
 	, bIsClothGravityOverrideEnabled(false)
 	, bEnableSolver(true)
 {
-	SetConfig(nullptr); // This will generate a local default config so we have something to use if there are no cloth assets..
-
-	Softs::FSolverParticles LocalParticles;
-	Softs::FSolverCollisionParticles RigidParticles;
-	Evolution.Reset(
-		new Softs::FPBDEvolution(
-			MoveTemp(LocalParticles),
-			MoveTemp(RigidParticles),
-			{}, // CollisionTriangles
-			FMath::Min(ClothingSimulationSolverDefault::NumIterations, ClothingSimulationSolverDefault::MaxNumIterations),
-			(Softs::FSolverReal)ClothingSimulationSolverDefault::CollisionThickness,
-			(Softs::FSolverReal)ClothingSimulationSolverDefault::SelfCollisionThickness,
-			(Softs::FSolverReal)ClothingSimulationSolverDefault::FrictionCoefficient,
-			(Softs::FSolverReal)ClothingSimulationSolverDefault::DampingCoefficient,
-			(Softs::FSolverReal)ClothingSimulationSolverDefault::LocalDampingCoefficient));
-
-	// Add simulation groups arrays
-	Evolution->AddArray(&PreSimulationTransforms);
-	Evolution->AddArray(&FictitiousAngularDisplacements);
-	Evolution->AddArray(&ReferenceSpaceLocations);
-
-	Evolution->Particles().AddArray(&Normals);
-	Evolution->Particles().AddArray(&OldAnimationPositions);
-	Evolution->Particles().AddArray(&AnimationPositions);
-	Evolution->Particles().AddArray(&InterpolatedAnimationPositions);
-	Evolution->Particles().AddArray(&OldAnimationNormals);
-	Evolution->Particles().AddArray(&AnimationNormals);
-	Evolution->Particles().AddArray(&InterpolatedAnimationNormals);
-	Evolution->Particles().AddArray(&AnimationVelocities);
-
-	Evolution->CollisionParticles().AddArray(&CollisionBoneIndices);
-	Evolution->CollisionParticles().AddArray(&CollisionBaseTransforms);
-	Evolution->CollisionParticles().AddArray(&OldCollisionTransforms);
-	Evolution->CollisionParticles().AddArray(&CollisionTransforms);
-
-	Evolution->SetKinematicUpdateFunction(
-		[this](Softs::FSolverParticles& ParticlesInput, const Softs::FSolverReal Dt, const Softs::FSolverReal LocalTime, const int32 Index)
+	SetConfig(InConfig); // This will generate a local default config if nullptr so we have something to use if there are no cloth assets..
+
+	if (bForceBasedSolver)
+	{
+		Evolution.Reset(
+			new Softs::FEvolution(
+				Config->GetProperties(SolverLOD)
+			)
+		);
+
+		// Add simulation groups arrays
+		Evolution->AddGroupArray(&PreSimulationTransforms);
+		Evolution->AddGroupArray(&FictitiousAngularDisplacements);
+		Evolution->AddGroupArray(&ReferenceSpaceLocations);
+
+		Evolution->AddParticleArray(&Normals);
+		Evolution->AddParticleArray(&OldAnimationPositions);
+		Evolution->AddParticleArray(&AnimationPositions);
+		Evolution->AddParticleArray(&InterpolatedAnimationPositions);
+		Evolution->AddParticleArray(&OldAnimationNormals);
+		Evolution->AddParticleArray(&AnimationNormals);
+		Evolution->AddParticleArray(&InterpolatedAnimationNormals);
+		Evolution->AddParticleArray(&AnimationVelocities);
+
+		Evolution->AddCollisionParticleArray(&CollisionBoneIndices);
+		Evolution->AddCollisionParticleArray(&CollisionBaseTransforms);
+		Evolution->AddCollisionParticleArray(&OldCollisionTransforms);
+		Evolution->AddCollisionParticleArray(&CollisionTransforms);
+		Evolution->AddCollisionParticleArray(&Collided);
+		Evolution->AddCollisionParticleArray(&LastSubframeCollisionTransformsCCD);
+
+		Evolution->SetKinematicUpdateFunction(
+			[this](Softs::FSolverParticlesRange& ParticlesInput, const Softs::FSolverReal Dt, const Softs::FSolverReal LocalTime)
+		{
+			Softs::FPAndInvM* const PAndInvM = ParticlesInput.GetPAndInvM().GetData();
+			const Softs::FSolverVec3* const InterpPos = ParticlesInput.GetConstArrayView(InterpolatedAnimationPositions).GetData();
+			for (int32 Index = 0; Index < ParticlesInput.GetRangeSize(); ++Index)
+			{
+				if (PAndInvM[Index].InvM == (FSolverReal)0.)
+				{
+					PAndInvM[Index].P = InterpPos[Index];  // X is the step initial condition, here it's P that needs to be updated so that constraints works with the correct step target
+				}
+			}
+		});
+
+		Evolution->SetCollisionKinematicUpdateFunction(
+			[this](Softs::FSolverCollisionParticlesRange& ParticlesInput, const Softs::FSolverReal Dt, const Softs::FSolverReal LocalTime)
+		{
+			checkSlow(Dt > SMALL_NUMBER && DeltaTime > SMALL_NUMBER);
+			const Softs::FSolverReal Alpha = (LocalTime - Time) / DeltaTime;
+			const Softs::FSolverRigidTransform3* const CollisionTransformsLocal = ParticlesInput.GetConstArrayView(CollisionTransforms).GetData();
+			const Softs::FSolverRigidTransform3* const OldCollisionTransformsLocal = ParticlesInput.GetConstArrayView(OldCollisionTransforms).GetData();
+			Softs::FSolverRigidTransform3* const LastSubframeCollisionTransformsCCDLocal = ParticlesInput.GetArrayView(LastSubframeCollisionTransformsCCD).GetData();
+			Softs::FSolverVec3* const X = ParticlesInput.XArray().GetData();
+			Softs::FSolverVec3* const V = ParticlesInput.GetV().GetData();
+			Softs::FSolverVec3* const W = ParticlesInput.GetW().GetData();
+			Softs::FSolverRotation3* const R = ParticlesInput.GetR().GetData();
+
+			for (int32 Index = 0; Index < ParticlesInput.GetRangeSize(); ++Index)
+			{
+				LastSubframeCollisionTransformsCCDLocal[Index] = Softs::FSolverRigidTransform3(X[Index], R[Index]);
+				const Softs::FSolverVec3 NewX =
+					Alpha * CollisionTransformsLocal[Index].GetTranslation() + ((Softs::FSolverReal)1. - Alpha) * OldCollisionTransformsLocal[Index].GetTranslation();
+				V[Index] = (NewX - X[Index]) / Dt;
+				X[Index] = NewX;
+				const Softs::FSolverRotation3 NewR = Softs::FSolverRotation3::Slerp(OldCollisionTransformsLocal[Index].GetRotation(), CollisionTransformsLocal[Index].GetRotation(), Alpha);
+				const Softs::FSolverRotation3 Delta = NewR * R[Index].Inverse();
+				const Softs::FSolverReal Angle = Delta.GetAngle();
+				const Softs::FSolverVec3 Axis = Delta.GetRotationAxis();
+				W[Index] = (Softs::FSolverVec3)Axis * Angle / Dt;
+				R[Index] = NewR;
+
+				if (TWeightedLatticeImplicitObject<FLevelSet>* SkinnedLevelSet =
+					const_cast<FImplicitObject*>(ParticlesInput.GetGeometry(Index).GetReference())->GetObject<TWeightedLatticeImplicitObject<FLevelSet>>	())
+				{
+					const TArray<int32>& SubBoneIndices = SkinnedLevelSet->GetSolverBoneIndices();
+					const FTransform RootTransformInv = TRigidTransform<FReal, 3>(X[Index], R[Index]).Inverse();
+					TArray<FTransform> SubBoneTransforms;
+					SubBoneTransforms.SetNum(SubBoneIndices.Num());
+					for (int32 SubBoneIdx = 0; SubBoneIdx < SubBoneIndices.Num(); ++SubBoneIdx)
+					{
+						const int32 SubBoneIndexLocal = SubBoneIndices[SubBoneIdx] - ParticlesInput.GetOffset();
+						checkSlow(SubBoneIndexLocal < Index);
+						SubBoneTransforms[SubBoneIdx] = TRigidTransform<FReal, 3>(ParticlesInput.X(SubBoneIndexLocal), ParticlesInput.R(SubBoneIndexLocal)) * RootTransformInv;
+					}
+					SkinnedLevelSet->DeformPoints(SubBoneTransforms);
+					SkinnedLevelSet->UpdateSpatialHierarchy();
+				}
+			}
+		});
+	}
+	else
+	{
+		Softs::FSolverParticles LocalParticles;
+		Softs::FSolverCollisionParticles RigidParticles;
+		PBDEvolution.Reset(
+			new Softs::FPBDEvolution(
+				MoveTemp(LocalParticles),
+				MoveTemp(RigidParticles),
+				{}, // CollisionTriangles
+				FMath::Min(ClothingSimulationSolverDefault::NumIterations, ClothingSimulationSolverDefault::MaxNumIterations),
+				(Softs::FSolverReal)ClothingSimulationSolverDefault::CollisionThickness,
+				(Softs::FSolverReal)ClothingSimulationSolverDefault::SelfCollisionThickness,
+				(Softs::FSolverReal)ClothingSimulationSolverDefault::FrictionCoefficient,
+				(Softs::FSolverReal)ClothingSimulationSolverDefault::DampingCoefficient,
+				(Softs::FSolverReal)ClothingSimulationSolverDefault::LocalDampingCoefficient));
+
+		// Add simulation groups arrays
+		PBDEvolution->AddArray(&PreSimulationTransforms);
+		PBDEvolution->AddArray(&FictitiousAngularDisplacements);
+		PBDEvolution->AddArray(&ReferenceSpaceLocations);
+
+		PBDEvolution->Particles().AddArray(&Normals);
+		PBDEvolution->Particles().AddArray(&OldAnimationPositions);
+		PBDEvolution->Particles().AddArray(&AnimationPositions);
+		PBDEvolution->Particles().AddArray(&InterpolatedAnimationPositions);
+		PBDEvolution->Particles().AddArray(&OldAnimationNormals);
+		PBDEvolution->Particles().AddArray(&AnimationNormals);
+		PBDEvolution->Particles().AddArray(&InterpolatedAnimationNormals);
+		PBDEvolution->Particles().AddArray(&AnimationVelocities);
+
+		PBDEvolution->CollisionParticles().AddArray(&CollisionBoneIndices);
+		PBDEvolution->CollisionParticles().AddArray(&CollisionBaseTransforms);
+		PBDEvolution->CollisionParticles().AddArray(&OldCollisionTransforms);
+		PBDEvolution->CollisionParticles().AddArray(&CollisionTransforms);
+		PBDEvolution->CollisionParticles().AddArray(&Collided);
+		PBDEvolution->CollisionParticles().AddArray(&LastSubframeCollisionTransformsCCD);
+
+		PBDEvolution->SetKinematicUpdateFunction(
+			[this](Softs::FSolverParticles& ParticlesInput, const Softs::FSolverReal Dt, const Softs::FSolverReal LocalTime, const int32 Index)
 		{
 			ParticlesInput.P(Index) = InterpolatedAnimationPositions[Index];  // X is the step initial condition, here it's P that needs to be updated so that constraints works with the correct step target
 		});
 
-	Evolution->SetCollisionKinematicUpdateFunction(
-		[this](Softs::FSolverCollisionParticles& ParticlesInput, const Softs::FSolverReal Dt, const Softs::FSolverReal LocalTime, const int32 Index)
+		PBDEvolution->SetCollisionKinematicUpdateFunction(
+			[this](Softs::FSolverCollisionParticles& ParticlesInput, const Softs::FSolverReal Dt, const Softs::FSolverReal LocalTime, const int32 Index)
 		{
+			LastSubframeCollisionTransformsCCD[Index] = Softs::FSolverRigidTransform3(ParticlesInput.X(Index), ParticlesInput.R(Index));
+
 			checkSlow(Dt > SMALL_NUMBER && DeltaTime > SMALL_NUMBER);
 			const Softs::FSolverReal Alpha = (LocalTime - Time) / DeltaTime;
 			const Softs::FSolverVec3 NewX =
@@ -168,7 +324,7 @@ FClothingSimulationSolver::FClothingSimulationSolver()
 			ParticlesInput.W(Index) = (Softs::FSolverVec3)Axis * Angle / Dt;
 			ParticlesInput.R(Index) = NewR;
 
-			if (TWeightedLatticeImplicitObject<FLevelSet>* SkinnedLevelSet = 
+			if (TWeightedLatticeImplicitObject<FLevelSet>* SkinnedLevelSet =
 				const_cast<FImplicitObject*>(ParticlesInput.GetGeometry(Index).GetReference())->GetObject<TWeightedLatticeImplicitObject<FLevelSet>>())
 			{
 				const TArray<int32>& SubBoneIndices = SkinnedLevelSet->GetSolverBoneIndices();
@@ -178,12 +334,13 @@ FClothingSimulationSolver::FClothingSimulationSolver()
 				for (int32 SubBoneIdx = 0; SubBoneIdx < SubBoneIndices.Num(); ++SubBoneIdx)
 				{
 					checkSlow(SubBoneIndices[SubBoneIdx] < Index);
-					SubBoneTransforms[SubBoneIdx] = TRigidTransform<FReal,3>(ParticlesInput.X(SubBoneIndices[SubBoneIdx]), ParticlesInput.R(SubBoneIndices[SubBoneIdx])) * RootTransformInv;
+					SubBoneTransforms[SubBoneIdx] = TRigidTransform<FReal, 3>(ParticlesInput.X(SubBoneIndices[SubBoneIdx]), ParticlesInput.R(SubBoneIndices[SubBoneIdx])) * RootTransformInv;
 				}
 				SkinnedLevelSet->DeformPoints(SubBoneTransforms);
 				SkinnedLevelSet->UpdateSpatialHierarchy();
 			}
 		});
+	}
 }
 
 FClothingSimulationSolver::~FClothingSimulationSolver()
@@ -226,7 +383,10 @@ void FClothingSimulationSolver::SetCloths(TArray<FClothingSimulationCloth*>&& In
 	}
 
 	// Update external collision's offset
-	CollisionParticlesOffset = Evolution->CollisionParticles().Size();
+	if (PBDEvolution)
+	{
+		CollisionParticlesOffset = PBDEvolution->CollisionParticles().Size();
+	}
 }
 
 void FClothingSimulationSolver::AddCloth(FClothingSimulationCloth* InCloth)
@@ -252,7 +412,7 @@ void FClothingSimulationSolver::AddCloth(FClothingSimulationCloth* InCloth)
 	InCloth->Update(this);
 
 	// Update external collision's offset
-	CollisionParticlesOffset = Evolution->CollisionParticles().Size();
+	CollisionParticlesOffset = PBDEvolution ? PBDEvolution->CollisionParticles().Size() : 0;
 }
 
 void FClothingSimulationSolver::RemoveCloth(FClothingSimulationCloth* InCloth)
@@ -268,11 +428,8 @@ void FClothingSimulationSolver::RemoveCloth(FClothingSimulationCloth* InCloth)
 	// Remove collider from array
 	Cloths.RemoveSwap(InCloth);
 
-	// Reset collisions so that there is never any external collision particles below the cloth's ones
-	ResetCollisionParticles();
-
-	// Reset cloth particles and associated elements
-	ResetParticles();
+	// Reset all particles, collisions, constraints
+	Reset();
 
 	// Re-add the remaining cloths' particles
 	for (FClothingSimulationCloth* const Cloth : Cloths)
@@ -284,9 +441,9 @@ void FClothingSimulationSolver::RemoveCloth(FClothingSimulationCloth* InCloth)
 		Cloth->PreUpdate(this);
 		Cloth->Update(this);
 	}
-
+	
 	// Update external collision's offset
-	CollisionParticlesOffset = Evolution->CollisionParticles().Size();
+	CollisionParticlesOffset = PBDEvolution ? PBDEvolution->CollisionParticles().Size() : 0;
 }
 
 void FClothingSimulationSolver::RemoveCloths()
@@ -299,10 +456,8 @@ void FClothingSimulationSolver::RemoveCloths()
 	Cloths.Reset();
 
 	// Reset solver collisions
-	ResetCollisionParticles();
-
 	// Reset cloth particles and associated elements
-	ResetParticles();
+	Reset();
 }
 
 void FClothingSimulationSolver::RefreshCloth(FClothingSimulationCloth* InCloth)
@@ -325,11 +480,8 @@ void FClothingSimulationSolver::RefreshCloths()
 		Cloth->Remove(this);
 	}
 
-	// Reset collision particles
-	ResetCollisionParticles();
-
-	// Reset cloth particles and associated elements
-	ResetParticles();
+	// Reset evolution particles and constraints
+	Reset();
 
 	// Re-add the cloths' & collisions' particles
 	for (FClothingSimulationCloth* const Cloth : Cloths)
@@ -343,7 +495,7 @@ void FClothingSimulationSolver::RefreshCloths()
 	}
 
 	// Update solver collider's offset
-	CollisionParticlesOffset = Evolution->CollisionParticles().Size();
+	CollisionParticlesOffset = PBDEvolution ? PBDEvolution->CollisionParticles().Size() : 0;
 }
 
 void FClothingSimulationSolver::SetConfig(FClothingSimulationConfig* InConfig)
@@ -369,11 +521,38 @@ void FClothingSimulationSolver::SetConfig(FClothingSimulationConfig* InConfig)
 	}
 }
 
+void FClothingSimulationSolver::Reset()
+{
+	if (Evolution)
+	{
+		Evolution->Reset();
+		ClothsConstraints.Reset();
+	}
+	else
+	{
+		ResetParticles();
+		ResetCollisionParticles();
+	}
+}
+
 void FClothingSimulationSolver::ResetParticles()
 {
-	Evolution->ResetParticles();
-	Evolution->ResetConstraintRules();
-	ClothsConstraints.Reset();
+	if (PBDEvolution)
+	{
+		PBDEvolution->ResetParticles();
+		PBDEvolution->ResetConstraintRules();
+		ClothsConstraints.Reset();
+	}
+}
+
+void FClothingSimulationSolver::ResetCollisionParticles(int32 InCollisionParticlesOffset)
+{
+	if (PBDEvolution)
+	{
+		PBDEvolution->ResetCollisionParticles(InCollisionParticlesOffset);
+		CollisionParticlesOffset = InCollisionParticlesOffset;
+		CollisionParticlesSize = 0;
+	}
 }
 
 int32 FClothingSimulationSolver::AddParticles(int32 NumParticles, uint32 GroupId)
@@ -382,35 +561,65 @@ int32 FClothingSimulationSolver::AddParticles(int32 NumParticles, uint32 GroupId
 	{
 		return INDEX_NONE;
 	}
-	const int32 Offset = Evolution->AddParticleRange(NumParticles, GroupId, /*bActivate =*/ false);
+	constexpr bool bActivateFalse = false;
+	const int32 ParticleRangeId = Evolution ? Evolution->AddSoftBody(GroupId, NumParticles, bActivateFalse) : 
+		PBDEvolution->AddParticleRange(NumParticles, GroupId, bActivateFalse);
 
 	// Add an empty constraints container for this range
-	check(!ClothsConstraints.Find(Offset));  // We cannot already have this Offset in the map, particle ranges are always added, never removed (unless reset)
+	check(!ClothsConstraints.Find(ParticleRangeId));  // We cannot already have this ParticleRangeId in the map, particle ranges are always added, never removed (unless reset)
 
-	ClothsConstraints.Emplace(Offset, MakeUnique<FClothConstraints>())
-		->Initialize(Evolution.Get(), InterpolatedAnimationPositions, OldAnimationPositions, InterpolatedAnimationNormals, AnimationVelocities, Offset, NumParticles);
+	if (Evolution)
+	{
+		ClothsConstraints.Emplace(ParticleRangeId, MakeUnique<FClothConstraints>())
+			->Initialize(
+				Evolution.Get(),
+				&PerSolverField,
+				InterpolatedAnimationPositions,
+				InterpolatedAnimationNormals,
+				AnimationVelocities,
+				Normals,
+				LastSubframeCollisionTransformsCCD,
+				Collided,
+				CollisionContacts,
+				CollisionNormals,
+				CollisionPhis,
+				ParticleRangeId);
+	}
+	else
+	{
+		check(PBDEvolution);
+		ClothsConstraints.Emplace(ParticleRangeId, MakeUnique<FClothConstraints>())
+			->Initialize(PBDEvolution.Get(), InterpolatedAnimationPositions, OldAnimationPositions, InterpolatedAnimationNormals, AnimationVelocities, ParticleRangeId, NumParticles);
+	}
 
 	// Always starts with particles disabled
-	EnableParticles(Offset, false);
+	EnableParticles(ParticleRangeId, false);
 
-	return Offset;
+	return ParticleRangeId;
 }
 
-void FClothingSimulationSolver::EnableParticles(int32 Offset, bool bEnable)
+void FClothingSimulationSolver::EnableParticles(int32 ParticleRangeId, bool bEnable)
 {
-	Evolution->ActivateParticleRange(Offset, bEnable);
-	GetClothConstraints(Offset).Enable(bEnable);
+	if (Evolution)
+	{
+		Evolution->ActivateSoftBody(ParticleRangeId, bEnable);
+	}
+	else
+	{
+		PBDEvolution->ActivateParticleRange(ParticleRangeId, bEnable);
+		GetClothConstraints(ParticleRangeId).Enable(bEnable);
+	}
 }
 
-void FClothingSimulationSolver::ResetStartPose(int32 Offset, int32 NumParticles)
+void FClothingSimulationSolver::ResetStartPose(int32 ParticleRangeId, int32 NumParticles)
 {
-	Softs::FPAndInvM* const PandInvMs = GetParticlePandInvMs(Offset);
-	Softs::FSolverVec3* const Xs = GetParticleXs(Offset);
-	Softs::FSolverVec3* const Vs = GetParticleVs(Offset);
-	const Softs::FSolverVec3* const Positions = GetAnimationPositions(Offset);
-	Softs::FSolverVec3* const OldPositions = GetOldAnimationPositions(Offset);
-	Softs::FSolverVec3* const InterpolatedPositions = GetInterpolatedAnimationPositions(Offset);
-	Softs::FSolverVec3* const AnimationVs = GetAnimationVelocities(Offset);
+	Softs::FPAndInvM* const PandInvMs = GetParticlePandInvMs(ParticleRangeId);
+	Softs::FSolverVec3* const Xs = GetParticleXs(ParticleRangeId);
+	Softs::FSolverVec3* const Vs = GetParticleVs(ParticleRangeId);
+	const Softs::FSolverVec3* const Positions = GetAnimationPositions(ParticleRangeId);
+	Softs::FSolverVec3* const OldPositions = GetOldAnimationPositions(ParticleRangeId);
+	Softs::FSolverVec3* const InterpolatedPositions = GetInterpolatedAnimationPositions(ParticleRangeId);
+	Softs::FSolverVec3* const AnimationVs = GetAnimationVelocities(ParticleRangeId);
 
 	for (int32 Index = 0; Index < NumParticles; ++Index)
 	{
@@ -421,96 +630,184 @@ void FClothingSimulationSolver::ResetStartPose(int32 Offset, int32 NumParticles)
 
 const TArray<Softs::FPAndInvM>& FClothingSimulationSolver::GetParticlePandInvMs() const
 {
-	return Evolution->Particles().GetPAndInvM();
+	return Evolution ? Evolution->GetParticles().GetPAndInvM() : PBDEvolution->GetParticles().GetPAndInvM();
 }
 
-const Softs::FPAndInvM* FClothingSimulationSolver::GetParticlePandInvMs(int32 Offset) const
+const TArray<Softs::FSolverVec3>& FClothingSimulationSolver::GetParticleXs() const
 {
-	return &Evolution->Particles().PAndInvM(Offset);
+	return Evolution ? Evolution->GetParticles().XArray() : PBDEvolution->GetParticles().XArray();
 }
 
-Softs::FPAndInvM* FClothingSimulationSolver::GetParticlePandInvMs(int32 Offset)
+const TArray<Softs::FSolverVec3>& FClothingSimulationSolver::GetParticleVs() const
 {
-	return &Evolution->Particles().PAndInvM(Offset);
+	return Evolution ? Evolution->GetParticles().GetV() : PBDEvolution->GetParticles().GetV();
 }
 
-const TArray<Softs::FSolverVec3>& FClothingSimulationSolver::GetParticleXs() const
+const TArray<Softs::FSolverReal>& FClothingSimulationSolver::GetParticleInvMasses() const
 {
-	return Evolution->Particles().XArray();
+	return Evolution ? Evolution->GetParticles().GetInvM() : PBDEvolution->GetParticles().GetInvM();
 }
 
-const Softs::FSolverVec3* FClothingSimulationSolver::GetParticleXs(int32 Offset) const
+TConstArrayView<Softs::FPAndInvM> FClothingSimulationSolver::GetParticlePandInvMsView(int32 ParticleRangeId) const
 {
-	return &Evolution->Particles().X(Offset);
+	return Evolution ? Evolution->GetSoftBodyParticles(ParticleRangeId).GetPAndInvM() : Private::GetPBDParticleConstArrayView(PBDEvolution, ParticleRangeId, PBDEvolution->GetParticles().GetPAndInvM());
 }
 
-Softs::FSolverVec3* FClothingSimulationSolver::GetParticleXs(int32 Offset)
+TArrayView<Softs::FPAndInvM> FClothingSimulationSolver::GetParticlePandInvMsView(int32 ParticleRangeId)
 {
-	return &Evolution->Particles().X(Offset);
+	return Evolution ? Evolution->GetSoftBodyParticles(ParticleRangeId).GetPAndInvM() : Private::GetPBDParticleArrayView(PBDEvolution, ParticleRangeId, PBDEvolution->GetParticles().GetPAndInvM());
 }
 
-const TArray<Softs::FSolverVec3>& FClothingSimulationSolver::GetParticleVs() const
+TConstArrayView<Softs::FSolverVec3> FClothingSimulationSolver::GetParticleXsView(int32 ParticleRangeId) const
 {
-	return Evolution->Particles().GetV();
+	return Evolution ? Evolution->GetSoftBodyParticles(ParticleRangeId).XArray() : Private::GetPBDParticleConstArrayView(PBDEvolution, ParticleRangeId, PBDEvolution->GetParticles().XArray());
 }
 
-const Softs::FSolverVec3* FClothingSimulationSolver::GetParticleVs(int32 Offset) const
+TArrayView<Softs::FSolverVec3> FClothingSimulationSolver::GetParticleXsView(int32 ParticleRangeId)
 {
-	return &Evolution->Particles().V(Offset);
+	return Evolution ? Evolution->GetSoftBodyParticles(ParticleRangeId).XArray() : Private::GetPBDParticleArrayView(PBDEvolution, ParticleRangeId, PBDEvolution->GetParticles().XArray());
 }
 
-Softs::FSolverVec3* FClothingSimulationSolver::GetParticleVs(int32 Offset)
+TConstArrayView<Softs::FSolverVec3> FClothingSimulationSolver::GetParticleVsView(int32 ParticleRangeId) const
 {
-	return &Evolution->Particles().V(Offset);
+	return Evolution ? Evolution->GetSoftBodyParticles(ParticleRangeId).GetV() : Private::GetPBDParticleConstArrayView(PBDEvolution, ParticleRangeId, PBDEvolution->GetParticles().GetV());
 }
 
-const TArray<Softs::FSolverReal>& FClothingSimulationSolver::GetParticleInvMasses() const
+TArrayView<Softs::FSolverVec3> FClothingSimulationSolver::GetParticleVsView(int32 ParticleRangeId)
 {
-	return Evolution->Particles().GetInvM();
+	return Evolution ? Evolution->GetSoftBodyParticles(ParticleRangeId).GetV() : Private::GetPBDParticleArrayView(PBDEvolution, ParticleRangeId, PBDEvolution->GetParticles().GetV());
 }
 
-const Softs::FSolverReal* FClothingSimulationSolver::GetParticleInvMasses(int32 Offset) const
+TConstArrayView<Softs::FSolverReal> FClothingSimulationSolver::GetParticleInvMassesView(int32 ParticleRangeId) const
 {
-	return &Evolution->Particles().InvM(Offset);
+	return Evolution ? Evolution->GetSoftBodyParticles(ParticleRangeId).GetInvM() : Private::GetPBDParticleConstArrayView(PBDEvolution, ParticleRangeId, PBDEvolution->GetParticles().GetInvM());
 }
 
-void FClothingSimulationSolver::ResetCollisionParticles(int32 InCollisionParticlesOffset)
+TArrayView<Softs::FSolverReal> FClothingSimulationSolver::GetParticleInvMassesView(int32 ParticleRangeId)
 {
-	Evolution->ResetCollisionParticles(InCollisionParticlesOffset);
-	CollisionParticlesOffset = InCollisionParticlesOffset;
-	CollisionParticlesSize = 0;
+	return Evolution ? Evolution->GetSoftBodyParticles(ParticleRangeId).GetInvM() : Private::GetPBDParticleArrayView(PBDEvolution, ParticleRangeId, PBDEvolution->GetParticles().GetInvM());
 }
 
-int32 FClothingSimulationSolver::AddCollisionParticles(int32 NumCollisionParticles, uint32 GroupId, int32 RecycledOffset)
+TConstArrayView<Softs::FSolverVec3> FClothingSimulationSolver::GetOldAnimationPositionsView(int32 ParticleRangeId) const
 {
-	// Try reusing the particle range
-	// This is used by external collisions so that they can be added/removed between every solver update.
-	// If it doesn't match then remove all ranges above the given offset to start again.
-	// This rely on the assumption that these ranges are added again in the same update order.
-	if (RecycledOffset == CollisionParticlesOffset + CollisionParticlesSize)
-	{
-		CollisionParticlesSize += NumCollisionParticles;
+	return Private::GetParticleConstArrayView(Evolution, PBDEvolution, ParticleRangeId, OldAnimationPositions); 
+}
 
-		// Check that the range still exists
-		if (CollisionParticlesOffset + CollisionParticlesSize <= (int32)Evolution->CollisionParticles().Size() &&  // Check first that the range hasn't been reset
-			NumCollisionParticles == Evolution->GetCollisionParticleRangeSize(RecycledOffset))  // This will assert if range has been reset
-		{
-			return RecycledOffset;
-		}
-		// Size has changed. must reset this collision range (and all of those following up) and reallocate some new particles
-		Evolution->ResetCollisionParticles(RecycledOffset);
-	}
+TArrayView<Softs::FSolverVec3> FClothingSimulationSolver::GetOldAnimationPositionsView(int32 ParticleRangeId)
+{
+	return Private::GetParticleArrayView(Evolution, PBDEvolution, ParticleRangeId, OldAnimationPositions);
+}
+
+TConstArrayView<Softs::FSolverVec3> FClothingSimulationSolver::GetAnimationPositionsView(int32 ParticleRangeId) const
+{
+	return Private::GetParticleConstArrayView(Evolution, PBDEvolution, ParticleRangeId, AnimationPositions);
+}
+
+TArrayView<Softs::FSolverVec3> FClothingSimulationSolver::GetAnimationPositionsView(int32 ParticleRangeId)
+{
+	return Private::GetParticleArrayView(Evolution, PBDEvolution, ParticleRangeId, AnimationPositions);
+}
+
+TConstArrayView<Softs::FSolverVec3> FClothingSimulationSolver::GetInterpolatedAnimationPositionsView(int32 ParticleRangeId) const
+{
+	return Private::GetParticleConstArrayView(Evolution, PBDEvolution, ParticleRangeId, InterpolatedAnimationPositions);
+}
+
+TArrayView<Softs::FSolverVec3> FClothingSimulationSolver::GetInterpolatedAnimationPositionsView(int32 ParticleRangeId)
+{
+	return Private::GetParticleArrayView(Evolution, PBDEvolution, ParticleRangeId, InterpolatedAnimationPositions);
+}
+
+TConstArrayView<Softs::FSolverVec3> FClothingSimulationSolver::GetOldAnimationNormalsView(int32 ParticleRangeId) const
+{
+	return Private::GetParticleConstArrayView(Evolution, PBDEvolution, ParticleRangeId, OldAnimationNormals);
+}
+
+TArrayView<Softs::FSolverVec3> FClothingSimulationSolver::GetOldAnimationNormalsView(int32 ParticleRangeId)
+{
+	return Private::GetParticleArrayView(Evolution, PBDEvolution, ParticleRangeId, OldAnimationNormals);
+}
+
+TConstArrayView<Softs::FSolverVec3> FClothingSimulationSolver::GetAnimationNormalsView(int32 ParticleRangeId) const
+{
+	return Private::GetParticleConstArrayView(Evolution, PBDEvolution, ParticleRangeId, AnimationNormals);
+}
+
+TArrayView<Softs::FSolverVec3> FClothingSimulationSolver::GetAnimationNormalsView(int32 ParticleRangeId)
+{
+	return Private::GetParticleArrayView(Evolution, PBDEvolution, ParticleRangeId, AnimationNormals);
+}
+
+TConstArrayView<Softs::FSolverVec3> FClothingSimulationSolver::GetInterpolatedAnimationNormalsView(int32 ParticleRangeId) const
+{
+	return Private::GetParticleConstArrayView(Evolution, PBDEvolution, ParticleRangeId, InterpolatedAnimationNormals);
+}
+
+TArrayView<Softs::FSolverVec3> FClothingSimulationSolver::GetInterpolatedAnimationNormalsView(int32 ParticleRangeId)
+{
+	return Private::GetParticleArrayView(Evolution, PBDEvolution, ParticleRangeId, InterpolatedAnimationNormals);
+}
+
+TConstArrayView<Softs::FSolverVec3> FClothingSimulationSolver::GetNormalsView(int32 ParticleRangeId) const
+{
+	return Private::GetParticleConstArrayView(Evolution, PBDEvolution, ParticleRangeId, Normals);
+}
+
+TArrayView<Softs::FSolverVec3> FClothingSimulationSolver::GetNormalsView(int32 ParticleRangeId)
+{
+	return Private::GetParticleArrayView(Evolution, PBDEvolution, ParticleRangeId, Normals);
+}
 
+TConstArrayView<Softs::FSolverVec3> FClothingSimulationSolver::GetAnimationVelocitiesView(int32 ParticleRangeId) const
+{
+	return Private::GetParticleConstArrayView(Evolution, PBDEvolution, ParticleRangeId, AnimationVelocities);
+}
+
+TArrayView<Softs::FSolverVec3> FClothingSimulationSolver::GetAnimationVelocitiesView(int32 ParticleRangeId)
+{
+	return Private::GetParticleArrayView(Evolution, PBDEvolution, ParticleRangeId, AnimationVelocities);
+}
+
+int32 FClothingSimulationSolver::AddCollisionParticles(int32 NumCollisionParticles, uint32 GroupId, int32 RecycledCollisionRangeId)
+{
 	if (!NumCollisionParticles)
 	{
 		return INDEX_NONE;
 	}
 
-	const int32 Offset = Evolution->AddCollisionParticleRange(NumCollisionParticles, GroupId, /*bActivate =*/ false);
+	constexpr bool bActivateFalse = false;
+
+	int32 CollisionRangeId = INDEX_NONE;
+	if (Evolution)
+	{
+		CollisionRangeId = Evolution->AddCollisionParticleRange(GroupId, NumCollisionParticles, bActivateFalse);
+	}
+	else
+	{
+		// Try reusing the particle range
+		// This is used by external collisions so that they can be added/removed between every solver update.
+		// If it doesn't match then remove all ranges above the given offset to start again.
+		// This rely on the assumption that these ranges are added again in the same update order.
+		if (RecycledCollisionRangeId == CollisionParticlesOffset + CollisionParticlesSize)
+		{
+			CollisionParticlesSize += NumCollisionParticles;
+
+			// Check that the range still exists
+			if (CollisionParticlesOffset + CollisionParticlesSize <= (int32)PBDEvolution->CollisionParticles().Size() &&  // Check first that the range hasn't been reset
+				NumCollisionParticles == PBDEvolution->GetCollisionParticleRangeSize(RecycledCollisionRangeId))  // This will assert if range has been reset
+			{
+				return RecycledCollisionRangeId;
+			}
+			// Size has changed. must reset this collision range (and all of those following up) and reallocate some new particles
+			PBDEvolution->ResetCollisionParticles(RecycledCollisionRangeId);
+		}
+
+		CollisionRangeId = PBDEvolution->AddCollisionParticleRange(NumCollisionParticles, GroupId, /*bActivate =*/ false);
+	}
 
 	// Always initialize the collision particle's transforms as otherwise setting the geometry will get NaNs detected during the bounding box updates
-	Softs::FSolverRotation3* const Rs = GetCollisionParticleRs(Offset);
-	Softs::FSolverVec3* const Xs = GetCollisionParticleXs(Offset);
+	Softs::FSolverRotation3* const Rs = GetCollisionParticleRs(CollisionRangeId);
+	Softs::FSolverVec3* const Xs = GetCollisionParticleXs(CollisionRangeId);
 
 	for (int32 Index = 0; Index < NumCollisionParticles; ++Index)
 	{
@@ -519,138 +816,251 @@ int32 FClothingSimulationSolver::AddCollisionParticles(int32 NumCollisionParticl
 	}
 
 	// Always starts with particles disabled
-	EnableCollisionParticles(Offset, false);
+	EnableCollisionParticles(CollisionRangeId, false);
 
-	return Offset;
+	return CollisionRangeId;
 }
 
 void FClothingSimulationSolver::EnableCollisionParticles(int32 Offset, bool bEnable)
 {
 #if !UE_BUILD_SHIPPING
-	if (bClothSolverDisableCollision)
+	const bool bFilteredEnable = bClothSolverDisableCollision ? false : bEnable;
+#else
+	const bool bFilteredEnable = bEnable;
+#endif
+	if (Evolution)
 	{
-		Evolution->ActivateCollisionParticleRange(Offset, false);
+		Evolution->ActivateCollisionParticleRange(Offset, bFilteredEnable);
 	}
 	else
-#endif  // #if !UE_BUILD_SHIPPING
 	{
-		Evolution->ActivateCollisionParticleRange(Offset, bEnable);
+		check(PBDEvolution);
+		PBDEvolution->ActivateCollisionParticleRange(Offset, bFilteredEnable);
 	}
 }
 
-void FClothingSimulationSolver::ResetCollisionStartPose(int32 Offset, int32 NumCollisionParticles)
+void FClothingSimulationSolver::ResetCollisionStartPose(int32 CollisionRangeId, int32 NumCollisionParticles)
 {
-	const Softs::FSolverRigidTransform3* const Transforms = GetCollisionTransforms(Offset);
-	Softs::FSolverRigidTransform3* const OldTransforms = GetOldCollisionTransforms(Offset);
-	Softs::FSolverRotation3* const Rs = GetCollisionParticleRs(Offset);
-	Softs::FSolverVec3* const Xs = GetCollisionParticleXs(Offset);
+	const Softs::FSolverRigidTransform3* const Transforms = GetCollisionTransforms(CollisionRangeId);
+	Softs::FSolverRigidTransform3* const OldTransforms = GetOldCollisionTransforms(CollisionRangeId);
+	Softs::FSolverRotation3* const Rs = GetCollisionParticleRs(CollisionRangeId);
+	Softs::FSolverVec3* const Xs = GetCollisionParticleXs(CollisionRangeId);
+	Softs::FSolverRigidTransform3* const LastSubframeCollisionTransformsCCDs = GetLastSubframeCollisionTransformsCCD(CollisionRangeId);
 
 	for (int32 Index = 0; Index < NumCollisionParticles; ++Index)
 	{
 		OldTransforms[Index] = Transforms[Index];
 		Xs[Index] = Transforms[Index].GetTranslation();
 		Rs[Index] = Transforms[Index].GetRotation();
+		LastSubframeCollisionTransformsCCDs[Index] = Transforms[Index];
 	}
 }
 
-const Softs::FSolverVec3* FClothingSimulationSolver::GetCollisionParticleXs(int32 Offset) const
+TConstArrayView<int32> FClothingSimulationSolver::GetCollisionBoneIndicesView(int32 CollisionRangeId) const
+{
+	return Private::GetCollisionParticleConstArrayView(Evolution, PBDEvolution, CollisionRangeId, CollisionBoneIndices);
+}
+
+TArrayView<int32> FClothingSimulationSolver::GetCollisionBoneIndicesView(int32 CollisionRangeId)
+{
+	return Private::GetCollisionParticleArrayView(Evolution, PBDEvolution, CollisionRangeId, CollisionBoneIndices);
+}
+
+TConstArrayView<Softs::FSolverRigidTransform3> FClothingSimulationSolver::GetCollisionBaseTransformsView(int32 CollisionRangeId) const
+{
+	return Private::GetCollisionParticleConstArrayView(Evolution, PBDEvolution, CollisionRangeId, CollisionBaseTransforms);
+}
+
+TArrayView<Softs::FSolverRigidTransform3> FClothingSimulationSolver::GetCollisionBaseTransformsView(int32 CollisionRangeId)
+{
+	return Private::GetCollisionParticleArrayView(Evolution, PBDEvolution, CollisionRangeId, CollisionBaseTransforms);
+}
+
+TConstArrayView<Softs::FSolverRigidTransform3> FClothingSimulationSolver::GetOldCollisionTransformsView(int32 CollisionRangeId) const
+{
+	return Private::GetCollisionParticleConstArrayView(Evolution, PBDEvolution, CollisionRangeId, OldCollisionTransforms);
+}
+
+TArrayView<Softs::FSolverRigidTransform3> FClothingSimulationSolver::GetOldCollisionTransformsView(int32 CollisionRangeId)
+{
+	return Private::GetCollisionParticleArrayView(Evolution, PBDEvolution, CollisionRangeId, OldCollisionTransforms);
+}
+
+TConstArrayView<Softs::FSolverRigidTransform3> FClothingSimulationSolver::GetCollisionTransformsView(int32 CollisionRangeId) const
+{
+	return Private::GetCollisionParticleConstArrayView(Evolution, PBDEvolution, CollisionRangeId, CollisionTransforms);
+}
+
+TArrayView<Softs::FSolverRigidTransform3> FClothingSimulationSolver::GetCollisionTransformsView(int32 CollisionRangeId)
 {
-	return &Evolution->CollisionParticles().X(Offset);
+	return Private::GetCollisionParticleArrayView(Evolution, PBDEvolution, CollisionRangeId, CollisionTransforms);
 }
 
-Softs::FSolverVec3* FClothingSimulationSolver::GetCollisionParticleXs(int32 Offset)
+TConstArrayView<Softs::FSolverVec3> FClothingSimulationSolver::GetCollisionParticleXsView(int32 CollisionRangeId) const
 {
-	return &Evolution->CollisionParticles().X(Offset);
+	return Evolution ? Evolution->GetCollisionParticleRange(CollisionRangeId).XArray() : Private::GetPBDCollisionParticleConstArrayView(PBDEvolution, CollisionRangeId, PBDEvolution->CollisionParticles().XArray());
 }
 
-const Softs::FSolverRotation3* FClothingSimulationSolver::GetCollisionParticleRs(int32 Offset) const
+TArrayView<Softs::FSolverVec3> FClothingSimulationSolver::GetCollisionParticleXsView(int32 CollisionRangeId)
 {
-	return &Evolution->CollisionParticles().R(Offset);
+	return Evolution ? Evolution->GetCollisionParticleRange(CollisionRangeId).XArray() : Private::GetPBDCollisionParticleArrayView(PBDEvolution, CollisionRangeId, PBDEvolution->CollisionParticles().XArray());
 }
 
-Softs::FSolverRotation3* FClothingSimulationSolver::GetCollisionParticleRs(int32 Offset)
+TConstArrayView<Softs::FSolverRotation3> FClothingSimulationSolver::GetCollisionParticleRsView(int32 CollisionRangeId) const
 {
-	return &Evolution->CollisionParticles().R(Offset);
+	return Evolution ? Evolution->GetCollisionParticleRange(CollisionRangeId).GetR() : Private::GetPBDCollisionParticleConstArrayView(PBDEvolution, CollisionRangeId, PBDEvolution->CollisionParticles().GetR());
 }
 
-void FClothingSimulationSolver::SetCollisionGeometry(int32 Offset, int32 Index, FImplicitObjectPtr&& Geometry)
+TArrayView<Softs::FSolverRotation3> FClothingSimulationSolver::GetCollisionParticleRsView(int32 CollisionRangeId)
 {
-	Evolution->CollisionParticles().SetGeometry(Offset + Index, MoveTemp(Geometry));
+	return Evolution ? Evolution->GetCollisionParticleRange(CollisionRangeId).GetR() : Private::GetPBDCollisionParticleArrayView(PBDEvolution, CollisionRangeId, PBDEvolution->CollisionParticles().GetR());
 }
 
-const FImplicitObjectPtr* FClothingSimulationSolver::GetCollisionGeometry(int32 Offset) const
+TConstArrayView<FImplicitObjectPtr> FClothingSimulationSolver::GetCollisionGeometryView(int32 CollisionRangeId) const
 {
-	return &Evolution->CollisionParticles().GetGeometry(Offset);
+	return Evolution ? Evolution->GetCollisionParticleRange(CollisionRangeId).GetAllGeometry() : Private::GetPBDCollisionParticleConstArrayView(PBDEvolution, CollisionRangeId, PBDEvolution->CollisionParticles().GetAllGeometry());
 }
 
-const bool* FClothingSimulationSolver::GetCollisionStatus(int32 Offset) const
+TConstArrayView<bool> FClothingSimulationSolver::GetCollisionStatusView(int32 CollisionRangeId) const
 {
-	return Evolution->GetCollisionStatus().GetData() + Offset;
+	return Evolution ? Evolution->GetCollisionParticleRange(CollisionRangeId).GetConstArrayView(Collided) : Private::GetPBDCollisionParticleConstArrayView(PBDEvolution, CollisionRangeId, PBDEvolution->GetCollisionStatus());
+}
+
+TConstArrayView<Softs::FSolverRigidTransform3> FClothingSimulationSolver::GetLastSubframeCollisionTransformsCCDView(int32 CollisionRangeId) const
+{
+	return Private::GetCollisionParticleConstArrayView(Evolution, PBDEvolution, CollisionRangeId, LastSubframeCollisionTransformsCCD);
+}
+
+TArrayView<Softs::FSolverRigidTransform3> FClothingSimulationSolver::GetLastSubframeCollisionTransformsCCDView(int32 CollisionRangeId)
+{
+	return Private::GetCollisionParticleArrayView(Evolution, PBDEvolution, CollisionRangeId, LastSubframeCollisionTransformsCCD);
+}
+
+void FClothingSimulationSolver::SetCollisionGeometry(int32 CollisionRangeId, int32 Index, FImplicitObjectPtr&& Geometry)
+{
+	if (Evolution)
+	{
+		Evolution->GetCollisionParticleRange(CollisionRangeId).SetGeometry(Index, MoveTemp(Geometry));
+	}
+	else
+	{
+		PBDEvolution->CollisionParticles().SetGeometry(CollisionRangeId + Index, MoveTemp(Geometry));
+	}
 }
 
 const TArray<Softs::FSolverVec3>& FClothingSimulationSolver::GetCollisionContacts() const
 {
-	return Evolution->GetCollisionContacts();
+	return Evolution ? CollisionContacts : PBDEvolution->GetCollisionContacts();
 }
 
 const TArray<Softs::FSolverReal>& FClothingSimulationSolver::GetCollisionPhis() const
 {
-	return Evolution->GetCollisionPhis();
+	return Evolution ? CollisionPhis : PBDEvolution->GetCollisionPhis();
 }
 
 const TArray<Softs::FSolverVec3>& FClothingSimulationSolver::GetCollisionNormals() const
 {
-	return Evolution->GetCollisionNormals();
+	return Evolution ? CollisionNormals : PBDEvolution->GetCollisionNormals();
 }
 
-void FClothingSimulationSolver::SetParticleMassUniform(int32 Offset, FRealSingle UniformMass, FRealSingle MinPerParticleMass, const FTriangleMesh& Mesh, const TFunctionRef<bool(int32)>& KinematicPredicate)
+void FClothingSimulationSolver::SetParticleMassUniform(int32 ParticleRangeId, FRealSingle UniformMass, FRealSingle MinPerParticleMass, const FTriangleMesh& Mesh, const TFunctionRef<bool(int32)>& KinematicPredicate)
 {
-	// Retrieve the particle block size
-	const int32 Size = Evolution->GetParticleRangeSize(Offset);
-
-	// Set mass from uniform mass
-	const TSet<int32> Vertices = Mesh.GetVertices();
-	Softs::FSolverParticles& Particles = Evolution->Particles();
-	for (int32 Index = Offset; Index < Offset + Size; ++Index)
+	if (Evolution)
 	{
-		Particles.M(Index) = Vertices.Contains(Index) ? (Softs::FSolverReal)UniformMass : (Softs::FSolverReal)0.;
+		Softs::FSolverParticlesRange& Particles = Evolution->GetSoftBodyParticles(ParticleRangeId);
+
+		// Set mass from uniform mass
+		const TSet<int32> Vertices = Mesh.GetVertices();
+		for (int32 Index = 0; Index < Particles.GetRangeSize(); ++Index)
+		{
+			Particles.M(Index) = Vertices.Contains(Index) ? (Softs::FSolverReal)UniformMass : (Softs::FSolverReal)0.;
+		}
+		ParticleMassClampAndKinematicStateUpdate(Particles, (Softs::FSolverReal)MinPerParticleMass, KinematicPredicate);
 	}
+	else
+	{
+		// Retrieve the particle block size
+		const int32 Size = PBDEvolution->GetParticleRangeSize(ParticleRangeId);
 
-	ParticleMassClampAndKinematicStateUpdate(Offset, Size, (Softs::FSolverReal)MinPerParticleMass, KinematicPredicate);
+		// Set mass from uniform mass
+		const TSet<int32> Vertices = Mesh.GetVertices();
+		Softs::FSolverParticles& Particles = PBDEvolution->Particles();
+		for (int32 Index = ParticleRangeId; Index < ParticleRangeId + Size; ++Index)
+		{
+			Particles.M(Index) = Vertices.Contains(Index) ? (Softs::FSolverReal)UniformMass : (Softs::FSolverReal)0.;
+		}
+		ParticleMassClampAndKinematicStateUpdate(ParticleRangeId, Size, (Softs::FSolverReal)MinPerParticleMass, KinematicPredicate);
+	}
 }
 
-void FClothingSimulationSolver::SetParticleMassFromTotalMass(int32 Offset, FRealSingle TotalMass, FRealSingle MinPerParticleMass, const FTriangleMesh& Mesh, const TFunctionRef<bool(int32)>& KinematicPredicate)
+void FClothingSimulationSolver::SetParticleMassFromTotalMass(int32 ParticleRangeId, FRealSingle TotalMass, FRealSingle MinPerParticleMass, const FTriangleMesh& Mesh, const TFunctionRef<bool(int32)>& KinematicPredicate)
 {
-	// Retrieve the particle block size
-	const int32 Size = Evolution->GetParticleRangeSize(Offset);
+	if (Evolution)
+	{
+		Softs::FSolverParticlesRange& Particles = Evolution->GetSoftBodyParticles(ParticleRangeId);
+
+		// Set mass per area
+		const Softs::FSolverReal TotalArea = SetParticleMassPerArea(Particles, Mesh);
+
+		// Find density
+		const Softs::FSolverReal Density = TotalArea > (Softs::FSolverReal)0. ? (Softs::FSolverReal)TotalMass / TotalArea : (Softs::FSolverReal)1.;
+
+		// Update mass from mesh and density
+		ParticleMassUpdateDensity(Particles, Mesh, Density);
+
+		ParticleMassClampAndKinematicStateUpdate(Particles, (Softs::FSolverReal)MinPerParticleMass, KinematicPredicate);
+	}
+	else
+	{
+		// Retrieve the particle block size
+		const int32 Size = PBDEvolution->GetParticleRangeSize(ParticleRangeId);
 
-	// Set mass per area
-	const Softs::FSolverReal TotalArea = SetParticleMassPerArea(Offset, Size, Mesh);
+		// Set mass per area
+		const Softs::FSolverReal TotalArea = SetParticleMassPerArea(ParticleRangeId, Size, Mesh);
 
-	// Find density
-	const Softs::FSolverReal Density = TotalArea > (Softs::FSolverReal)0. ? (Softs::FSolverReal)TotalMass / TotalArea : (Softs::FSolverReal)1.;
+		// Find density
+		const Softs::FSolverReal Density = TotalArea > (Softs::FSolverReal)0. ? (Softs::FSolverReal)TotalMass / TotalArea : (Softs::FSolverReal)1.;
 
-	// Update mass from mesh and density
-	ParticleMassUpdateDensity(Mesh, Density);
+		// Update mass from mesh and density
+		ParticleMassUpdateDensity(Mesh, Density);
 
-	ParticleMassClampAndKinematicStateUpdate(Offset, Size, (Softs::FSolverReal)MinPerParticleMass, KinematicPredicate);
+		ParticleMassClampAndKinematicStateUpdate(ParticleRangeId, Size, (Softs::FSolverReal)MinPerParticleMass, KinematicPredicate);
+	}
 }
 
-void FClothingSimulationSolver::SetParticleMassFromDensity(int32 Offset, FRealSingle Density, FRealSingle MinPerParticleMass, const FTriangleMesh& Mesh, const TFunctionRef<bool(int32)>& KinematicPredicate)
+void FClothingSimulationSolver::SetParticleMassFromDensity(int32 ParticleRangeId, FRealSingle Density, FRealSingle MinPerParticleMass, const FTriangleMesh& Mesh, const TFunctionRef<bool(int32)>& KinematicPredicate)
 {
-	// Retrieve the particle block size
-	const int32 Size = Evolution->GetParticleRangeSize(Offset);
+	if (Evolution)
+	{
+		Softs::FSolverParticlesRange& Particles = Evolution->GetSoftBodyParticles(ParticleRangeId);
 
-	// Set mass per area
-	const Softs::FSolverReal TotalArea = SetParticleMassPerArea(Offset, Size, Mesh);
+		// Set mass per area
+		const Softs::FSolverReal TotalArea = SetParticleMassPerArea(Particles, Mesh);
 
-	// Set density from cm2 to m2
-	const Softs::FSolverReal DensityScaled = (Softs::FSolverReal)(Density / FMath::Square(ClothingSimulationSolverConstant::WorldScale));
+		// Set density from cm2 to m2
+		const Softs::FSolverReal DensityScaled = (Softs::FSolverReal)(Density / FMath::Square(ClothingSimulationSolverConstant::WorldScale));
 
-	// Update mass from mesh and density
-	ParticleMassUpdateDensity(Mesh, DensityScaled);
+		// Update mass from mesh and density
+		ParticleMassUpdateDensity(Particles, Mesh, DensityScaled);
+
+		ParticleMassClampAndKinematicStateUpdate(Particles, (Softs::FSolverReal)MinPerParticleMass, KinematicPredicate);
+	}
+	else
+	{
+		// Retrieve the particle block size
+		const int32 Size = PBDEvolution->GetParticleRangeSize(ParticleRangeId);
 
-	ParticleMassClampAndKinematicStateUpdate(Offset, Size, (Softs::FSolverReal)MinPerParticleMass, KinematicPredicate);
+		// Set mass per area
+		const Softs::FSolverReal TotalArea = SetParticleMassPerArea(ParticleRangeId, Size, Mesh);
+
+		// Set density from cm2 to m2
+		const Softs::FSolverReal DensityScaled = (Softs::FSolverReal)(Density / FMath::Square(ClothingSimulationSolverConstant::WorldScale));
+
+		// Update mass from mesh and density
+		ParticleMassUpdateDensity(Mesh, DensityScaled);
+
+		ParticleMassClampAndKinematicStateUpdate(ParticleRangeId, Size, (Softs::FSolverReal)MinPerParticleMass, KinematicPredicate);
+	}
 }
 
 void FClothingSimulationSolver::SetReferenceVelocityScale(
@@ -658,8 +1068,8 @@ void FClothingSimulationSolver::SetReferenceVelocityScale(
 	const FRigidTransform3& OldReferenceSpaceTransform,  // Transforms are in world space so have to be FReal based for LWC
 	const FRigidTransform3& ReferenceSpaceTransform,
 	const TVec3<FRealSingle>& LinearVelocityScale,
-	FRealSingle AngularVelocityScale,
-	FRealSingle FictitiousAngularScale)
+	FRealSingle AngularVelocityScale, FRealSingle FictitiousAngularScale
+)
 {
 	FRigidTransform3 OldRootBoneLocalTransform = OldReferenceSpaceTransform;
 	OldRootBoneLocalTransform.AddToTranslation(-OldLocalSpaceLocation);
@@ -685,22 +1095,78 @@ void FClothingSimulationSolver::SetReferenceVelocityScale(
 	const FReal PartialDeltaAngle = DeltaAngle * FMath::Clamp((FReal)1. - (FReal)AngularVelocityScale * SolverVelocityScale, (FReal)0., (FReal)1.);
 	DeltaRotation = UE::Math::TQuat<FReal>(Axis, PartialDeltaAngle);
 
-	// Transform points back into the previous frame of reference before applying the adjusted deltas
-	const FRigidTransform3 PreSimulationTransform = OldRootBoneLocalTransform.Inverse() * FRigidTransform3(DeltaPosition, DeltaRotation) * OldRootBoneLocalTransform;
-	PreSimulationTransforms[GroupId] = Softs::FSolverRigidTransform3(  // Store the delta in solver precision, no need for LWC here
-		Softs::FSolverVec3(PreSimulationTransform.GetTranslation()),
-		Softs::FSolverRotation3(PreSimulationTransform.GetRotation()));
+	// Transform points back into the previous frame of reference before applying the adjusted deltas
+	const FRigidTransform3 PreSimulationTransform = OldRootBoneLocalTransform.Inverse() * FRigidTransform3(DeltaPosition, DeltaRotation) * OldRootBoneLocalTransform;
+	PreSimulationTransforms[GroupId] = Softs::FSolverRigidTransform3(  // Store the delta in solver precision, no need for LWC here
+		Softs::FSolverVec3(PreSimulationTransform.GetTranslation()),
+		Softs::FSolverRotation3(PreSimulationTransform.GetRotation()));
+
+	// Fictitious angular scale only applied for PBDEvolution. It's applied by ExternalForces for Evolution.
+	const FReal AppliedFictitiousAngularScale = PBDEvolution ? FMath::Min((FReal)2., (FReal)FictitiousAngularScale) : (FReal)1.;
+
+	// Save the reference bone relative angular velocity for calculating the fictitious forces
+	const FVec3 FictitiousAngularDisplacement = ReferenceSpaceTransform.TransformVector(Axis * PartialDeltaAngle)
+		* AppliedFictitiousAngularScale;
+	FictitiousAngularDisplacements[GroupId] = Softs::FSolverVec3(FictitiousAngularDisplacement);
+	ReferenceSpaceLocations[GroupId] = ReferenceSpaceTransform.GetLocation() - LocalSpaceLocation;
+}
+
+Softs::FSolverReal FClothingSimulationSolver::SetParticleMassPerArea(Softs::FSolverParticlesRange& Particles, const FTriangleMesh& Mesh)
+{
+	// Zero out masses
+	for (int32 Index = 0; Index < Particles.GetRangeSize(); ++Index)
+	{
+		Particles.M(Index) = (Softs::FSolverReal)0.;
+	}
+
+	// Assign per particle mass proportional to connected area.
+	const TArray<TVec3<int32>>& SurfaceElements = Mesh.GetSurfaceElements();
+	Softs::FSolverReal TotalArea = (Softs::FSolverReal)0.;
+	for (const TVec3<int32>& Tri : SurfaceElements)
+	{
+		const Softs::FSolverReal TriArea = (Softs::FSolverReal)0.5 * Softs::FSolverVec3::CrossProduct(
+			Particles.X(Tri[1]) - Particles.X(Tri[0]),
+			Particles.X(Tri[2]) - Particles.X(Tri[0])).Size();
+		TotalArea += TriArea;
+		const Softs::FSolverReal ThirdTriArea = TriArea / (Softs::FSolverReal)3.;
+		Particles.M(Tri[0]) += ThirdTriArea;
+		Particles.M(Tri[1]) += ThirdTriArea;
+		Particles.M(Tri[2]) += ThirdTriArea;
+	}
+
+	UE_LOG(LogChaosCloth, Verbose, TEXT("Total area: %f, SI total area: %f"), TotalArea, TotalArea / FMath::Square(ClothingSimulationSolverConstant::WorldScale));
+	return TotalArea;
+
+}
+
+void FClothingSimulationSolver::ParticleMassUpdateDensity(Softs::FSolverParticlesRange& Particles, const FTriangleMesh& Mesh, Softs::FSolverReal Density)
+{
+	const TSet<int32> Vertices = Mesh.GetVertices();
+	FReal TotalMass = 0.f;
+	for (const int32 Vertex : Vertices)
+	{
+		Particles.M(Vertex) *= Density;
+		TotalMass += Particles.M(Vertex);
+	}
+
+	UE_LOG(LogChaosCloth, Verbose, TEXT("Total mass: %f, "), TotalMass);
+}
 
-	// Save the reference bone relative angular velocity for calculating the fictitious forces
-	const FVec3 FictitiousAngularDisplacement = ReferenceSpaceTransform.TransformVector(Axis * PartialDeltaAngle * FMath::Min((FReal)2., (FReal)FictitiousAngularScale));  // Clamp to 2x the delta angle
-	FictitiousAngularDisplacements[GroupId] = Softs::FSolverVec3(FictitiousAngularDisplacement);
-	ReferenceSpaceLocations[GroupId] = ReferenceSpaceTransform.GetLocation() - LocalSpaceLocation;
+void FClothingSimulationSolver::ParticleMassClampAndKinematicStateUpdate(Softs::FSolverParticlesRange& Particles, Softs::FSolverReal MinPerParticleMass, const TFunctionRef<bool(int32)>& KinematicPredicate)
+{
+	for (int32 Index = 0; Index < Particles.GetRangeSize(); ++Index)
+	{
+		Particles.M(Index) = FMath::Max(Particles.M(Index), MinPerParticleMass);
+		Particles.InvM(Index) = KinematicPredicate(Index) ? (Softs::FSolverReal)0. : (Softs::FSolverReal)1. / Particles.M(Index);
+	}
 }
 
 Softs::FSolverReal FClothingSimulationSolver::SetParticleMassPerArea(int32 Offset, int32 Size, const FTriangleMesh& Mesh)
 {
+	check(PBDEvolution);
+
 	// Zero out masses
-	Softs::FSolverParticles& Particles = Evolution->Particles();
+	Softs::FSolverParticles& Particles = PBDEvolution->Particles();
 	for (int32 Index = Offset; Index < Offset + Size; ++Index)
 	{
 		Particles.M(Index) = (Softs::FSolverReal)0.;
@@ -727,8 +1193,10 @@ Softs::FSolverReal FClothingSimulationSolver::SetParticleMassPerArea(int32 Offse
 
 void FClothingSimulationSolver::ParticleMassUpdateDensity(const FTriangleMesh& Mesh, Softs::FSolverReal Density)
 {
+	check(PBDEvolution);
+
 	const TSet<int32> Vertices = Mesh.GetVertices();
-	Softs::FSolverParticles& Particles = Evolution->Particles();
+	Softs::FSolverParticles& Particles = PBDEvolution->Particles();
 	FReal TotalMass = 0.f;
 	for (const int32 Vertex : Vertices)
 	{
@@ -741,7 +1209,9 @@ void FClothingSimulationSolver::ParticleMassUpdateDensity(const FTriangleMesh& M
 
 void FClothingSimulationSolver::ParticleMassClampAndKinematicStateUpdate(int32 Offset, int32 Size, Softs::FSolverReal MinPerParticleMass, const TFunctionRef<bool(int32)>& KinematicPredicate)
 {
-	Softs::FSolverParticles& Particles = Evolution->Particles();
+	check(PBDEvolution);
+
+	Softs::FSolverParticles& Particles = PBDEvolution->Particles();
 	for (int32 Index = Offset; Index < Offset + Size; ++Index)
 	{
 		Particles.M(Index) = FMath::Max(Particles.M(Index), MinPerParticleMass);
@@ -749,22 +1219,46 @@ void FClothingSimulationSolver::ParticleMassClampAndKinematicStateUpdate(int32 O
 	}
 }
 
+void FClothingSimulationSolver::SetProperties(int32 ParticleRangeId, const Softs::FCollectionPropertyConstFacade& InPropertyCollection)
+{
+	if (Evolution)
+	{
+		Evolution->SetSoftBodyProperties(ParticleRangeId, InPropertyCollection);
+
+		const uint32 GroupId = Evolution->GetSoftBodyGroupId(ParticleRangeId);
+		// Set properties to constraints that come from the solver (e.g., solver-level gravity, wind)
+		GetClothConstraints(ParticleRangeId).UpdateFromSolver(
+			Gravity, bIsClothGravityOverrideEnabled,
+			FictitiousAngularDisplacements[GroupId], ReferenceSpaceLocations[GroupId],
+			WindVelocity, LegacyWindAdaption);
+	}
+}
+
 void FClothingSimulationSolver::SetProperties(uint32 GroupId, FRealSingle DampingCoefficient, FRealSingle LocalDampingCoefficient, FRealSingle CollisionThickness, FRealSingle FrictionCoefficient)
 {
-	Evolution->SetDamping(DampingCoefficient, GroupId);
-	Evolution->SetLocalDamping(LocalDampingCoefficient, GroupId);
-	Evolution->SetCollisionThickness(CollisionThickness, GroupId);
-	Evolution->SetCoefficientOfFriction(FrictionCoefficient, GroupId);
+	if (PBDEvolution)
+	{
+		PBDEvolution->SetDamping(DampingCoefficient, GroupId);
+		PBDEvolution->SetLocalDamping(LocalDampingCoefficient, GroupId);
+		PBDEvolution->SetCollisionThickness(CollisionThickness, GroupId);
+		PBDEvolution->SetCoefficientOfFriction(FrictionCoefficient, GroupId);
+	}
 }
 
 void FClothingSimulationSolver::SetUseCCD(uint32 GroupId, bool bUseCCD)
 {
-	Evolution->SetUseCCD(bUseCCD, GroupId);
+	if (PBDEvolution)
+	{
+		PBDEvolution->SetUseCCD(bUseCCD, GroupId);
+	}
 }
 
 void FClothingSimulationSolver::SetGravity(uint32 GroupId, const TVec3<FRealSingle>& InGravity)
 {
-	Evolution->SetGravity(Softs::FSolverVec3(InGravity), GroupId);
+	if (PBDEvolution)
+	{
+		PBDEvolution->SetGravity(Softs::FSolverVec3(InGravity), GroupId);
+	}
 }
 
 void FClothingSimulationSolver::SetWindVelocity(const TVec3<FRealSingle>& InWindVelocity, FRealSingle InLegacyWindAdaption)
@@ -820,8 +1314,11 @@ int32 FClothingSimulationSolver::GetNumSubsteps() const
 
 void FClothingSimulationSolver::SetWindVelocity(uint32 GroupId, const TVec3<FRealSingle>& InWindVelocity)
 {
-	Softs::FVelocityAndPressureField& VelocityAndPressureField = Evolution->GetVelocityAndPressureField(GroupId);
-	VelocityAndPressureField.SetVelocity(Softs::FSolverVec3(InWindVelocity));
+	if (PBDEvolution)
+	{
+		Softs::FVelocityAndPressureField& VelocityAndPressureField = PBDEvolution->GetVelocityAndPressureField(GroupId);
+		VelocityAndPressureField.SetVelocity(Softs::FSolverVec3(InWindVelocity));
+	}
 }
 
 void FClothingSimulationSolver::SetWindAndPressureGeometry(
@@ -830,8 +1327,11 @@ void FClothingSimulationSolver::SetWindAndPressureGeometry(
 	const Softs::FCollectionPropertyConstFacade& InPropertyCollection,
 	const TMap<FString, TConstArrayView<FRealSingle>>& WeightMaps)
 {
-	Softs::FVelocityAndPressureField& VelocityAndPressureField = Evolution->GetVelocityAndPressureField(GroupId);
-	VelocityAndPressureField.SetGeometry(&TriangleMesh, InPropertyCollection, WeightMaps, ClothingSimulationSolverConstant::WorldScale);
+	if (PBDEvolution)
+	{
+		Softs::FVelocityAndPressureField& VelocityAndPressureField = PBDEvolution->GetVelocityAndPressureField(GroupId);
+		VelocityAndPressureField.SetGeometry(&TriangleMesh, InPropertyCollection, WeightMaps, ClothingSimulationSolverConstant::WorldScale);
+	}
 }
 
 void FClothingSimulationSolver::SetWindAndPressureGeometry(
@@ -841,8 +1341,11 @@ void FClothingSimulationSolver::SetWindAndPressureGeometry(
 	const TConstArrayView<FRealSingle>& LiftMultipliers,
 	const TConstArrayView<FRealSingle>& PressureMultipliers)
 {
-	Softs::FVelocityAndPressureField& VelocityAndPressureField = Evolution->GetVelocityAndPressureField(GroupId);
-	VelocityAndPressureField.SetGeometry(&TriangleMesh, DragMultipliers, LiftMultipliers, PressureMultipliers);
+	if (PBDEvolution)
+	{
+		Softs::FVelocityAndPressureField& VelocityAndPressureField = PBDEvolution->GetVelocityAndPressureField(GroupId);
+		VelocityAndPressureField.SetGeometry(&TriangleMesh, DragMultipliers, LiftMultipliers, PressureMultipliers);
+	}
 }
 
 void FClothingSimulationSolver::SetWindAndPressureProperties(
@@ -851,8 +1354,11 @@ void FClothingSimulationSolver::SetWindAndPressureProperties(
 	const TMap<FString, TConstArrayView<FRealSingle>>& WeightMaps,
 	bool bEnableAerodynamics)
 {
-	Softs::FVelocityAndPressureField& VelocityAndPressureField = Evolution->GetVelocityAndPressureField(GroupId);
-	VelocityAndPressureField.SetProperties(InPropertyCollection, WeightMaps, ClothingSimulationSolverConstant::WorldScale, bEnableAerodynamics);
+	if (PBDEvolution)
+	{
+		Softs::FVelocityAndPressureField& VelocityAndPressureField = PBDEvolution->GetVelocityAndPressureField(GroupId);
+		VelocityAndPressureField.SetProperties(InPropertyCollection, WeightMaps, ClothingSimulationSolverConstant::WorldScale, bEnableAerodynamics);
+	}
 }
 
 void FClothingSimulationSolver::SetWindAndPressureProperties(
@@ -862,22 +1368,42 @@ void FClothingSimulationSolver::SetWindAndPressureProperties(
 	FRealSingle FluidDensity,
 	const TVec2<FRealSingle>& Pressure)
 {
-	Softs::FVelocityAndPressureField& VelocityAndPressureField = Evolution->GetVelocityAndPressureField(GroupId);
-	VelocityAndPressureField.SetProperties(
-		Drag,
-		Lift,
-		FluidDensity / FMath::Cube(ClothingSimulationSolverConstant::WorldScale),  // Fluid density is given in kg/m^3. Need to convert to kg/cm^3 for solver.
-		Pressure / ClothingSimulationSolverConstant::WorldScale);  // UI Pressure is in kg/m s^2. Need to convert to kg/cm s^2 for solver.
+	if (PBDEvolution)
+	{
+		Softs::FVelocityAndPressureField& VelocityAndPressureField = PBDEvolution->GetVelocityAndPressureField(GroupId);
+		VelocityAndPressureField.SetProperties(
+			Drag,
+			Lift,
+			FluidDensity / FMath::Cube(ClothingSimulationSolverConstant::WorldScale),  // Fluid density is given in kg/m^3. Need to convert to kg/cm^3 for solver.
+			Pressure / ClothingSimulationSolverConstant::WorldScale);  // UI Pressure is in kg/m s^2. Need to convert to kg/cm s^2 for solver.
+	}
 }
 
 const Softs::FVelocityAndPressureField& FClothingSimulationSolver::GetWindVelocityAndPressureField(uint32 GroupId) const
 {
-	return Evolution->GetVelocityAndPressureField(GroupId);
+	if (Evolution)
+	{
+		// Return field from first active softbody in this group
+		const TSet<int32>& ParticleRanges = Evolution->GetGroupActiveSoftBodies(GroupId);
+		for (const int32 ParticleRangeId : ParticleRanges)
+		{
+			const TSharedPtr<Softs::FVelocityAndPressureField>& VelocityField = GetClothConstraints(ParticleRangeId).GetVelocityAndPressureField();
+			check(VelocityField);
+			return *VelocityField.Get();
+		}
+		checkf(false, TEXT("No active velocity field found for this groupId %i"), GroupId);
+		static Softs::FVelocityAndPressureField DummyField;
+		return DummyField;
+	}
+	else
+	{
+		return PBDEvolution->GetVelocityAndPressureField(GroupId);
+	}
 }
 
 void FClothingSimulationSolver::AddExternalForces(uint32 GroupId, bool bUseLegacyWind)
 {
-	if (Evolution)
+	if (PBDEvolution)
 	{
 		const FVec3& AngularDisplacement = FictitiousAngularDisplacements[GroupId];
 		const FVec3& ReferenceSpaceLocation = ReferenceSpaceLocations[GroupId];
@@ -886,7 +1412,7 @@ void FClothingSimulationSolver::AddExternalForces(uint32 GroupId, bool bUseLegac
 		static const FReal LegacyWindMultiplier = (FReal)25.;
 		const FVec3 LegacyWindVelocity = WindVelocity * LegacyWindMultiplier;
 
-		Evolution->GetForceFunction(GroupId) =
+		PBDEvolution->GetForceFunction(GroupId) =
 			[this, bHasFictitiousForces, bUseLegacyWind, LegacyWindVelocity, AngularDisplacement, ReferenceSpaceLocation](Softs::FSolverParticles& Particles, const FReal Dt, const int32 Index)
 			{
 				FVec3 Forces((FReal)0.);
@@ -943,13 +1469,115 @@ void FClothingSimulationSolver::ApplyPreSimulationTransforms()
 	TRACE_CPUPROFILER_EVENT_SCOPE(FClothingSimulationSolver_ApplyPreSimulationTransforms);
 	const Softs::FSolverVec3 DeltaLocalSpaceLocation(LocalSpaceLocation - OldLocalSpaceLocation);  // LWC, note that LocalSpaceLocation is FReal based. but the delta can be stored in FSolverReal
 
-	const TPBDActiveView<Softs::FSolverParticles>& ParticlesActiveView = Evolution->ParticlesActiveView();
-	const TArray<uint32>& ParticleGroupIds = Evolution->ParticleGroupIds();
-
 	const FSolverReal MaxVelocitySquared = (ClothSolverMaxVelocity > 0.f) ? FMath::Square((FSolverReal)ClothSolverMaxVelocity) : TNumericLimits<FSolverReal>::Max();
 
-	ParticlesActiveView.RangeFor(
-		[this, &ParticleGroupIds, &DeltaLocalSpaceLocation, MaxVelocitySquared](Softs::FSolverParticles& Particles, int32 Offset, int32 Range)
+	if (Evolution)
+	{
+		const TArray<uint32> ActiveGroups = Evolution->GetActiveGroups().Array();
+		PhysicsParallelFor(ActiveGroups.Num(), [this, &ActiveGroups, &DeltaLocalSpaceLocation, MaxVelocitySquared](int32 ActiveGroupIndex)
+		{
+			const uint32 GroupId = ActiveGroups[ActiveGroupIndex];
+
+		// Update particles
+		const TSet<int32>& ActiveSoftBodies = Evolution->GetGroupActiveSoftBodies(GroupId);
+		const Softs::FSolverRigidTransform3& GroupSpaceTransform = PreSimulationTransforms[GroupId];
+		for (const int32 SoftBodyId : ActiveSoftBodies)
+		{
+			Softs::FSolverParticlesRange& Particles = Evolution->GetSoftBodyParticles(SoftBodyId);
+#if INTEL_ISPC
+			if (bRealTypeCompatibleWithISPC && bChaos_PreSimulationTransforms_ISPC_Enabled)  // TODO: Make the ISPC works with both Single and Double depending on the FSolverReal type
+			{
+				if (MaxVelocitySquared == TNumericLimits<FSolverReal>::Max())
+				{
+					ispc::ApplyPreSimulationTransform(
+						(ispc::FVector4f*)Particles.GetPAndInvM().GetData(),
+						(ispc::FVector3f*)Particles.GetV().GetData(),
+						(ispc::FVector3f*)Particles.XArray().GetData(),
+						(ispc::FVector3f*)Particles.GetArrayView(OldAnimationPositions).GetData(),
+						(ispc::FVector3f*)Particles.GetArrayView(AnimationVelocities).GetData(),
+						Particles.GetInvM().GetData(),
+						(const ispc::FVector3f*)Particles.GetConstArrayView(AnimationPositions).GetData(),
+						(const ispc::FTransform3f&)GroupSpaceTransform,
+						(const ispc::FVector3f&)DeltaLocalSpaceLocation,
+						DeltaTime,
+						Particles.GetRangeSize());
+				}
+				else
+				{
+					ispc::ApplyPreSimulationTransformAndClampVelocity(
+						(ispc::FVector4f*)Particles.GetPAndInvM().GetData(),
+						(ispc::FVector3f*)Particles.GetV().GetData(),
+						(ispc::FVector3f*)Particles.XArray().GetData(),
+						(ispc::FVector3f*)Particles.GetArrayView(OldAnimationPositions).GetData(),
+						(ispc::FVector3f*)Particles.GetArrayView(AnimationVelocities).GetData(),
+						Particles.GetInvM().GetData(),
+						(const ispc::FVector3f*)Particles.GetConstArrayView(AnimationPositions).GetData(),
+						(const ispc::FTransform3f&)GroupSpaceTransform,
+						(ispc::FVector3f&)DeltaLocalSpaceLocation,
+						DeltaTime,
+						Particles.GetRangeSize(),
+						MaxVelocitySquared);
+				}
+			}
+			else
+#endif
+			{
+				TArrayView<Softs::FSolverVec3> OldAnimationPositionsView = Particles.GetArrayView(OldAnimationPositions);
+				TArrayView<Softs::FSolverVec3> AnimationVelocitiesView = Particles.GetArrayView(AnimationVelocities);
+				TConstArrayView<Softs::FSolverVec3> AnimationPositionsView = Particles.GetConstArrayView(AnimationPositions);
+
+				PhysicsParallelFor(Particles.GetRangeSize(),
+					[this, &GroupSpaceTransform, &OldAnimationPositionsView, &AnimationVelocitiesView, &AnimationPositionsView, &DeltaLocalSpaceLocation, &Particles, MaxVelocitySquared](int32 Index)
+				{
+
+					// Update initial state for particles
+					Particles.P(Index) = Particles.X(Index) = GroupSpaceTransform.TransformPositionNoScale(Particles.X(Index)) - DeltaLocalSpaceLocation;
+					Particles.V(Index) = GroupSpaceTransform.TransformVector(Particles.V(Index));
+
+					// Copy InvM over to PAndInvM
+					Particles.PAndInvM(Index).InvM = Particles.InvM(Index);
+
+					// Update anim initial state (target updated by skinning)
+					OldAnimationPositionsView[Index] = GroupSpaceTransform.TransformPositionNoScale(OldAnimationPositionsView[Index]) - DeltaLocalSpaceLocation;
+
+					// Update Animation velocities
+					AnimationVelocitiesView[Index] = (AnimationPositionsView[Index] - OldAnimationPositionsView[Index]) / DeltaTime;
+
+					// Clamp relative velocity
+					const FSolverVec3 RelVelocity = Particles.V(Index) - AnimationVelocitiesView[Index];
+					const FSolverReal RelVelocitySquaredLength = RelVelocity.SquaredLength();
+					if (RelVelocitySquaredLength > MaxVelocitySquared)
+					{
+						Particles.V(Index) = AnimationVelocitiesView[Index] + RelVelocity * FMath::Sqrt(MaxVelocitySquared / RelVelocitySquaredLength);
+					}
+					}, Particles.GetRangeSize() < ClothSolverMinParallelBatchSize);
+				}
+			}
+
+			// Update collision particles
+			const TSet<int32>& ActiveCollisionRanges = Evolution->GetGroupActiveCollisionParticleRanges(GroupId);
+			for (const int32 CollisionRangeId : ActiveCollisionRanges)
+			{
+				Softs::FSolverCollisionParticlesRange& CollisionParticles = Evolution->GetCollisionParticleRange(CollisionRangeId);
+				TArrayView<Softs::FSolverRigidTransform3> OldCollisionTransformsView = CollisionParticles.GetArrayView(OldCollisionTransforms);
+				for (int32 Index = 0; Index < CollisionParticles.GetRangeSize(); ++Index)
+				{
+					// Update initial state for collisions
+					OldCollisionTransformsView[Index] = OldCollisionTransformsView[Index] * GroupSpaceTransform;
+					OldCollisionTransformsView[Index].AddToTranslation(-DeltaLocalSpaceLocation);
+					CollisionParticles.X(Index) = OldCollisionTransformsView[Index].GetTranslation();
+					CollisionParticles.R(Index) = OldCollisionTransformsView[Index].GetRotation();
+				}
+			}
+		}, /*bForceSingleThreaded =*/ !bClothSolverParallelClothPreUpdate);
+	}
+	else
+	{
+		const TPBDActiveView<Softs::FSolverParticles>& ParticlesActiveView = PBDEvolution->ParticlesActiveView();
+		const TArray<uint32>& ParticleGroupIds = PBDEvolution->ParticleGroupIds();
+
+		ParticlesActiveView.RangeFor(
+			[this, &ParticleGroupIds, &DeltaLocalSpaceLocation, MaxVelocitySquared](Softs::FSolverParticles& Particles, int32 Offset, int32 Range)
 		{
 			TRACE_CPUPROFILER_EVENT_SCOPE(FClothingSimulationSolver_ParticlePreSimulationTransforms);
 			SCOPE_CYCLE_COUNTER(STAT_ChaosClothParticlePreSimulationTransforms);
@@ -1000,48 +1628,48 @@ void FClothingSimulationSolver::ApplyPreSimulationTransforms()
 			{
 				PhysicsParallelFor(RangeSize,
 					[this, &ParticleGroupIds, &DeltaLocalSpaceLocation, &Particles, Offset, MaxVelocitySquared](int32 i)
+				{
+					const int32 Index = Offset + i;
+					const Softs::FSolverRigidTransform3& GroupSpaceTransform = PreSimulationTransforms[ParticleGroupIds[Index]];
+
+					// Update initial state for particles
+					Particles.P(Index) = Particles.X(Index) = GroupSpaceTransform.TransformPositionNoScale(Particles.X(Index)) - DeltaLocalSpaceLocation;
+					Particles.V(Index) = GroupSpaceTransform.TransformVector(Particles.V(Index));
+
+					// Copy InvM over to PAndInvM
+					Particles.PAndInvM(Index).InvM = Particles.InvM(Index);
+
+					// Update anim initial state (target updated by skinning)
+					OldAnimationPositions[Index] = GroupSpaceTransform.TransformPositionNoScale(OldAnimationPositions[Index]) - DeltaLocalSpaceLocation;
+
+					// Update Animation velocities
+					AnimationVelocities[Index] = (AnimationPositions[Index] - OldAnimationPositions[Index]) / DeltaTime;
+
+					// Clamp relative velocity
+					const FSolverVec3 RelVelocity = Particles.V(Index) - AnimationVelocities[Index];
+					const FSolverReal RelVelocitySquaredLength = RelVelocity.SquaredLength();
+					if (RelVelocitySquaredLength > MaxVelocitySquared)
 					{
-						const int32 Index = Offset + i;
-						const Softs::FSolverRigidTransform3& GroupSpaceTransform = PreSimulationTransforms[ParticleGroupIds[Index]];
-
-						// Update initial state for particles
-						Particles.P(Index) = Particles.X(Index) = GroupSpaceTransform.TransformPositionNoScale(Particles.X(Index)) - DeltaLocalSpaceLocation;
-						Particles.V(Index) = GroupSpaceTransform.TransformVector(Particles.V(Index));
-
-						// Copy InvM over to PAndInvM
-						Particles.PAndInvM(Index).InvM = Particles.InvM(Index);
-
-						// Update anim initial state (target updated by skinning)
-						OldAnimationPositions[Index] = GroupSpaceTransform.TransformPositionNoScale(OldAnimationPositions[Index]) - DeltaLocalSpaceLocation;
-
-						// Update Animation velocities
-						AnimationVelocities[Index] = (AnimationPositions[Index] - OldAnimationPositions[Index]) / DeltaTime;
-
-						// Clamp relative velocity
-						const FSolverVec3 RelVelocity = Particles.V(Index) - AnimationVelocities[Index];
-						const FSolverReal RelVelocitySquaredLength = RelVelocity.SquaredLength();
-						if (RelVelocitySquaredLength > MaxVelocitySquared)
-						{
-							Particles.V(Index) = AnimationVelocities[Index] + RelVelocity * FMath::Sqrt(MaxVelocitySquared / RelVelocitySquaredLength);
-						}
-					}, RangeSize < ClothSolverMinParallelBatchSize);
+						Particles.V(Index) = AnimationVelocities[Index] + RelVelocity * FMath::Sqrt(MaxVelocitySquared / RelVelocitySquaredLength);
+					}
+				}, RangeSize < ClothSolverMinParallelBatchSize);
 			}
 		}, /*bForceSingleThreaded =*/ !bClothSolverParallelClothPreUpdate);
 
 #if FRAMEPRO_ENABLED
-	FRAMEPRO_CUSTOM_STAT("ChaosClothSolverMinParallelBatchSize", ClothSolverMinParallelBatchSize, "ChaosClothSolver", "Particles", FRAMEPRO_COLOUR(128,0,255));
-	FRAMEPRO_CUSTOM_STAT("ChaosClothSolverParallelClothPreUpdate", bClothSolverParallelClothPreUpdate, "ChaosClothSolver", "Enabled", FRAMEPRO_COLOUR(128, 128, 64));
+		FRAMEPRO_CUSTOM_STAT("ChaosClothSolverMinParallelBatchSize", ClothSolverMinParallelBatchSize, "ChaosClothSolver", "Particles", FRAMEPRO_COLOUR(128, 0, 255));
+		FRAMEPRO_CUSTOM_STAT("ChaosClothSolverParallelClothPreUpdate", bClothSolverParallelClothPreUpdate, "ChaosClothSolver", "Enabled", FRAMEPRO_COLOUR(128, 128, 64));
 #endif
 
-	{
-		TRACE_CPUPROFILER_EVENT_SCOPE(FClothingSimulationSolver_CollisionPreSimulationTransforms);
-		SCOPE_CYCLE_COUNTER(STAT_ChaosClothCollisionPreSimulationTransforms);
+		{
+			TRACE_CPUPROFILER_EVENT_SCOPE(FClothingSimulationSolver_CollisionPreSimulationTransforms);
+			SCOPE_CYCLE_COUNTER(STAT_ChaosClothCollisionPreSimulationTransforms);
 
-		const TPBDActiveView<Softs::FSolverCollisionParticles>& CollisionParticlesActiveView = Evolution->CollisionParticlesActiveView();
-		const TArray<uint32>& CollisionParticleGroupIds = Evolution->CollisionParticleGroupIds();
+			const TPBDActiveView<Softs::FSolverCollisionParticles>& CollisionParticlesActiveView = PBDEvolution->CollisionParticlesActiveView();
+			const TArray<uint32>& CollisionParticleGroupIds = PBDEvolution->CollisionParticleGroupIds();
 
-		CollisionParticlesActiveView.SequentialFor(  // There's unlikely to ever have enough collision particles for a parallel for
-			[this, &CollisionParticleGroupIds, &DeltaLocalSpaceLocation](Softs::FSolverCollisionParticles& CollisionParticles, int32 Index)
+			CollisionParticlesActiveView.SequentialFor(  // There's unlikely to ever have enough collision particles for a parallel for
+				[this, &CollisionParticleGroupIds, &DeltaLocalSpaceLocation](Softs::FSolverCollisionParticles& CollisionParticles, int32 Index)
 			{
 				const Softs::FSolverRigidTransform3& GroupSpaceTransform = PreSimulationTransforms[CollisionParticleGroupIds[Index]];
 
@@ -1051,6 +1679,7 @@ void FClothingSimulationSolver::ApplyPreSimulationTransforms()
 				CollisionParticles.X(Index) = OldCollisionTransforms[Index].GetTranslation();
 				CollisionParticles.R(Index) = OldCollisionTransforms[Index].GetRotation();
 			});
+		}
 	}
 }
 
@@ -1058,15 +1687,73 @@ void FClothingSimulationSolver::PreSubstep(const Softs::FSolverReal Interpolatio
 {
 	TRACE_CPUPROFILER_EVENT_SCOPE(FClothingSimulationSolver_PreSubstep);
 
-	const TPBDActiveView<Softs::FSolverParticles>& ParticlesActiveView = Evolution->ParticlesActiveView();
-	ParticlesActiveView.RangeFor(
-		[this, InterpolationAlpha](Softs::FSolverParticles& Particles, int32 Offset, int32 Range)
+	if (Evolution)
+	{
+		const TArray<uint32> ActiveGroups = Evolution->GetActiveGroups().Array();
+		PhysicsParallelFor(ActiveGroups.Num(), [this, &ActiveGroups, InterpolationAlpha](int32 ActiveGroupIndex)
+		{
+			const uint32 GroupId = ActiveGroups[ActiveGroupIndex];
+			const TSet<int32>& ActiveSoftBodies = Evolution->GetGroupActiveSoftBodies(GroupId);
+			for (const int32 SoftBodyId : ActiveSoftBodies)
+			{
+				TRACE_CPUPROFILER_EVENT_SCOPE(FClothingSimulationSolver_ParticlePreSubstepKinematicInterpolation);
+				SCOPE_CYCLE_COUNTER(STAT_ChaosClothParticlePreSubstepKinematicInterpolation);
+				Softs::FSolverParticlesRange& Particles = Evolution->GetSoftBodyParticles(SoftBodyId);
+
+#if INTEL_ISPC
+				if (bRealTypeCompatibleWithISPC && bChaos_PreSubstepInterpolation_ISPC_Enabled)  // TODO: Make the ISPC works with both Single and Double depending on the FSolverReal type
+				{
+					ispc::PreSubstepInterpolation(
+						(ispc::FVector3f*)Particles.GetArrayView(InterpolatedAnimationPositions).GetData(),
+						(ispc::FVector3f*)Particles.GetArrayView(InterpolatedAnimationNormals).GetData(),
+						(const ispc::FVector3f*)Particles.GetConstArrayView(AnimationPositions).GetData(),
+						(const ispc::FVector3f*)Particles.GetConstArrayView(OldAnimationPositions).GetData(),
+						(const ispc::FVector3f*)Particles.GetConstArrayView(AnimationNormals).GetData(),
+						(const ispc::FVector3f*)Particles.GetConstArrayView(OldAnimationNormals).GetData(),
+						InterpolationAlpha,
+						0,
+						Particles.GetRangeSize());
+				}
+				else
+#endif
+				{
+					TArrayView<Softs::FSolverVec3> InterpolatedAnimationPositionsView = Particles.GetArrayView(InterpolatedAnimationPositions);
+					TArrayView<Softs::FSolverVec3> InterpolatedAnimationNormalsView = Particles.GetArrayView(InterpolatedAnimationNormals);
+					TConstArrayView<Softs::FSolverVec3> AnimationPositionsView = Particles.GetConstArrayView(AnimationPositions);
+					TConstArrayView<Softs::FSolverVec3> OldAnimationPositionsView = Particles.GetConstArrayView(OldAnimationPositions);
+					TConstArrayView<Softs::FSolverVec3> AnimationNormalsView = Particles.GetConstArrayView(AnimationNormals);
+					TConstArrayView<Softs::FSolverVec3> OldAnimationNormalsView = Particles.GetConstArrayView(OldAnimationNormals);
+					PhysicsParallelFor(Particles.GetRangeSize(),
+						[&InterpolatedAnimationPositionsView, &InterpolatedAnimationNormalsView, &AnimationPositionsView,
+						&OldAnimationPositionsView, &AnimationNormalsView, &OldAnimationNormalsView, InterpolationAlpha](int32 Index)
+					{
+						InterpolatedAnimationPositionsView[Index] = InterpolationAlpha * AnimationPositionsView[Index] + ((Softs::FSolverReal)1. - InterpolationAlpha) * OldAnimationPositionsView[Index];
+						InterpolatedAnimationNormalsView[Index] = (InterpolationAlpha * AnimationNormalsView[Index] + ((Softs::FSolverReal)1. - InterpolationAlpha) * OldAnimationNormalsView[Index]).GetSafeNormal();
+
+					}, Particles.GetRangeSize() < ClothSolverMinParallelBatchSize);
+				}
+			}
+		}, /*bForceSingleThreaded =*/ !bClothSolverParallelClothPreUpdate);
+
+		{
+			TRACE_CPUPROFILER_EVENT_SCOPE(ChaosPBDClearCollidedArray);
+			memset(Collided.GetData(), 0, Collided.Num() * sizeof(bool));
+		}
+		CollisionContacts.Reset();
+		CollisionNormals.Reset();
+		CollisionPhis.Reset();
+	}
+	else
+	{
+		const TPBDActiveView<Softs::FSolverParticles>& ParticlesActiveView = PBDEvolution->ParticlesActiveView();
+		ParticlesActiveView.RangeFor(
+			[this, InterpolationAlpha](Softs::FSolverParticles& Particles, int32 Offset, int32 Range)
 		{
 			TRACE_CPUPROFILER_EVENT_SCOPE(FClothingSimulationSolver_ParticlePreSubstepKinematicInterpolation);
 			SCOPE_CYCLE_COUNTER(STAT_ChaosClothParticlePreSubstepKinematicInterpolation);
 
 			const int32 RangeSize = Range - Offset;
-
+	
 #if INTEL_ISPC
 			if (bRealTypeCompatibleWithISPC && bChaos_PreSubstepInterpolation_ISPC_Enabled)  // TODO: Make the ISPC works with both Single and Double depending on the FSolverReal type
 			{
@@ -1084,41 +1771,67 @@ void FClothingSimulationSolver::PreSubstep(const Softs::FSolverReal Interpolatio
 			else
 #endif
 			{
-				PhysicsParallelFor(RangeSize,
+			PhysicsParallelFor(RangeSize,
 					[this, Offset, InterpolationAlpha](int32 i)
-					{
-						const int32 Index = Offset + i;
-						InterpolatedAnimationPositions[Index] = InterpolationAlpha * AnimationPositions[Index] + ((Softs::FSolverReal)1. - InterpolationAlpha) * OldAnimationPositions[Index];
-						InterpolatedAnimationNormals[Index] = (InterpolationAlpha * AnimationNormals[Index] + ((Softs::FSolverReal)1. - InterpolationAlpha) * OldAnimationNormals[Index]).GetSafeNormal();
+				{
+					const int32 Index = Offset + i;
+					InterpolatedAnimationPositions[Index] = InterpolationAlpha * AnimationPositions[Index] + ((Softs::FSolverReal)1. - InterpolationAlpha) * OldAnimationPositions[Index];
+					InterpolatedAnimationNormals[Index] = (InterpolationAlpha * AnimationNormals[Index] + ((Softs::FSolverReal)1. - InterpolationAlpha) * OldAnimationNormals[Index]).GetSafeNormal();
 
-					}, RangeSize < ClothSolverMinParallelBatchSize);
+				}, RangeSize < ClothSolverMinParallelBatchSize);
 			}
-
-
 		}, /*bForceSingleThreaded =*/ !bClothSolverParallelClothPreUpdate);
+	}
 }
 
 void FClothingSimulationSolver::UpdateSolverField()
 {
-	if (Evolution && !PerSolverField.IsEmpty())
+	if ((Evolution || PBDEvolution) && !PerSolverField.IsEmpty())
 	{
 		TArray<FVector>& SamplePositions = PerSolverField.GetSamplePositions();
 		TArray<FFieldContextIndex>& SampleIndices = PerSolverField.GetSampleIndices();
 
-		const uint32 NumParticles = Evolution->Particles().Size();
-		const uint32 NumActiveParticles = Evolution->ParticlesActiveView().GetActiveSize();
+		if (Evolution)
+		{
+			const uint32 NumParticles = Evolution->GetParticles().Size();
+			const uint32 NumActiveParticles = Evolution->NumActiveParticles();
+
+			SamplePositions.SetNum(NumParticles, false);
+			SampleIndices.SetNum(NumActiveParticles, false);
+
+			int32 SampleIndex = 0;
+			for (const uint32 GroupId : Evolution->GetActiveGroups())
+			{
+				for (const int32 SoftBodyId : Evolution->GetGroupActiveSoftBodies(GroupId))
+				{
+					const Softs::FSolverParticlesRange& Particles = Evolution->GetSoftBodyParticles(SoftBodyId);
+					for (int32 LocalParticleIndex = 0; LocalParticleIndex < Particles.GetRangeSize(); ++LocalParticleIndex)
+					{
+						const int32 GlobalParticleIndex = LocalParticleIndex + Particles.GetOffset();
+						SamplePositions[GlobalParticleIndex] = FVector(Particles.X(LocalParticleIndex)) + LocalSpaceLocation;
+						SampleIndices[SampleIndex] = FFieldContextIndex(GlobalParticleIndex, SampleIndex);
+						++SampleIndex;
+					}
+				}
+			}
+		}
+		else
+		{
+			const uint32 NumParticles = PBDEvolution->Particles().Size();
+			const uint32 NumActiveParticles = PBDEvolution->ParticlesActiveView().GetActiveSize();
 
-		SamplePositions.SetNum(NumParticles, false);
-		SampleIndices.SetNum(NumActiveParticles, false);
+			SamplePositions.SetNum(NumParticles, false);
+			SampleIndices.SetNum(NumActiveParticles, false);
 
-		int32 SampleIndex = 0;
-		Evolution->ParticlesActiveView().SequentialFor(
-			[this, &SamplePositions, &SampleIndices, &SampleIndex](Softs::FSolverParticles& Particles, int32 ParticleIndex)
+			int32 SampleIndex = 0;
+			PBDEvolution->ParticlesActiveView().SequentialFor(
+				[this, &SamplePositions, &SampleIndices, &SampleIndex](Softs::FSolverParticles& Particles, int32 ParticleIndex)
 			{
 				SamplePositions[ParticleIndex] = FVector(Particles.X(ParticleIndex)) + LocalSpaceLocation;
 				SampleIndices[SampleIndex] = FFieldContextIndex(ParticleIndex, SampleIndex);
 				++SampleIndex;
 			});
+		}
 
 		PerSolverField.ComputeFieldLinearImpulse(GetTime());
 	}
@@ -1176,9 +1889,12 @@ void FClothingSimulationSolver::Update(Softs::FSolverReal InDeltaTime)
 			FClothingSimulationCloth* const Cloth = Cloths[ClothIndex];
 			const uint32 GroupId = Cloth->GetGroupId();
 
-			// Pre-update overridable solver properties first
-			Evolution->SetGravity(Gravity, GroupId);
-			Evolution->GetVelocityAndPressureField(GroupId).SetVelocity(WindVelocity);
+			if (PBDEvolution)
+			{
+				// Pre-update overridable solver properties first
+				PBDEvolution->SetGravity(Gravity, GroupId);
+				PBDEvolution->GetVelocityAndPressureField(GroupId).SetVelocity(WindVelocity);
+			}
 
 			Cloth->Update(this);
 		}, /*bForceSingleThreaded =*/ !bClothSolverParallelClothUpdate);
@@ -1215,27 +1931,35 @@ void FClothingSimulationSolver::Update(Softs::FSolverReal InDeltaTime)
 			check(Config);
 			const Softs::FCollectionPropertyFacade& Properties = Config->GetProperties(SolverLOD);
 
-			const int32 ConfigMaxNumIterations = FMath::Max(
-				Properties.GetValue<int32>(TEXT("MaxNumIterations"), ClothingSimulationSolverDefault::MaxNumIterations),
-				ClothingSimulationSolverDefault::MinNumIterations);
-
-			const int32 ConfigNumIterations = FMath::Clamp(
-				Properties.GetValue<int32>(TEXT("NumIterations"), ClothingSimulationSolverDefault::NumIterations),
-				ClothingSimulationSolverDefault::MinNumIterations,
-				ConfigMaxNumIterations);
-
 			const int32 ConfigNumSubsteps = FMath::Max(
 				Properties.GetValue<int32>(TEXT("NumSubsteps"), ClothingSimulationSolverDefault::NumSubsteps),
 				ClothingSimulationSolverDefault::MinNumSubsteps);
 
-			// Update solver time dependent parameters
-			const Softs::FSolverReal SolverFrequency = (Softs::FSolverReal)ClothingSimulationSolverDefault::SolverFrequency;  // 60Hz default TODO: Should this become a solver property?
+			if (Evolution)
+			{
+				Evolution->SetSolverProperties(Properties);
+				Evolution->SetDisableTimeDependentNumIterations(bClothSolverDisableTimeDependentNumIterations);
+			}
+			else
+			{
+				const int32 ConfigMaxNumIterations = FMath::Max(
+					Properties.GetValue<int32>(TEXT("MaxNumIterations"), ClothingSimulationSolverDefault::MaxNumIterations),
+					ClothingSimulationSolverDefault::MinNumIterations);
+
+				const int32 ConfigNumIterations = FMath::Clamp(
+					Properties.GetValue<int32>(TEXT("NumIterations"), ClothingSimulationSolverDefault::NumIterations),
+					ClothingSimulationSolverDefault::MinNumIterations,
+					ConfigMaxNumIterations);
 
-			const int32 TimeDependentNumIterations = bClothSolverDisableTimeDependentNumIterations ?
-				ConfigNumIterations :
-				FMath::RoundToInt32(SolverFrequency * DeltaTime * (Softs::FSolverReal)ConfigNumIterations);
+				// Update solver time dependent parameters
+				const Softs::FSolverReal SolverFrequency = (Softs::FSolverReal)ClothingSimulationSolverDefault::SolverFrequency;  // 60Hz default TODO: Should this become a solver property?
 
-			Evolution->SetIterations(FMath::Clamp(TimeDependentNumIterations, 1, ConfigMaxNumIterations));
+				const int32 TimeDependentNumIterations = bClothSolverDisableTimeDependentNumIterations ?
+					ConfigNumIterations :
+					FMath::RoundToInt32(SolverFrequency * DeltaTime * (Softs::FSolverReal)ConfigNumIterations);
+
+				PBDEvolution->SetIterations(FMath::Clamp(TimeDependentNumIterations, 1, ConfigMaxNumIterations));
+			}
 
 			// Advance substeps
 			const Softs::FSolverReal SubstepDeltaTime = DeltaTime / (Softs::FSolverReal)ConfigNumSubsteps;
@@ -1243,10 +1967,17 @@ void FClothingSimulationSolver::Update(Softs::FSolverReal InDeltaTime)
 			for (int32 i = 0; i < ConfigNumSubsteps; ++i)
 			{
 				PreSubstep(FMath::Clamp((Softs::FSolverReal)(i + 1) / (Softs::FSolverReal)ConfigNumSubsteps, (Softs::FSolverReal)0., (Softs::FSolverReal)1.));
-				Evolution->AdvanceOneTimeStep(SubstepDeltaTime);
+				if (Evolution)
+				{
+					Evolution->AdvanceOneTimeStep(SubstepDeltaTime, (Softs::FSolverReal)ConfigNumSubsteps);
+				}
+				else
+				{
+					PBDEvolution->AdvanceOneTimeStep(SubstepDeltaTime);
+				}
 			}
 
-			Time = Evolution->GetTime();
+			Time = Evolution ? Evolution->GetTime() : PBDEvolution->GetTime();
 			UE_LOG(LogChaosCloth, VeryVerbose, TEXT("DeltaTime: %.6f, Time = %.6f"), DeltaTime, Time);
 		}
 	}
@@ -1273,7 +2004,7 @@ void FClothingSimulationSolver::Update(Softs::FSolverReal InDeltaTime)
 
 void FClothingSimulationSolver::UpdateFromCache(const FClothingSimulationCacheData& CacheData)
 {
-	Chaos::Softs::FSolverParticles& SolverParticles = Evolution->Particles();
+	Chaos::Softs::FSolverParticles& SolverParticles = Evolution->GetParticles();
 	const int32 NumParticles = GetNumParticles();
 	const int32 NumCachedParticles = CacheData.CacheIndices.Num();	
 	const bool bHasVelocity = CacheData.CachedVelocities.Num() > 0;
@@ -1297,7 +2028,7 @@ void FClothingSimulationSolver::UpdateFromCache(const FClothingSimulationCacheDa
 
 void FClothingSimulationSolver::UpdateFromCache(const TArray<FVector>& CachedPositions, const TArray<FVector>& CachedVelocities) 
 {
-	Chaos::Softs::FSolverParticles& SolverParticles = Evolution->Particles();
+	Chaos::Softs::FSolverParticles& SolverParticles = Evolution->GetParticles();
 	const int32 NumParticles = GetNumParticles();
 	const bool bHasVelocity = CachedVelocities.Num() > 0;
 	if(CachedPositions.Num() == NumParticles)
@@ -1317,7 +2048,7 @@ void FClothingSimulationSolver::UpdateFromCache(const TArray<FVector>& CachedPos
 
 int32 FClothingSimulationSolver::GetNumUsedIterations() const
 {
-	return Evolution->GetIterations();
+	return Evolution ? Evolution->GetIterations() : PBDEvolution->GetIterations();
 }
 
 FBoxSphereBounds FClothingSimulationSolver::CalculateBounds() const
@@ -1325,18 +2056,93 @@ FBoxSphereBounds FClothingSimulationSolver::CalculateBounds() const
 	TRACE_CPUPROFILER_EVENT_SCOPE(FClothingSimulationSolver_CalculateBounds);
 	SCOPE_CYCLE_COUNTER(STAT_ChaosClothSolverCalculateBounds);
 
-	const TPBDActiveView<Softs::FSolverParticles>& ParticlesActiveView = Evolution->ParticlesActiveView();
-
-	if (ParticlesActiveView.HasActiveRange())
+	if (Evolution)
 	{
 		// Calculate bounding box
 		Softs::FSolverAABB3 BoundingBox = Softs::FSolverAABB3::EmptyAABB();
+		for (const uint32 GroupId : Evolution->GetActiveGroups())
+		{
+			for (const int32 SoftBodyId : Evolution->GetGroupActiveSoftBodies(GroupId))
+			{
+				const Softs::FSolverParticlesRange& Particles = Evolution->GetSoftBodyParticles(SoftBodyId);
+#if INTEL_ISPC
+				if (bRealTypeCompatibleWithISPC && bChaos_CalculateBounds_ISPC_Enabled)
+				{
+					Softs::FSolverVec3 NewMin = BoundingBox.Min();
+					Softs::FSolverVec3 NewMax = BoundingBox.Max();
+
+					ispc::CalculateBounds(
+						(ispc::FVector3f&)NewMin,
+						(ispc::FVector3f&)NewMax,
+						(const ispc::FVector3f*)Particles.XArray().GetData(),
+						0,
+						Particles.GetRangeSize());
+
+					BoundingBox.GrowToInclude(Softs::FSolverAABB3(NewMin, NewMax));
+				}
+				else
+#endif
+				{
+					for (const Softs::FSolverVec3& X : Particles.XArray())
+					{
+						BoundingBox.GrowToInclude(X);
+					}
+				}
+			}
+		}
+
+		if (BoundingBox.IsEmpty())
+		{
+			return FBoxSphereBounds(LocalSpaceLocation, FVector(0.f), 0.f);
+		}
+
+		// Calculate (squared) radius
+		const Softs::FSolverVec3 Center = BoundingBox.Center();
+		Softs::FSolverReal SquaredRadius = (Softs::FSolverReal)0.;
 
+		for (const uint32 GroupId : Evolution->GetActiveGroups())
+		{
+			for (const int32 SoftBodyId : Evolution->GetGroupActiveSoftBodies(GroupId))
+			{
+				const Softs::FSolverParticlesRange& Particles = Evolution->GetSoftBodyParticles(SoftBodyId);
 #if INTEL_ISPC
-		if (bRealTypeCompatibleWithISPC && bChaos_CalculateBounds_ISPC_Enabled)
+				if (bRealTypeCompatibleWithISPC && bChaos_CalculateBounds_ISPC_Enabled)
+				{
+					ispc::CalculateSquaredRadius(
+						SquaredRadius,
+						(const ispc::FVector3f&)Center,
+						(const ispc::FVector3f*)Particles.XArray().GetData(),
+						0,
+						Particles.GetRangeSize());
+				}
+				else
+#endif
+				{
+					for (const Softs::FSolverVec3& X : Particles.XArray())
+					{
+						SquaredRadius = FMath::Max(SquaredRadius, (X - Center).SizeSquared());
+					}
+				}
+			}
+		}
+
+		// Update bounds with this cloth
+		return FBoxSphereBounds(LocalSpaceLocation + BoundingBox.Center(), FVector(BoundingBox.Extents() * 0.5f), FMath::Sqrt(SquaredRadius));
+	}
+	else
+	{
+		const TPBDActiveView<Softs::FSolverParticles>& ParticlesActiveView = PBDEvolution->ParticlesActiveView();
+
+		if (ParticlesActiveView.HasActiveRange())
 		{
-			ParticlesActiveView.RangeFor(
-				[&BoundingBox](Softs::FSolverParticles& Particles, int32 Offset, int32 Range)
+			// Calculate bounding box
+			Softs::FSolverAABB3 BoundingBox = Softs::FSolverAABB3::EmptyAABB();
+
+#if INTEL_ISPC
+			if (bRealTypeCompatibleWithISPC && bChaos_CalculateBounds_ISPC_Enabled)
+			{
+				ParticlesActiveView.RangeFor(
+					[&BoundingBox](Softs::FSolverParticles& Particles, int32 Offset, int32 Range)
 				{
 					Softs::FSolverVec3 NewMin = BoundingBox.Min();
 					Softs::FSolverVec3 NewMax = BoundingBox.Max();
@@ -1350,26 +2156,26 @@ FBoxSphereBounds FClothingSimulationSolver::CalculateBounds() const
 
 					BoundingBox.GrowToInclude(Softs::FSolverAABB3(NewMin, NewMax));
 				});
-		}
-		else
+			}
+			else
 #endif
-		{
-			ParticlesActiveView.SequentialFor(
-				[&BoundingBox](Softs::FSolverParticles& Particles, int32 Index)
+			{
+				ParticlesActiveView.SequentialFor(
+					[&BoundingBox](Softs::FSolverParticles& Particles, int32 Index)
 				{
 					BoundingBox.GrowToInclude(Particles.X(Index));
 				});
-		}
+			}
 
-		// Calculate (squared) radius
-		const Softs::FSolverVec3 Center = BoundingBox.Center();
-		Softs::FSolverReal SquaredRadius = (Softs::FSolverReal)0.;
+			// Calculate (squared) radius
+			const Softs::FSolverVec3 Center = BoundingBox.Center();
+			Softs::FSolverReal SquaredRadius = (Softs::FSolverReal)0.;
 
 #if INTEL_ISPC
-		if (bRealTypeCompatibleWithISPC && bChaos_CalculateBounds_ISPC_Enabled)
-		{
-			ParticlesActiveView.RangeFor(
-				[&SquaredRadius, &Center](Softs::FSolverParticles& Particles, int32 Offset, int32 Range)
+			if (bRealTypeCompatibleWithISPC && bChaos_CalculateBounds_ISPC_Enabled)
+			{
+				ParticlesActiveView.RangeFor(
+					[&SquaredRadius, &Center](Softs::FSolverParticles& Particles, int32 Offset, int32 Range)
 				{
 					ispc::CalculateSquaredRadius(
 						SquaredRadius,
@@ -1378,27 +2184,38 @@ FBoxSphereBounds FClothingSimulationSolver::CalculateBounds() const
 						Offset,
 						Range);
 				});
-		}
-		else
+			}
+			else
 #endif
-		{
-			ParticlesActiveView.SequentialFor(
-				[&SquaredRadius, &Center](Softs::FSolverParticles& Particles, int32 Index)
+			{
+				ParticlesActiveView.SequentialFor(
+					[&SquaredRadius, &Center](Softs::FSolverParticles& Particles, int32 Index)
 				{
 					SquaredRadius = FMath::Max(SquaredRadius, (Particles.X(Index) - Center).SizeSquared());
 				});
+			}
+
+			// Update bounds with this cloth
+			return FBoxSphereBounds(LocalSpaceLocation + BoundingBox.Center(), FVector(BoundingBox.Extents() * 0.5f), FMath::Sqrt(SquaredRadius));
 		}
 
-		// Update bounds with this cloth
-		return FBoxSphereBounds(LocalSpaceLocation + BoundingBox.Center(), FVector(BoundingBox.Extents() * 0.5f), FMath::Sqrt(SquaredRadius));
+		return FBoxSphereBounds(LocalSpaceLocation, FVector(0.f), 0.f);
 	}
-
-	return FBoxSphereBounds(LocalSpaceLocation, FVector(0.f), 0.f);
 }
 
 uint32 FClothingSimulationSolver::GetNumParticles() const
 {
-	return Evolution->Particles().Size();
+	return Evolution ? Evolution->GetParticles().Size() : PBDEvolution->GetParticles().Size();
+}
+
+int32 FClothingSimulationSolver::GetNumActiveParticles() const
+{
+	return Evolution ? Evolution->NumActiveParticles() : PBDEvolution->ParticlesActiveView().GetActiveSize();
+}
+
+int32 FClothingSimulationSolver::GetGlobalParticleOffset(int32 ParticleRangeId) const
+{
+	return Evolution ? Evolution->GetSoftBodyParticles(ParticleRangeId).GetOffset() : ParticleRangeId;
 }
 
 } // End namespace Chaos
