Update the Cloth solver plugin to optionally use the new Softs::FEvolution.

diff --git a/Engine/Plugins/Experimental/ChaosCloth/Source/Chaos/Public/ChaosCloth/ChaosClothConstraints.h b/Engine/Plugins/Experimental/ChaosCloth/Source/Chaos/Public/ChaosCloth/ChaosClothConstraints.h
--- a/Engine/Plugins/Experimental/ChaosCloth/Source/Chaos/Public/ChaosCloth/ChaosClothConstraints.h
+++ b/Engine/Plugins/Experimental/ChaosCloth/Source/Chaos/Public/ChaosCloth/ChaosClothConstraints.h
@@ -27,14 +27,27 @@ PRAGMA_ENABLE_DEPRECATION_WARNINGS
 		~FClothConstraints();
 
 		// ---- Solver interface ----
-		UE_DEPRECATED(5.2, "Use the other Initialize supplying AnimationVelocities for correct AnimDrive behavior")
+		// Force-based solver
 		void Initialize(
-			Softs::FPBDEvolution* InEvolution,
-			const TArray<Softs::FSolverVec3>& InAnimationPositions,
-			const TArray<Softs::FSolverVec3>& InOldAnimationPositions,
-			const TArray<Softs::FSolverVec3>& InAnimationNormals,
-			int32 InParticleOffset,
-			int32 InNumParticles);
+			Softs::FEvolution* InEvolution,
+			FPerSolverFieldSystem* InPerSolverField,
+			const TArray<Softs::FSolverVec3>& InInterpolatedAnimationPositions,
+			const TArray<Softs::FSolverVec3>& InInterpolatedAnimationNormals,
+			const TArray<Softs::FSolverVec3>& InAnimationVelocities,
+			const TArray<Softs::FSolverVec3>& InNormals,
+			const TArray<Softs::FSolverRigidTransform3>& InLastSubframeCollisionTransformsCCD,
+			TArray<bool>& InCollisionParticleCollided,
+			TArray<Softs::FSolverVec3>& InCollisionContacts,
+			TArray<Softs::FSolverVec3>& InCollisionNormals,
+			TArray<Softs::FSolverReal>& InCollisionPhis,
+			int32 InParticleRangeId);
+
+		// Force-based solver
+		void UpdateFromSolver(const Softs::FSolverVec3& SolverGravity, bool bPerClothGravityOverrideEnabled,
+			const Softs::FSolverVec3& FictitiousAngularDisplacement, const Softs::FSolverVec3& ReferenceSpaceLocation,
+			const Softs::FSolverVec3& SolverWindVelocity, const Softs::FSolverReal LegacyWindAdaptation);
+
+		// PBD solver
 		void Initialize(
 			Softs::FPBDEvolution* InEvolution,
 			const TArray<Softs::FSolverVec3>& InInterpolatedAnimationPositions,
@@ -43,6 +56,7 @@ PRAGMA_ENABLE_DEPRECATION_WARNINGS
 			const TArray<Softs::FSolverVec3>& InAnimationVelocities,
 			int32 InParticleOffset,
 			int32 InNumParticles);
+
 		// ---- End of Solver interface ----
 
 		// ---- Cloth interface ----
@@ -95,106 +109,14 @@ PRAGMA_ENABLE_DEPRECATION_WARNINGS
 			Softs::FSolverReal MeshScale,
 			Softs::FSolverReal MaxDistancesScale = (Softs::FSolverReal)1.);
 
+		// NOTE: this only does something if using the PBDSolver. Force-based solver constraints are activated automatically
+		// when activating a particle range.
 		void Enable(bool bEnable);
-
-		UE_DEPRECATED(5.2, "Use AddRules() instead.")
-		void SetEdgeConstraints(const TArray<TVec3<int32>>& SurfaceElements, const TConstArrayView<FRealSingle>& StiffnessMultipliers, bool bUseXPBDConstraints);
-		UE_DEPRECATED(5.2, "Use AddRules() instead.")
-		void SetXPBDEdgeConstraints(const TArray<TVec3<int32>>& SurfaceElements, const TConstArrayView<FRealSingle>& StiffnessMultipliers, const TConstArrayView<FRealSingle>& DampingRatioMultipliers);
-		UE_DEPRECATED(5.2, "Use AddRules() instead.")
-		void SetBendingConstraints(const TArray<TVec2<int32>>& Edges, const TConstArrayView<FRealSingle>& StiffnessMultipliers, bool bUseXPBDConstraints);
-		UE_DEPRECATED(5.2, "Use AddRules() instead.")
-		void SetBendingConstraints(TArray<TVec4<int32>>&& BendingElements, const TConstArrayView<FRealSingle>& StiffnessMultipliers, const TConstArrayView<FRealSingle>& BucklingStiffnessMultipliers, bool bUseXPBDConstraints);
-		UE_DEPRECATED(5.2, "Use AddRules() instead.")
-		void SetXPBDBendingConstraints(TArray<TVec4<int32>>&& BendingElements, const TConstArrayView<FRealSingle>& StiffnessMultipliers, const TConstArrayView<FRealSingle>& BucklingStiffnessMultipliers, const TConstArrayView<FRealSingle>& DampingRatioMultipliers);
-		UE_DEPRECATED(5.1, "Use SetBendingConstraints(TArray<TVec4<int32>>&& BendingElements, const TConstArrayView<FRealSingle>& StiffnessMultipliers, const TConstArrayView<FRealSingle>& BucklingStiffnessMultipliers, bool bUseXPBDConstraints) instead.")
-		void SetBendingConstraints(TArray<TVec4<int32>>&& BendingElements, Softs::FSolverReal BendingStiffness);
-		UE_DEPRECATED(5.2, "Use AddRules() instead.")
-		void SetAreaConstraints(const TArray<TVec3<int32>>& SurfaceElements, const TConstArrayView<FRealSingle>& StiffnessMultipliers, bool bUseXPBDConstraints);
-		UE_DEPRECATED(5.2, "Use AddRules() instead.")
-		void SetVolumeConstraints(const TArray<TVec2<int32>>& DoubleBendingEdges, Softs::FSolverReal VolumeStiffness);
-		UE_DEPRECATED(5.2, "Use AddRules() instead.")
-		void SetVolumeConstraints(TArray<TVec3<int32>>&& SurfaceElements, Softs::FSolverReal VolumeStiffness);
-		UE_DEPRECATED(5.2, "Use AddRules() instead.")
-		void SetLongRangeConstraints(
-			const TArray<TConstArrayView<TTuple<int32,int32, FRealSingle>>>& Tethers,
-			const TConstArrayView<FRealSingle>& TetherStiffnessMultipliers,
-			const TConstArrayView<FRealSingle>& TetherScaleMultipliers,
-			const Softs::FSolverVec2& TetherScale,
-			Softs::FSolverReal MeshScale = (Softs::FSolverReal)1.);
-		UE_DEPRECATED(5.1, "LongRange XPBDConstraints are not supported.")
-		void SetLongRangeConstraints(
-			const TArray<TConstArrayView<TTuple<int32, int32, FRealSingle>>>& Tethers,
-			const TConstArrayView<FRealSingle>& TetherStiffnessMultipliers,
-			const TConstArrayView<FRealSingle>& TetherScaleMultipliers,
-			const Softs::FSolverVec2& TetherScale,
-			bool bUseXPBDConstraints,
-			Softs::FSolverReal MeshScale = (Softs::FSolverReal)1.);
-		UE_DEPRECATED(5.2, "Use AddRules() instead.")
-		void SetMaximumDistanceConstraints(const TConstArrayView<FRealSingle>& MaxDistances);
-		UE_DEPRECATED(5.2, "Use AddRules() instead.")
-		void SetBackstopConstraints(const TConstArrayView<FRealSingle>& BackstopDistances, const TConstArrayView<FRealSingle>& BackstopRadiuses, bool bUseLegacyBackstop);
-		UE_DEPRECATED(5.2, "Use AddRules() instead.")
-		void SetAnimDriveConstraints(const TConstArrayView<FRealSingle>& AnimDriveStiffnessMultipliers, const TConstArrayView<FRealSingle>& AnimDriveDampingMultipliers);
-		UE_DEPRECATED(5.2, "Use AddRules() instead.")
-		void SetShapeTargetConstraints(Softs::FSolverReal ShapeTargetStiffness);
-		UE_DEPRECATED(5.2, "Use AddRules() instead.")
-		void SetSelfCollisionConstraints(
-			const class FTriangleMesh& TriangleMesh,
-			TSet<TVec2<int32>>&& DisabledCollisionElements,
-			Softs::FSolverReal SelfCollisionThickness,
-			Softs::FSolverReal SelfCollisionFriction = Softs::FPBDCollisionSpringConstraintsBase::BackCompatFrictionCoefficient,
-			bool bGlobalIntersectionAnalysis = false,
-			bool bContourMinimization = false);
-
-		UE_DEPRECATED(5.2, "CreateRules will soon be made private, use AddRules() instead.")
-		void CreateRules();
-
-		UE_DEPRECATED(5.2, "Use Update() instead.")
-		void SetEdgeProperties(const Softs::FSolverVec2& EdgeStiffness, const Softs::FSolverVec2& DampingRatio = Softs::FSolverVec2::ZeroVector);
-		UE_DEPRECATED(5.2, "Use Update() instead.")
-		void SetBendingProperties(const Softs::FSolverVec2& BendingStiffness, Softs::FSolverReal BucklingRatio = 0.f, const Softs::FSolverVec2& BucklingStiffness = Softs::FSolverVec2::UnitVector, const Softs::FSolverVec2& DampingRatio = Softs::FSolverVec2::ZeroVector);
-		UE_DEPRECATED(5.2, "Use Update() instead.")
-		void SetAreaProperties(const Softs::FSolverVec2& AreaStiffness);
-		UE_DEPRECATED(5.2, "Use Update() instead.")
-		void SetThinShellVolumeProperties(Softs::FSolverReal VolumeStiffness);
-		UE_DEPRECATED(5.2, "Use Update() instead.")
-		void SetVolumeProperties(Softs::FSolverReal VolumeStiffness);
-		UE_DEPRECATED(5.2, "Use Update() instead.")
-		void SetLongRangeAttachmentProperties(
-			const Softs::FSolverVec2& TetherStiffness,
-			const Softs::FSolverVec2& TetherScale,
-			Softs::FSolverReal MeshScale = (Softs::FSolverReal)1.);
-		UE_DEPRECATED(5.2, "Use Update() instead.")
-		void SetMaximumDistanceProperties(Softs::FSolverReal MaxDistancesMultiplier);
-		UE_DEPRECATED(5.2, "Use Update() instead.")
-		void SetAnimDriveProperties(const Softs::FSolverVec2& AnimDriveStiffness, const Softs::FSolverVec2& AnimDriveDamping);
-		UE_DEPRECATED(5.2, "Use Update() instead.")
-		void SetSelfCollisionProperties(Softs::FSolverReal SelfCollisionThickness, Softs::FSolverReal SelfCollisionFriction = Softs::FPBDCollisionSpringConstraintsBase::BackCompatFrictionCoefficient, 
-			bool bGlobalIntersectionAnalysis = false, bool bContourMinimization = false);
-		UE_DEPRECATED(5.2, "Use Update() instead.")
-		void SetBackstopProperties(bool bEnabled, Softs::FSolverReal BackstopDistancesMultiplier);
 		// ---- End of Cloth interface ----
 
 		// ---- Debug functions ----
-		UE_DEPRECATED(5.2, "Use GetEdgeSpringConstraints() instead.")
-		const TSharedPtr<Softs::FPBDSpringConstraints>& GetEdgeConstraints() const { return EdgeConstraints_Deprecated; }
-		UE_DEPRECATED(5.2, "Use GetXEdgeSpringConstraints() instead.")
-		const TSharedPtr<Softs::FXPBDSpringConstraints>& GetXEdgeConstraints() const { return XEdgeConstraints_Deprecated; }
-		UE_DEPRECATED(5.2, "Use GetBendingSpringConstraints() instead.")
-		const TSharedPtr<Softs::FPBDSpringConstraints>& GetBendingConstraints() const { return BendingConstraints_Deprecated; }
-		UE_DEPRECATED(5.2, "Use GetXBendingSpringConstraints() instead.")
-		const TSharedPtr<Softs::FXPBDSpringConstraints>& GetXBendingConstraints() const { return XBendingConstraints_Deprecated; }
-		UE_DEPRECATED(5.2, "Use GetAreaSpringConstraints() instead.")
-		const TSharedPtr<Softs::FPBDAxialSpringConstraints>& GetAreaConstraints() const { return AreaConstraints_Deprecated; }
-		UE_DEPRECATED(5.2, "Use GetXAreaSpringConstraints() instead.")
-		const TSharedPtr<Softs::FXPBDAxialSpringConstraints>& GetXAreaConstraints() const { return XAreaConstraints_Deprecated; }
-		UE_DEPRECATED(5.2, "ThinShellVolume constraints are no longer supported.")
-		const TSharedPtr<Softs::FPBDSpringConstraints>& GetThinShellVolumeConstraints() const { return ThinShellVolumeConstraints_Deprecated; }
-		UE_DEPRECATED(5.2, "Volume constraints are no longer supported.")
-		const TSharedPtr<Softs::FPBDVolumeConstraint>& GetVolumeConstraints() const { return VolumeConstraints_Deprecated; }
-		UE_DEPRECATED(5.1,"LongRange XPBDConstraints are no longer supported.")
-		const TSharedPtr<Softs::FXPBDLongRangeConstraints>& GetXLongRangeConstraints() const { return XLongRangeConstraints_Deprecated; }
+		UE_DEPRECATED(5.4, "Shape constraints are no longer supported.")
+		const TSharedPtr<Softs::FPBDShapeConstraints>& GetShapeConstraints() const { return ShapeConstraints_Deprecated; }
 
 		const TSharedPtr<Softs::FPBDEdgeSpringConstraints>& GetEdgeSpringConstraints() const { return EdgeConstraints; }
 		const TSharedPtr<Softs::FXPBDEdgeSpringConstraints>& GetXEdgeSpringConstraints() const { return XEdgeConstraints; }
@@ -210,12 +132,14 @@ PRAGMA_ENABLE_DEPRECATION_WARNINGS
 		const TSharedPtr<Softs::FPBDSphericalConstraint>& GetMaximumDistanceConstraints() const { return MaximumDistanceConstraints; }
 		const TSharedPtr<Softs::FPBDSphericalBackstopConstraint>& GetBackstopConstraints() const { return BackstopConstraints; }
 		const TSharedPtr<Softs::FPBDAnimDriveConstraint>& GetAnimDriveConstraints() const { return AnimDriveConstraints; }
-		const TSharedPtr<Softs::FPBDShapeConstraints>& GetShapeConstraints() const { return ShapeConstraints; }
 		const TSharedPtr<Softs::FPBDCollisionSpringConstraints>& GetSelfCollisionConstraints() const { return SelfCollisionConstraints; }
 		const TSharedPtr<Softs::FPBDTriangleMeshIntersections>& GetSelfIntersectionConstraints() const { return SelfIntersectionConstraints; }
 		const TSharedPtr<Softs::FPBDTriangleMeshCollisions>& GetSelfCollisionInit() const { return SelfCollisionInit; }
 		const TSharedPtr<Softs::FPBDSelfCollisionSphereConstraints>& GetSelfCollisionSphereConstraints() const
 		{ return SelfCollisionSphereConstraints; }
+		const TSharedPtr<Softs::FVelocityAndPressureField>& GetVelocityAndPressureField() const { return VelocityAndPressureField; }
+		const TSharedPtr<Softs::FExternalForces>& GetExternalForces() const { return ExternalForces; }
+		const TSharedPtr<Softs::FPBDSoftBodyCollisionConstraint>& GetCollisionConstraint() const { return CollisionConstraint; }
 		// ---- End of debug functions ----
 
 	private:
@@ -253,19 +177,26 @@ PRAGMA_ENABLE_DEPRECATION_WARNINGS
 		void CreateAnimDriveConstraints(
 			const Softs::FCollectionPropertyConstFacade& ConfigProperties,
 			const TMap<FString, TConstArrayView<FRealSingle>>& WeightMaps);
+		void CreateVelocityAndPressureField(
+			const Softs::FCollectionPropertyConstFacade& ConfigProperties,
+			const TMap<FString, TConstArrayView<FRealSingle>>& WeightMaps,
+			const FTriangleMesh& TriangleMesh
+		);
+		void CreateExternalForces(
+			const Softs::FCollectionPropertyConstFacade& ConfigProperties
+		);
+		void CreateCollisionConstraint(
+			const Softs::FCollectionPropertyConstFacade& ConfigProperties,
+			Softs::FSolverReal MeshScale );
+
+		void CreateForceBasedRules();
+		void CreatePBDRules();
 		void CreateGSRules();
 		void GetGSNumRules();
 
+
 		//~ Begin deprecated constraints
-		TSharedPtr<Softs::FPBDSpringConstraints> EdgeConstraints_Deprecated;
-		TSharedPtr<Softs::FXPBDSpringConstraints> XEdgeConstraints_Deprecated;
-		TSharedPtr<Softs::FPBDSpringConstraints> BendingConstraints_Deprecated;
-		TSharedPtr<Softs::FXPBDSpringConstraints> XBendingConstraints_Deprecated;
-		TSharedPtr<Softs::FPBDAxialSpringConstraints> AreaConstraints_Deprecated;
-		TSharedPtr<Softs::FXPBDAxialSpringConstraints> XAreaConstraints_Deprecated;
-		TSharedPtr<Softs::FPBDSpringConstraints> ThinShellVolumeConstraints_Deprecated;;
-		TSharedPtr<Softs::FPBDVolumeConstraint> VolumeConstraints_Deprecated;
-		TSharedPtr<Softs::FXPBDLongRangeConstraints> XLongRangeConstraints_Deprecated;
+		TSharedPtr<Softs::FPBDShapeConstraints> ShapeConstraints_Deprecated;
 		//~ End deprecated constraints
 
 		TSharedPtr<Softs::FPBDEdgeSpringConstraints> EdgeConstraints;
@@ -282,29 +213,57 @@ PRAGMA_ENABLE_DEPRECATION_WARNINGS
 		TSharedPtr<Softs::FPBDSphericalConstraint> MaximumDistanceConstraints;
 		TSharedPtr<Softs::FPBDSphericalBackstopConstraint> BackstopConstraints;
 		TSharedPtr<Softs::FPBDAnimDriveConstraint> AnimDriveConstraints;
-		TSharedPtr<Softs::FPBDShapeConstraints> ShapeConstraints;
 		TSharedPtr<Softs::FPBDTriangleMeshCollisions> SelfCollisionInit;
 		TSharedPtr<Softs::FPBDCollisionSpringConstraints> SelfCollisionConstraints;
 		TSharedPtr<Softs::FPBDTriangleMeshIntersections> SelfIntersectionConstraints;
 		TSharedPtr<Softs::FPBDSelfCollisionSphereConstraints> SelfCollisionSphereConstraints;
 		TSharedPtr<Softs::FGaussSeidelMasterConstraint<Softs::FSolverReal, Softs::FSolverParticles>> GSMasterConstraint;
 		TSharedPtr<Softs::FGaussSeidelCorotatedCodimensionalConstraints<Softs::FSolverReal, Softs::FSolverParticles>> GSCorotatedCodimensionalConstraint;
-		
-		Softs::FPBDEvolution* Evolution;
+		//~ Begin Force-based solver only constraints
+		Softs::FSolverVec3 SolverWindVelocity; // Set from solver and added to wind from the config
+		TSharedPtr<Softs::FVelocityAndPressureField> VelocityAndPressureField;
+		TSharedPtr<Softs::FExternalForces> ExternalForces;
+		TSharedPtr<Softs::FPBDSoftBodyCollisionConstraint> CollisionConstraint;
+		//~ End Force-based solver only constraints
+
+		// Exactly one of these should be non-null
+		Softs::FEvolution* Evolution;
+		Softs::FPBDEvolution* PBDEvolution;
+
 		const TArray<Softs::FSolverVec3>* AnimationPositions;
-		const TArray<Softs::FSolverVec3>* OldAnimationPositions_Deprecated;
 		const TArray<Softs::FSolverVec3>* AnimationNormals;
 		const TArray<Softs::FSolverVec3>* AnimationVelocities;
 
 		int32 ParticleOffset;
+		int32 ParticleRangeId;
 		int32 NumParticles;
-		int32 ConstraintInitOffset;
-		int32 ConstraintRuleOffset;
-		int32 PostCollisionConstraintRuleOffset;
-		int32 PostprocessingConstraintRuleOffset;
+
 		int32 NumConstraintInits;
 		int32 NumConstraintRules;
 		int32 NumPostCollisionConstraintRules;
 		int32 NumPostprocessingConstraintRules;
+
+		//~ Begin Force-based solver only fields
+		FPerSolverFieldSystem* PerSolverField;
+		const TArray<Softs::FSolverVec3>* Normals;
+		const TArray<Softs::FSolverRigidTransform3>* LastSubframeCollisionTransformsCCD;
+		TArray<bool>* CollisionParticleCollided;
+		TArray<Softs::FSolverVec3>* CollisionContacts;
+		TArray<Softs::FSolverVec3>* CollisionNormals;
+		TArray<Softs::FSolverReal>* CollisionPhis;
+
+		int32 NumPreSubstepInits;
+		int32 NumExternalForceRules;
+		int32 NumPreSubstepConstraintRules;
+		int32 NumCollisionConstraintRules;
+		//~ End Force-based solver only fields
+
+		//~ Begin PBD solver only fields
+		const TArray<Softs::FSolverVec3>* OldAnimationPositions_Deprecated;
+		int32 ConstraintInitOffset;
+		int32 ConstraintRuleOffset;
+		int32 PostCollisionConstraintRuleOffset;
+		int32 PostprocessingConstraintRuleOffset;
+		//~ End PBD solver only fields
 	};
 } // namespace Chaos
diff --git a/Engine/Plugins/Experimental/ChaosCloth/Source/Chaos/Public/ChaosCloth/ChaosClothingSimulationCloth.h b/Engine/Plugins/Experimental/ChaosCloth/Source/Chaos/Public/ChaosCloth/ChaosClothingSimulationCloth.h
--- a/Engine/Plugins/Experimental/ChaosCloth/Source/Chaos/Public/ChaosCloth/ChaosClothingSimulationCloth.h
+++ b/Engine/Plugins/Experimental/ChaosCloth/Source/Chaos/Public/ChaosCloth/ChaosClothingSimulationCloth.h
@@ -170,8 +170,10 @@ PRAGMA_ENABLE_DEPRECATION_WARNINGS
 		CHAOSCLOTH_API TVec3<FRealSingle> GetGravity(const FClothingSimulationSolver* Solver) const;
 		// Return the current bounding box based on a given solver, not thread safe, call must be done right after the solver update.
 		CHAOSCLOTH_API FAABB3 CalculateBoundingBox(const FClothingSimulationSolver* Solver) const;
-		// Return the current LOD offset in the solver's particle array, or INDEX_NONE if no LOD is currently selected.
-		CHAOSCLOTH_API int32 GetOffset(const FClothingSimulationSolver* Solver) const;
+		// Return the current LOD ParticleRangeId, or INDEX_NONE if no LOD is currently selected.
+		CHAOSCLOTH_API int32 GetParticleRangeId(const FClothingSimulationSolver* Solver) const;
+		UE_DEPRECATED(5.4, "Offset has been renamed ParticleRangeId to reflect that it is no longer an offset.")
+		int32 GetOffset(const FClothingSimulationSolver* Solver) const { return GetParticleRangeId(Solver); }
 		// Return the current LOD num particles, or 0 if no LOD is currently selected.
 		CHAOSCLOTH_API int32 GetNumParticles(const FClothingSimulationSolver* Solver) const;
 		// Return the current LOD mesh.
@@ -203,7 +205,7 @@ PRAGMA_ENABLE_DEPRECATION_WARNINGS
 
 	private:
 		CHAOSCLOTH_API int32 GetNumParticles(int32 InLODIndex) const;
-		CHAOSCLOTH_API int32 GetOffset(const FClothingSimulationSolver* Solver, int32 InLODIndex) const;
+		CHAOSCLOTH_API int32 GetParticleRangeId(const FClothingSimulationSolver* Solver, int32 InLODIndex) const;
 
 	private:
 		struct FLODData;
diff --git a/Engine/Plugins/Experimental/ChaosCloth/Source/Chaos/Public/ChaosCloth/ChaosClothingSimulationCollider.h b/Engine/Plugins/Experimental/ChaosCloth/Source/Chaos/Public/ChaosCloth/ChaosClothingSimulationCollider.h
--- a/Engine/Plugins/Experimental/ChaosCloth/Source/Chaos/Public/ChaosCloth/ChaosClothingSimulationCollider.h
+++ b/Engine/Plugins/Experimental/ChaosCloth/Source/Chaos/Public/ChaosCloth/ChaosClothingSimulationCollider.h
@@ -103,11 +103,14 @@ namespace Chaos
 
 		CHAOSCLOTH_API int32 GetNumGeometries(int32 InSlotIndex) const;
 
-		// Return the collision particle offset for the specified slot being LODLess, external, or any of the LODs collision.
-		CHAOSCLOTH_API int32 GetOffset(const FClothingSimulationSolver* Solver, const FClothingSimulationCloth* Cloth, int32 InSlotIndex) const;
+		// Return the collision particle range ID for the specified slot being LODLess, external, or any of the LODs collision.
+		CHAOSCLOTH_API int32 GetCollisionRangeId(const FClothingSimulationSolver* Solver, const FClothingSimulationCloth* Cloth, int32 InSlotIndex) const;
 
-		// Return the collision particle offset and number of geometries for the specified type if valid. If ECollisionDataType::LODs is asked, then the offset returned is for the current LOD.
-		CHAOSCLOTH_API bool GetOffsetAndNumGeometries(const FClothingSimulationSolver* Solver, const FClothingSimulationCloth* Cloth, ECollisionDataType CollisionDataType, int32& OutOffset, int32& OutNumGeometries) const;
+		// Return the collision particle range ID for the specified type if valid. If ECollisionDataType::LODs is asked, then the offset returned is for the current LOD.
+		CHAOSCLOTH_API int32 GetCollisionRangeId(const FClothingSimulationSolver* Solver, const FClothingSimulationCloth* Cloth, ECollisionDataType CollisionDataType) const;
+
+		// Return the collision particle  range ID and number of geometries for the specified type if valid. If ECollisionDataType::LODs is asked, then the offset returned is for the current LOD.
+		CHAOSCLOTH_API bool GetCollisionRangeIdAndNumGeometries(const FClothingSimulationSolver* Solver, const FClothingSimulationCloth* Cloth, ECollisionDataType CollisionDataType, int32& OutCollisionRangeId, int32& OutNumGeometries) const;
 
 	private:
 		typedef TPair<const FClothingSimulationSolver*, const FClothingSimulationCloth*> FSolverClothPair;
diff --git a/Engine/Plugins/Experimental/ChaosCloth/Source/Chaos/Public/ChaosCloth/ChaosClothingSimulationSolver.h b/Engine/Plugins/Experimental/ChaosCloth/Source/Chaos/Public/ChaosCloth/ChaosClothingSimulationSolver.h
--- a/Engine/Plugins/Experimental/ChaosCloth/Source/Chaos/Public/ChaosCloth/ChaosClothingSimulationSolver.h
+++ b/Engine/Plugins/Experimental/ChaosCloth/Source/Chaos/Public/ChaosCloth/ChaosClothingSimulationSolver.h
@@ -31,7 +31,7 @@ namespace Chaos
 	class FClothingSimulationSolver final : public FPhysicsSolverEvents
 	{
 	public:
-		CHAOSCLOTH_API FClothingSimulationSolver();
+		CHAOSCLOTH_API FClothingSimulationSolver(bool bForceBasedSolver = false, FClothingSimulationConfig* InConfig = nullptr);
 		CHAOSCLOTH_API ~FClothingSimulationSolver();
 		
 		FClothingSimulationSolver(const FClothingSimulationSolver&) = delete;
@@ -39,6 +39,8 @@ namespace Chaos
 		FClothingSimulationSolver& operator=(const FClothingSimulationSolver&) = delete;
 		FClothingSimulationSolver& operator=(FClothingSimulationSolver&&) = delete;
 
+		bool IsForceBasedSolver() const { return !!Evolution; }
+
 		// ---- Animatable property setters ----
 		CHAOSCLOTH_API void SetLocalSpaceLocation(const FVec3& InLocalSpaceLocation, bool bReset = false);
 		const FVec3& GetLocalSpaceLocation() const { return LocalSpaceLocation; }
@@ -112,17 +114,18 @@ namespace Chaos
 
 		// ---- Cloth interface ----
 		CHAOSCLOTH_API int32 AddParticles(int32 NumParticles, uint32 GroupId);
-		CHAOSCLOTH_API void EnableParticles(int32 Offset, bool bEnable);
+		CHAOSCLOTH_API void EnableParticles(int32 ParticleRangeId, bool bEnable);
 
-		CHAOSCLOTH_API void ResetStartPose(int32 Offset, int32 NumParticles);
+		CHAOSCLOTH_API void ResetStartPose(int32 ParticleRangeId, int32 NumParticles);
 
 		// Get the current solver time
 		FSolverReal GetTime() const { return Time; }
-		CHAOSCLOTH_API void SetParticleMassUniform(int32 Offset, FRealSingle UniformMass, FRealSingle MinPerParticleMass, const FTriangleMesh& Mesh, const TFunctionRef<bool(int32)>& KinematicPredicate);
-		CHAOSCLOTH_API void SetParticleMassFromTotalMass(int32 Offset, FRealSingle TotalMass, FRealSingle MinPerParticleMass, const FTriangleMesh& Mesh, const TFunctionRef<bool(int32)>& KinematicPredicate);
-		CHAOSCLOTH_API void SetParticleMassFromDensity(int32 Offset, FRealSingle Density, FRealSingle MinPerParticleMass, const FTriangleMesh& Mesh, const TFunctionRef<bool(int32)>& KinematicPredicate);
+		CHAOSCLOTH_API void SetParticleMassUniform(int32 ParticleRangeId, FRealSingle UniformMass, FRealSingle MinPerParticleMass, const FTriangleMesh& Mesh, const TFunctionRef<bool(int32)>& KinematicPredicate);
+		CHAOSCLOTH_API void SetParticleMassFromTotalMass(int32 ParticleRangeId, FRealSingle TotalMass, FRealSingle MinPerParticleMass, const FTriangleMesh& Mesh, const TFunctionRef<bool(int32)>& KinematicPredicate);
+		CHAOSCLOTH_API void SetParticleMassFromDensity(int32 ParticleRangeId, FRealSingle Density, FRealSingle MinPerParticleMass, const FTriangleMesh& Mesh, const TFunctionRef<bool(int32)>& KinematicPredicate);
 
 		// Set the amount of velocity allowed to filter from the given change in reference space transform, including local simulation space.
+		// NOTE: Force-based solver does not apply FictitiousAngularScale here. It's applied directly via the PropertyCollection.
 		CHAOSCLOTH_API void SetReferenceVelocityScale(uint32 GroupId,
 			const FRigidTransform3& OldReferenceSpaceTransform,
 			const FRigidTransform3& ReferenceSpaceTransform,
@@ -130,7 +133,7 @@ namespace Chaos
 			FRealSingle AngularVelocityScale,
 			FRealSingle FictitiousAngularScale);
 
-		// Set general cloth simulation properties.
+		/** PBDSolver version */
 		CHAOSCLOTH_API void SetProperties(
 			uint32 GroupId,
 			FRealSingle DampingCoefficient,
@@ -138,6 +141,10 @@ namespace Chaos
 			FRealSingle CollisionThickness,
 			FRealSingle FrictionCoefficient);
 
+		/** Force based solver version */
+		CHAOSCLOTH_API void SetProperties(int32 ParticleRangeId, const Softs::FCollectionPropertyConstFacade& PropertyCollection);
+
+		/** Begin PBD-solver only property methods (these properties are controlled directly via PropertyCollection in force-based solver) */
 		// Set whether to use continuous collision detection.
 		CHAOSCLOTH_API void SetUseCCD(uint32 GroupId, bool bUseCCD);
 
@@ -180,72 +187,154 @@ namespace Chaos
 
 		// Add external forces to the particles
 		CHAOSCLOTH_API void AddExternalForces(uint32 GroupId, bool bUseLegacyWind);
+		/** End PBD-only based property methods */
 
 		const TArray<Softs::FSolverVec3>& GetOldAnimationPositions() const { return OldAnimationPositions; }
-		const Softs::FSolverVec3* GetOldAnimationPositions(int32 Offset) const { return OldAnimationPositions.GetData() + Offset; }
-		Softs::FSolverVec3* GetOldAnimationPositions(int32 Offset) { return OldAnimationPositions.GetData() + Offset; }
 		const TArray<Softs::FSolverVec3>& GetAnimationPositions() const { return AnimationPositions; }
-		const Softs::FSolverVec3* GetAnimationPositions(int32 Offset) const { return AnimationPositions.GetData() + Offset; }
-		Softs::FSolverVec3* GetAnimationPositions(int32 Offset) { return AnimationPositions.GetData() + Offset; }
 		const TArray<Softs::FSolverVec3>& GetInterpolatedAnimationPositions() const { return InterpolatedAnimationPositions; }
-		const Softs::FSolverVec3* GetInterpolatedAnimationPositions(int32 Offset) const { return InterpolatedAnimationPositions.GetData() + Offset; }
-		Softs::FSolverVec3* GetInterpolatedAnimationPositions(int32 Offset) { return InterpolatedAnimationPositions.GetData() + Offset; }
 		const TArray<Softs::FSolverVec3>& GetOldAnimationNormals() const { return OldAnimationNormals; }
-		const Softs::FSolverVec3* GetOldAnimationNormals(int32 Offset) const { return OldAnimationNormals.GetData() + Offset; }
-		Softs::FSolverVec3* GetOldAnimationNormals(int32 Offset) { return OldAnimationNormals.GetData() + Offset; }
 		const TArray<Softs::FSolverVec3>& GetAnimationNormals() const { return AnimationNormals; }
-		const Softs::FSolverVec3* GetAnimationNormals(int32 Offset) const { return AnimationNormals.GetData() + Offset; }
-		Softs::FSolverVec3* GetAnimationNormals(int32 Offset) { return AnimationNormals.GetData() + Offset; }
 		const TArray<Softs::FSolverVec3>& GetInterpolatedAnimationNormals() const { return InterpolatedAnimationNormals; }
-		const Softs::FSolverVec3* GetInterpolatedAnimationNormals(int32 Offset) const { return InterpolatedAnimationNormals.GetData() + Offset; }
-		Softs::FSolverVec3* GetInterpolatedAnimationNormals(int32 Offset) { return InterpolatedAnimationNormals.GetData() + Offset; }
 		const TArray<Softs::FSolverVec3>& GetNormals() const { return Normals; }
-		const Softs::FSolverVec3* GetNormals(int32 Offset) const { return Normals.GetData() + Offset; }
-		Softs::FSolverVec3* GetNormals(int32 Offset) { return Normals.GetData() + Offset; }
-		const TArray<FSolverVec3>& GetAnimationVelocities() const { return AnimationVelocities; }
-		const Softs::FSolverVec3* GetAnimationVelocities(int32 Offset) const { return AnimationVelocities.GetData() + Offset; }
-		Softs::FSolverVec3* GetAnimationVelocities(int32 Offset) { return AnimationVelocities.GetData() + Offset; }
+		const TArray<Softs::FSolverVec3>& GetAnimationVelocities() const { return AnimationVelocities; }
+
+		CHAOSCLOTH_API TConstArrayView<Softs::FSolverVec3> GetOldAnimationPositionsView(int32 ParticleRangeId) const;
+		CHAOSCLOTH_API TArrayView<Softs::FSolverVec3> GetOldAnimationPositionsView(int32 ParticleRangeId);
+		CHAOSCLOTH_API TConstArrayView<Softs::FSolverVec3> GetAnimationPositionsView(int32 ParticleRangeId) const;
+		CHAOSCLOTH_API TArrayView<Softs::FSolverVec3> GetAnimationPositionsView(int32 ParticleRangeId);
+		CHAOSCLOTH_API TConstArrayView<Softs::FSolverVec3> GetInterpolatedAnimationPositionsView(int32 ParticleRangeId) const;
+		CHAOSCLOTH_API TArrayView<Softs::FSolverVec3> GetInterpolatedAnimationPositionsView(int32 ParticleRangeId);
+		CHAOSCLOTH_API TConstArrayView<Softs::FSolverVec3> GetOldAnimationNormalsView(int32 ParticleRangeId) const;
+		CHAOSCLOTH_API TArrayView<Softs::FSolverVec3> GetOldAnimationNormalsView(int32 ParticleRangeId);
+		CHAOSCLOTH_API TConstArrayView<Softs::FSolverVec3> GetAnimationNormalsView(int32 ParticleRangeId) const;
+		CHAOSCLOTH_API TArrayView<Softs::FSolverVec3> GetAnimationNormalsView(int32 ParticleRangeId);
+		CHAOSCLOTH_API TConstArrayView<Softs::FSolverVec3> GetInterpolatedAnimationNormalsView(int32 ParticleRangeId) const;
+		CHAOSCLOTH_API TArrayView<Softs::FSolverVec3> GetInterpolatedAnimationNormalsView(int32 ParticleRangeId);
+		CHAOSCLOTH_API TConstArrayView<Softs::FSolverVec3> GetNormalsView(int32 ParticleRangeId) const;
+		CHAOSCLOTH_API TArrayView<Softs::FSolverVec3> GetNormalsView(int32 ParticleRangeId);
+		CHAOSCLOTH_API TConstArrayView<Softs::FSolverVec3> GetAnimationVelocitiesView(int32 ParticleRangeId) const;
+		CHAOSCLOTH_API TArrayView<Softs::FSolverVec3> GetAnimationVelocitiesView(int32 ParticleRangeId);
+
+		const Softs::FSolverVec3* GetOldAnimationPositions(int32 ParticleRangeId) const { return GetOldAnimationPositionsView(ParticleRangeId).GetData(); }
+		Softs::FSolverVec3* GetOldAnimationPositions(int32 ParticleRangeId) { return GetOldAnimationPositionsView(ParticleRangeId).GetData(); }
+		const Softs::FSolverVec3* GetAnimationPositions(int32 ParticleRangeId) const { return GetAnimationPositionsView(ParticleRangeId).GetData(); }
+		Softs::FSolverVec3* GetAnimationPositions(int32 ParticleRangeId) { return GetAnimationPositionsView(ParticleRangeId).GetData(); }
+		const Softs::FSolverVec3* GetInterpolatedAnimationPositions(int32 ParticleRangeId) const { return GetInterpolatedAnimationPositionsView(ParticleRangeId).GetData(); }
+		Softs::FSolverVec3* GetInterpolatedAnimationPositions(int32 ParticleRangeId) { return GetInterpolatedAnimationPositionsView(ParticleRangeId).GetData(); }
+		const Softs::FSolverVec3* GetOldAnimationNormals(int32 ParticleRangeId) const { return GetOldAnimationNormalsView(ParticleRangeId).GetData(); }
+		Softs::FSolverVec3* GetOldAnimationNormals(int32 ParticleRangeId) { return GetOldAnimationNormalsView(ParticleRangeId).GetData(); }
+		const Softs::FSolverVec3* GetAnimationNormals(int32 ParticleRangeId) const { return GetAnimationNormalsView(ParticleRangeId).GetData(); }
+		Softs::FSolverVec3* GetAnimationNormals(int32 ParticleRangeId) { return GetAnimationNormalsView(ParticleRangeId).GetData();
+		}
+		const Softs::FSolverVec3* GetInterpolatedAnimationNormals(int32 ParticleRangeId) const { return GetInterpolatedAnimationNormalsView(ParticleRangeId).GetData(); }
+		Softs::FSolverVec3* GetInterpolatedAnimationNormals(int32 ParticleRangeId) { return GetInterpolatedAnimationNormalsView(ParticleRangeId).GetData(); }
+		const Softs::FSolverVec3* GetNormals(int32 ParticleRangeId) const { return GetNormalsView(ParticleRangeId).GetData(); }
+		Softs::FSolverVec3* GetNormals(int32 ParticleRangeId) { return GetNormalsView(ParticleRangeId).GetData(); }
+		const Softs::FSolverVec3* GetAnimationVelocities(int32 ParticleRangeId) const { return GetAnimationVelocitiesView(ParticleRangeId).GetData(); }
+		Softs::FSolverVec3* GetAnimationVelocities(int32 ParticleRangeId) { return GetAnimationVelocitiesView(ParticleRangeId).GetData(); }
 
 		CHAOSCLOTH_API const TArray<Softs::FPAndInvM>& GetParticlePandInvMs() const;
-		CHAOSCLOTH_API const Softs::FPAndInvM* GetParticlePandInvMs(int32 Offset) const;
-		CHAOSCLOTH_API Softs::FPAndInvM* GetParticlePandInvMs(int32 Offset);		
 		CHAOSCLOTH_API const TArray<Softs::FSolverVec3>& GetParticleXs() const;
-		CHAOSCLOTH_API const Softs::FSolverVec3* GetParticleXs(int32 Offset) const;
-		CHAOSCLOTH_API Softs::FSolverVec3* GetParticleXs(int32 Offset);
 		CHAOSCLOTH_API const TArray<Softs::FSolverVec3>& GetParticleVs() const;
-		CHAOSCLOTH_API const Softs::FSolverVec3* GetParticleVs(int32 Offset) const;
-		CHAOSCLOTH_API Softs::FSolverVec3* GetParticleVs(int32 Offset);
 		CHAOSCLOTH_API const TArray<Softs::FSolverReal>& GetParticleInvMasses() const;
-		CHAOSCLOTH_API const Softs::FSolverReal* GetParticleInvMasses(int32 Offset) const;
-		const FClothConstraints& GetClothConstraints(int32 Offset) const { return *ClothsConstraints.FindChecked(Offset); }
-		FClothConstraints& GetClothConstraints(int32 Offset) { return *ClothsConstraints.FindChecked(Offset); }
+
+		CHAOSCLOTH_API TConstArrayView<Softs::FPAndInvM> GetParticlePandInvMsView(int32 ParticleRangeId) const;
+		CHAOSCLOTH_API TArrayView<Softs::FPAndInvM> GetParticlePandInvMsView(int32 ParticleRangeId);
+		CHAOSCLOTH_API TConstArrayView<Softs::FSolverVec3> GetParticleXsView(int32 ParticleRangeId) const;
+		CHAOSCLOTH_API TArrayView<Softs::FSolverVec3> GetParticleXsView(int32 ParticleRangeId);
+		CHAOSCLOTH_API TConstArrayView<Softs::FSolverVec3> GetParticleVsView(int32 ParticleRangeId) const;
+		CHAOSCLOTH_API TArrayView<Softs::FSolverVec3> GetParticleVsView(int32 ParticleRangeId);
+		CHAOSCLOTH_API TConstArrayView<Softs::FSolverReal> GetParticleInvMassesView(int32 ParticleRangeId) const;
+		CHAOSCLOTH_API TArrayView<Softs::FSolverReal> GetParticleInvMassesView(int32 ParticleRangeId);
+
+		const Softs::FPAndInvM* GetParticlePandInvMs(int32 ParticleRangeId) const 
+		{
+			return GetParticlePandInvMsView(ParticleRangeId).GetData();
+		}
+		Softs::FPAndInvM* GetParticlePandInvMs(int32 ParticleRangeId)
+		{
+			return GetParticlePandInvMsView(ParticleRangeId).GetData();
+		}
+		const Softs::FSolverVec3* GetParticleXs(int32 ParticleRangeId) const
+		{
+			return GetParticleXsView(ParticleRangeId).GetData();
+		}
+		Softs::FSolverVec3* GetParticleXs(int32 ParticleRangeId)
+		{
+			return GetParticleXsView(ParticleRangeId).GetData();
+		}
+		const Softs::FSolverVec3* GetParticleVs(int32 ParticleRangeId) const
+		{
+			return GetParticleVsView(ParticleRangeId).GetData();
+		}
+		Softs::FSolverVec3* GetParticleVs(int32 ParticleRangeId)
+		{
+			return GetParticleVsView(ParticleRangeId).GetData();
+		}
+		const Softs::FSolverReal* GetParticleInvMasses(int32 ParticleRangeId) const
+		{
+			return GetParticleInvMassesView(ParticleRangeId).GetData();
+		}
+
+		const FClothConstraints& GetClothConstraints(int32 ParticleRangeId) const { return *ClothsConstraints.FindChecked(ParticleRangeId); }
+		FClothConstraints& GetClothConstraints(int32 ParticleRangeId) { return *ClothsConstraints.FindChecked(ParticleRangeId); }
 		CHAOSCLOTH_API uint32 GetNumParticles() const;
+		CHAOSCLOTH_API int32 GetNumActiveParticles() const;
+		CHAOSCLOTH_API int32 GetGlobalParticleOffset(int32 ParticleRangeId) const;
 		// ---- End of the Cloth interface ----
 
 		// ---- Collider interface ----
-		CHAOSCLOTH_API int32 AddCollisionParticles(int32 NumCollisionParticles, uint32 GroupId, int32 RecycledOffset = 0);
-		CHAOSCLOTH_API void EnableCollisionParticles(int32 Offset, bool bEnable);
-
-		CHAOSCLOTH_API void ResetCollisionStartPose(int32 Offset, int32 NumCollisionParticles);
-
-		const int32* GetCollisionBoneIndices(int32 Offset) const { return CollisionBoneIndices.GetData() + Offset; }
-		int32* GetCollisionBoneIndices(int32 Offset) { return CollisionBoneIndices.GetData() + Offset; }
-		const Softs::FSolverRigidTransform3* GetCollisionBaseTransforms(int32 Offset) const { return CollisionBaseTransforms.GetData() + Offset; }
-		Softs::FSolverRigidTransform3* GetCollisionBaseTransforms(int32 Offset) { return CollisionBaseTransforms.GetData() + Offset; }
-		const Softs::FSolverRigidTransform3* GetOldCollisionTransforms(int32 Offset) const { return OldCollisionTransforms.GetData() + Offset; }
-		Softs::FSolverRigidTransform3* GetOldCollisionTransforms(int32 Offset) { return OldCollisionTransforms.GetData() + Offset; }
-		const Softs::FSolverRigidTransform3* GetCollisionTransforms(int32 Offset) const { return CollisionTransforms.GetData() + Offset; }
-		Softs::FSolverRigidTransform3* GetCollisionTransforms(int32 Offset) { return CollisionTransforms.GetData() + Offset; }
-		CHAOSCLOTH_API const Softs::FSolverVec3* GetCollisionParticleXs(int32 Offset) const;
-		CHAOSCLOTH_API Softs::FSolverVec3* GetCollisionParticleXs(int32 Offset);
-		CHAOSCLOTH_API const Softs::FSolverRotation3* GetCollisionParticleRs(int32 Offset) const;
-		CHAOSCLOTH_API Softs::FSolverRotation3* GetCollisionParticleRs(int32 Offset);
-		CHAOSCLOTH_API void SetCollisionGeometry(int32 Offset, int32 Index, FImplicitObjectPtr&& Geometry);
-		CHAOSCLOTH_API const FImplicitObjectPtr* GetCollisionGeometry(int32 Offset) const;
-		CHAOSCLOTH_API const bool* GetCollisionStatus(int32 Offset) const;
+		CHAOSCLOTH_API int32 AddCollisionParticles(int32 NumCollisionParticles, uint32 GroupId, int32 RecycledCollisionRangeId = 0);
+		CHAOSCLOTH_API void EnableCollisionParticles(int32 CollisionRangeId, bool bEnable);
+
+		CHAOSCLOTH_API void ResetCollisionStartPose(int32 CollisionRangeId, int32 NumCollisionParticles);
+
+		CHAOSCLOTH_API TConstArrayView<int32> GetCollisionBoneIndicesView(int32 CollisionRangeId) const;
+		CHAOSCLOTH_API TArrayView<int32> GetCollisionBoneIndicesView(int32 CollisionRangeId);
+		CHAOSCLOTH_API TConstArrayView<Softs::FSolverRigidTransform3> GetCollisionBaseTransformsView(int32 CollisionRangeId) const;
+		CHAOSCLOTH_API TArrayView<Softs::FSolverRigidTransform3> GetCollisionBaseTransformsView(int32 CollisionRangeId);
+		CHAOSCLOTH_API TConstArrayView<Softs::FSolverRigidTransform3> GetOldCollisionTransformsView(int32 CollisionRangeId) const;
+		CHAOSCLOTH_API TArrayView<Softs::FSolverRigidTransform3> GetOldCollisionTransformsView(int32 CollisionRangeId);
+		CHAOSCLOTH_API TConstArrayView<Softs::FSolverRigidTransform3> GetCollisionTransformsView(int32 CollisionRangeId) const;
+		CHAOSCLOTH_API TArrayView<Softs::FSolverRigidTransform3> GetCollisionTransformsView(int32 CollisionRangeId);
+		CHAOSCLOTH_API TConstArrayView<Softs::FSolverVec3> GetCollisionParticleXsView(int32 CollisionRangeId) const;
+		CHAOSCLOTH_API TArrayView<Softs::FSolverVec3> GetCollisionParticleXsView(int32 CollisionRangeId);
+		CHAOSCLOTH_API TConstArrayView<Softs::FSolverRotation3> GetCollisionParticleRsView(int32 CollisionRangeId) const;
+		CHAOSCLOTH_API TArrayView<Softs::FSolverRotation3> GetCollisionParticleRsView(int32 CollisionRangeId);
+		CHAOSCLOTH_API TConstArrayView<FImplicitObjectPtr> GetCollisionGeometryView(int32 CollisionRangeId) const;
+		CHAOSCLOTH_API TConstArrayView<bool> GetCollisionStatusView(int32 CollisionRangeId) const;
+		CHAOSCLOTH_API TConstArrayView<Softs::FSolverRigidTransform3> GetLastSubframeCollisionTransformsCCDView(int32 CollisionRangeId) const;
+		CHAOSCLOTH_API TArrayView<Softs::FSolverRigidTransform3> GetLastSubframeCollisionTransformsCCDView(int32 CollisionRangeId);
+
+		CHAOSCLOTH_API void SetCollisionGeometry(int32 CollisionRangeId, int32 Index, FImplicitObjectPtr&& Geometry);
+
 		CHAOSCLOTH_API const TArray<Softs::FSolverVec3>& GetCollisionContacts() const;
 		CHAOSCLOTH_API const TArray<Softs::FSolverVec3>& GetCollisionNormals() const;
 		CHAOSCLOTH_API const TArray<Softs::FSolverReal>& GetCollisionPhis() const;
+
+		const int32* GetCollisionBoneIndices(int32 CollisionRangeId) const { return GetCollisionBoneIndicesView(CollisionRangeId).GetData(); }
+		int32* GetCollisionBoneIndices(int32 CollisionRangeId) { return GetCollisionBoneIndicesView(CollisionRangeId).GetData(); }
+		const Softs::FSolverRigidTransform3* GetCollisionBaseTransforms(int32 CollisionRangeId) const { return GetCollisionBaseTransformsView(CollisionRangeId).GetData(); }
+		Softs::FSolverRigidTransform3* GetCollisionBaseTransforms(int32 CollisionRangeId) { return GetCollisionBaseTransformsView(CollisionRangeId).GetData(); }
+		const Softs::FSolverRigidTransform3* GetOldCollisionTransforms(int32 CollisionRangeId) const { return GetOldCollisionTransformsView(CollisionRangeId).GetData(); }
+		Softs::FSolverRigidTransform3* GetOldCollisionTransforms(int32 CollisionRangeId) { return GetOldCollisionTransformsView(CollisionRangeId).GetData(); }
+		const Softs::FSolverRigidTransform3* GetCollisionTransforms(int32 CollisionRangeId) const { return GetCollisionTransformsView(CollisionRangeId).GetData(); }
+		Softs::FSolverRigidTransform3* GetCollisionTransforms(int32 CollisionRangeId) { return GetCollisionTransformsView(CollisionRangeId).GetData(); }
+		const Softs::FSolverVec3* GetCollisionParticleXs(int32 CollisionRangeId) const { return GetCollisionParticleXsView(CollisionRangeId).GetData(); }
+		Softs::FSolverVec3* GetCollisionParticleXs(int32 CollisionRangeId) { return GetCollisionParticleXsView(CollisionRangeId).GetData(); }
+		const Softs::FSolverRotation3* GetCollisionParticleRs(int32 CollisionRangeId) const { return GetCollisionParticleRsView(CollisionRangeId).GetData(); }
+		Softs::FSolverRotation3* GetCollisionParticleRs(int32 CollisionRangeId) { return GetCollisionParticleRsView(CollisionRangeId).GetData(); }
+		const FImplicitObjectPtr* GetCollisionGeometry(int32 CollisionRangeId) const { return GetCollisionGeometryView(CollisionRangeId).GetData(); }
+		const bool* GetCollisionStatus(int32 CollisionRangeId) const { return GetCollisionStatusView(CollisionRangeId).GetData(); }
+		const Softs::FSolverRigidTransform3* GetLastSubframeCollisionTransformsCCD(int32 CollisionRangeId) const 
+		{
+			return GetLastSubframeCollisionTransformsCCDView(CollisionRangeId).GetData();
+		}
+		Softs::FSolverRigidTransform3* GetLastSubframeCollisionTransformsCCD(int32 CollisionRangeId)
+		{
+			return GetLastSubframeCollisionTransformsCCDView(CollisionRangeId).GetData();
+		}
+
 		// ---- End of the Collider interface ----
 
 		UE_DEPRECATED(5.4, "Use SetCollisionGeometry instead.")
@@ -260,19 +349,34 @@ namespace Chaos
 		// ---- End of the Field interface ----
 
 	private:
+		void Reset();
+
+		/** Begin Force-only methods */
+		void ParticleMassClampAndKinematicStateUpdate(Softs::FSolverParticlesRange& Particles, Softs::FSolverReal MinPerParticleMass, const TFunctionRef<bool(int32)>& KinematicPredicate);
+		Softs::FSolverReal SetParticleMassPerArea(Softs::FSolverParticlesRange& Particles, const FTriangleMesh& Mesh);
+		void ParticleMassUpdateDensity(Softs::FSolverParticlesRange& Particles, const FTriangleMesh& Mesh, Softs::FSolverReal Density);
+		/** End Force-only methods */
+
+		/** Begin PBD-only methods */
 		CHAOSCLOTH_API void ResetParticles();
 		CHAOSCLOTH_API void ResetCollisionParticles(int32 InCollisionParticlesOffset = 0);
-		CHAOSCLOTH_API void ApplyPreSimulationTransforms();
-		CHAOSCLOTH_API void PreSubstep(const Softs::FSolverReal InterpolationAlpha);
+		CHAOSCLOTH_API void ParticleMassClampAndKinematicStateUpdate(int32 Offset, int32 Size, Softs::FSolverReal MinPerParticleMass, const TFunctionRef<bool(int32)>& KinematicPredicate);
 		CHAOSCLOTH_API Softs::FSolverReal SetParticleMassPerArea(int32 Offset, int32 Size, const FTriangleMesh& Mesh);
 		CHAOSCLOTH_API void ParticleMassUpdateDensity(const FTriangleMesh& Mesh, Softs::FSolverReal Density);
-		CHAOSCLOTH_API void ParticleMassClampAndKinematicStateUpdate(int32 Offset, int32 Size, Softs::FSolverReal MinPerParticleMass, const TFunctionRef<bool(int32)>& KinematicPredicate);
+		/** End PBD - only methods */
+
+		CHAOSCLOTH_API void ApplyPreSimulationTransforms();
+		CHAOSCLOTH_API void PreSubstep(const Softs::FSolverReal InterpolationAlpha);
 
 		// Update the solver field forces/velocities at the particles location
 		CHAOSCLOTH_API void UpdateSolverField();
 
 	private:
-		TUniquePtr<Softs::FPBDEvolution> Evolution;
+
+		// Exactly one of these should be non-null
+		TUniquePtr<Softs::FEvolution> Evolution;
+		TUniquePtr<Softs::FPBDEvolution> PBDEvolution;
+
 
 		// Object arrays
 		TArray<FClothingSimulationCloth*> Cloths;
@@ -302,6 +406,8 @@ namespace Chaos
 		TArrayCollectionArray<Softs::FSolverRigidTransform3> CollisionBaseTransforms;
 		TArrayCollectionArray<Softs::FSolverRigidTransform3> OldCollisionTransforms;
 		TArrayCollectionArray<Softs::FSolverRigidTransform3> CollisionTransforms;
+		TArrayCollectionArray<Softs::FSolverRigidTransform3> LastSubframeCollisionTransformsCCD;
+		TArrayCollectionArray<bool> Collided;
 
 		// Cloth constraints
 		TMap<int32, TUniquePtr<FClothConstraints>> ClothsConstraints;
@@ -316,7 +422,12 @@ namespace Chaos
 		FSolverReal Time;
 		FSolverReal DeltaTime;
 
-		// Solver colliders offset
+		// FEvolution-only
+		TArray<FSolverVec3> CollisionContacts;
+		TArray<FSolverVec3> CollisionNormals;
+		TArray<FSolverReal> CollisionPhis;
+
+		// Solver colliders offset (PBD Evolution-only)
 		int32 CollisionParticlesOffset;  // Collision particle offset on the first solver/non cloth collider
 		int32 CollisionParticlesSize;  // Number of solver only colliders
 
diff --git a/Engine/Plugins/Experimental/ChaosClothAsset/Source/ChaosClothAssetEngine/Private/ChaosClothAsset/ClothSimulationProxy.cpp b/Engine/Plugins/Experimental/ChaosClothAsset/Source/ChaosClothAssetEngine/Private/ChaosClothAsset/ClothSimulationProxy.cpp
--- a/Engine/Plugins/Experimental/ChaosClothAsset/Source/ChaosClothAssetEngine/Private/ChaosClothAsset/ClothSimulationProxy.cpp
+++ b/Engine/Plugins/Experimental/ChaosClothAsset/Source/ChaosClothAssetEngine/Private/ChaosClothAsset/ClothSimulationProxy.cpp
@@ -88,14 +88,27 @@ namespace UE::Chaos::ClothAsset
 	FClothSimulationProxy::FClothSimulationProxy(const UChaosClothComponent& InClothComponent)
 		: ClothComponent(InClothComponent)
 		, ClothSimulationContext(MakeUnique<FClothSimulationContext>())
-		, Solver(MakeUnique<::Chaos::FClothingSimulationSolver>())
-		, Visualization(MakeUnique<::Chaos::FClothVisualization>(Solver.Get()))
+		, Solver(nullptr)
+		, Visualization(nullptr)
 		, MaxDeltaTime(UPhysicsSettings::Get()->MaxPhysicsDeltaTime)
 	{
 		using namespace ::Chaos;
 
 		check(IsInGameThread());
 
+		// Create solver config simulation thread object first. Need to know which solver type we're creating.
+		const int32 SolverConfigIndex = Configs.Emplace(MakeUnique<FClothingSimulationConfig>(ClothComponent.GetPropertyCollections()));  // TODO: Use a separate solver config for outfits
+
+		// Create a solver which can handle force-based solving if any of the LODs require it.
+		bool bForceBasedSolver = false;
+		Configs[SolverConfigIndex]->ForAllProperties([&bForceBasedSolver](Softs::FCollectionPropertyFacade& Property)
+		{
+			bForceBasedSolver = bForceBasedSolver || Property.GetValue<bool>(TEXT("EnableForceBasedSolver"), false);
+		});
+
+		Solver = MakeUnique<::Chaos::FClothingSimulationSolver>(bForceBasedSolver, Configs[SolverConfigIndex].Get());
+		Visualization = MakeUnique<::Chaos::FClothVisualization>(Solver.Get());
+
 		// Need a valid context to initialize the mesh
 		constexpr bool bIsInitialization = true;
 		constexpr Softs::FSolverReal NoAdvanceDt = 0.f;
@@ -153,10 +166,6 @@ PRAGMA_ENABLE_DEPRECATION_WARNINGS
 		NumKinematicParticles = Cloths[ClothIndex]->GetNumActiveKinematicParticles();
 		NumDynamicParticles = Cloths[ClothIndex]->GetNumActiveDynamicParticles();
 
-		// Create solver config simulation thread object
-		const int32 SolverConfigIndex = Configs.Emplace(MakeUnique<FClothingSimulationConfig>(ClothComponent.GetPropertyCollections()));  // TODO: Use a separate solver config for outfits
-		Solver->SetConfig(Configs[SolverConfigIndex].Get());
-
 		// Set start pose (update the context, then the solver without advancing the simulation)
 		ClothSimulationContext->Fill(ClothComponent, NoAdvanceDt, MaxDeltaTime);
 		Solver->Update((Softs::FSolverReal)ClothSimulationContext->DeltaTime);
@@ -348,7 +357,7 @@ PRAGMA_ENABLE_DEPRECATION_WARNINGS
 				Solver->Update(Softs::FSolverReal(0.));  // Update for LOD switching, but do not simulate
 			}
 
-			if (Cloth->GetOffset(Solver.Get()) == INDEX_NONE || Cloth->GetLODIndex(Solver.Get()) == INDEX_NONE)
+			if (Cloth->GetParticleRangeId(Solver.Get()) == INDEX_NONE || Cloth->GetLODIndex(Solver.Get()) == INDEX_NONE)
 			{
 				CurrentSimulationData.Remove(AssetIndex);  // Ensures that the cloth vertex factory won't run unnecessarily
 				continue;  // No valid LOD, there's nothing to write out
diff --git a/Engine/Plugins/Experimental/ChaosClothAssetEditor/Source/ChaosClothAssetDataflowNodes/Private/ChaosClothAsset/SimulationSolverConfigNode.cpp b/Engine/Plugins/Experimental/ChaosClothAssetEditor/Source/ChaosClothAssetDataflowNodes/Private/ChaosClothAsset/SimulationSolverConfigNode.cpp
--- a/Engine/Plugins/Experimental/ChaosClothAssetEditor/Source/ChaosClothAssetDataflowNodes/Private/ChaosClothAsset/SimulationSolverConfigNode.cpp
+++ b/Engine/Plugins/Experimental/ChaosClothAssetEditor/Source/ChaosClothAssetDataflowNodes/Private/ChaosClothAsset/SimulationSolverConfigNode.cpp
@@ -17,4 +17,9 @@ void FChaosClothAssetSimulationSolverConfigNode::AddProperties(Dataflow::FContex
 	UE_CHAOS_CLOTHASSET_SIMULATIONCONFIG_SETPROPERTY(NumIterations);
 	UE_CHAOS_CLOTHASSET_SIMULATIONCONFIG_SETPROPERTY(MaxNumIterations);
 	UE_CHAOS_CLOTHASSET_SIMULATIONCONFIG_SETPROPERTY(NumSubsteps);
+	UE_CHAOS_CLOTHASSET_SIMULATIONCONFIG_SETPROPERTYBOOL(EnableForceBasedSolver);
+	UE_CHAOS_CLOTHASSET_SIMULATIONCONFIG_SETPROPERTY(NumNewtonIterations);
+	UE_CHAOS_CLOTHASSET_SIMULATIONCONFIG_SETPROPERTY(MaxNumCGIterations);
+	UE_CHAOS_CLOTHASSET_SIMULATIONCONFIG_SETPROPERTY(CGResidualTolerance);
+	UE_CHAOS_CLOTHASSET_SIMULATIONCONFIG_SETPROPERTYBOOL(DoQuasistatics);
 }
diff --git a/Engine/Plugins/Experimental/ChaosClothAssetEditor/Source/ChaosClothAssetDataflowNodes/Private/ChaosClothAsset/TerminalNode.cpp b/Engine/Plugins/Experimental/ChaosClothAssetEditor/Source/ChaosClothAssetDataflowNodes/Private/ChaosClothAsset/TerminalNode.cpp
--- a/Engine/Plugins/Experimental/ChaosClothAssetEditor/Source/ChaosClothAssetDataflowNodes/Private/ChaosClothAsset/TerminalNode.cpp
+++ b/Engine/Plugins/Experimental/ChaosClothAssetEditor/Source/ChaosClothAssetDataflowNodes/Private/ChaosClothAsset/TerminalNode.cpp
@@ -51,7 +51,7 @@ namespace UE::Chaos::ClothAsset::Private
 		return Checksum;
 	}
 
-	bool PropertyKeysMatch(const TArray<TSharedRef<FManagedArrayCollection>>& Collections0, const TArray<TSharedRef<FManagedArrayCollection>>& Collections1)
+	bool PropertyKeysAndSolverTypesMatch(const TArray<TSharedRef<FManagedArrayCollection>>& Collections0, const TArray<TSharedRef<FManagedArrayCollection>>& Collections1)
 	{
 		if (Collections0.Num() != Collections1.Num())
 		{
@@ -67,10 +67,14 @@ namespace UE::Chaos::ClothAsset::Private
 			}
 			for (int32 PropertyIndex = 0; PropertyIndex < Property0.Num(); ++PropertyIndex)
 			{
-if (Property0.GetKey(PropertyIndex) != Property1.GetKey(PropertyIndex))
-{
-	return false;
-}
+				if (Property0.GetKey(PropertyIndex) != Property1.GetKey(PropertyIndex))
+				{
+					return false;
+				}
+			}
+			if (Property0.GetValue<bool>(TEXT("EnableForceBasedSolver"), false) != Property1.GetValue<bool>(TEXT("EnableForceBasedSolver"), false))
+			{
+				return false;
 			}
 		}
 		return true;
@@ -97,7 +101,7 @@ void FChaosClothAssetTerminalNode::SetAssetValue(TObjectPtr<UObject> Asset, Data
 		ClothColllectionChecksum = Private::CalculateClothChecksum(InClothCollections);
 		bClothCollectionChecksumValid = InClothCollections.Num() > 0;
 
-		if (bPreviousChecksumsValid && PreviousChecksum == ClothColllectionChecksum && Private::PropertyKeysMatch(InClothCollections, ClothCollections))
+		if (bPreviousChecksumsValid && PreviousChecksum == ClothColllectionChecksum && Private::PropertyKeysAndSolverTypesMatch(InClothCollections, ClothCollections))
 		{
 			// Cloth and property keys match. Just update property values.
 			check(InClothCollections.Num() == ClothCollections.Num());
diff --git a/Engine/Plugins/Experimental/ChaosClothAssetEditor/Source/ChaosClothAssetDataflowNodes/Public/ChaosClothAsset/SimulationSolverConfigNode.h b/Engine/Plugins/Experimental/ChaosClothAssetEditor/Source/ChaosClothAssetDataflowNodes/Public/ChaosClothAsset/SimulationSolverConfigNode.h
--- a/Engine/Plugins/Experimental/ChaosClothAssetEditor/Source/ChaosClothAssetDataflowNodes/Public/ChaosClothAsset/SimulationSolverConfigNode.h
+++ b/Engine/Plugins/Experimental/ChaosClothAssetEditor/Source/ChaosClothAssetDataflowNodes/Public/ChaosClothAsset/SimulationSolverConfigNode.h
@@ -19,7 +19,7 @@ public:
 	 * At lower fps up to MaxNumIterations may be used instead. At higher fps as low as one single iteration might be used.
 	 * Higher number of iterations will increase the stiffness of all constraints and improve convergence, but will also increase the CPU cost of the simulation.
 	 */
-	UPROPERTY(EditAnywhere, Category = Simulation, meta = (UIMin = "1", UIMax = "10", ClampMin = "1", ClampMax = "100"))
+	UPROPERTY(EditAnywhere, Category = Simulation, meta = (UIMin = "1", UIMax = "10", ClampMin = "0", ClampMax = "100"))
 	int32 NumIterations = 1;
 
 	/**
@@ -36,6 +36,36 @@ public:
 	UPROPERTY(EditAnywhere, Category = Simulation, meta = (UIMin = "1", UIMax = "10", ClampMin = "1", ClampMax = "100"))
 	int32 NumSubsteps = 1;
 
+	/**
+	 * Enable the higher accuracy force-based solver (experimental).
+	 */
+	UPROPERTY(EditAnywhere, Category = Experimental)
+	bool bEnableForceBasedSolver = false;
+	
+	/**
+	 * Number of Newton iterations for force-based solver
+	 */
+	UPROPERTY(EditAnywhere, Category = Experimental, meta = (UIMin = "1", UIMax = "10", ClampMin = "0", ClampMax = "100", EditCondition = "bEnableForceBasedSolver"))
+	int32 NumNewtonIterations = 1;
+
+	/**
+	 * Max number of CG Iterations per linear solve
+	 */
+	UPROPERTY(EditAnywhere, Category = Experimental, meta = (UIMin = "1", UIMax = "100", ClampMin = "1", ClampMax = "1000", EditCondition = "bEnableForceBasedSolver"))
+	int32 MaxNumCGIterations = 50;
+
+	/**
+	 * CG Tolerance
+	 */
+	UPROPERTY(EditAnywhere, Category = Experimental, meta = (ClampMin = "0", EditCondition = "bEnableForceBasedSolver"))
+	float CGResidualTolerance = 1e-4;
+
+	/**
+	 * Solve quasistaticly (no inertia) 
+	 */
+	UPROPERTY(EditAnywhere, Category = Experimental, meta = (EditCondition = "bEnableForceBasedSolver"))
+	bool bDoQuasistatics = false;
+
 	FChaosClothAssetSimulationSolverConfigNode(const Dataflow::FNodeParameters& InParam, FGuid InGuid = FGuid::NewGuid());
 
 private:
diff --git a/Engine/Source/Runtime/Core/Public/Math/Quat.isph b/Engine/Source/Runtime/Core/Public/Math/Quat.isph
--- a/Engine/Source/Runtime/Core/Public/Math/Quat.isph
+++ b/Engine/Source/Runtime/Core/Public/Math/Quat.isph
@@ -248,6 +248,30 @@ inline FVector4 VectorQuaternionRotateVector(const uniform FVector4& Quat, const
 	return Rotated;
 }
 
+inline FVector4f VectorQuaternionRotateVector(const uniform FVector4f& Quat, const FVector4f& VectorW0)
+{
+	// Q * V * Q.Inverse
+	//const VectorRegister InverseRotation = VectorQuaternionInverse(Quat);
+	//const VectorRegister Temp = VectorQuaternionMultiply2(Quat, VectorW0);
+	//const VectorRegister Rotated = VectorQuaternionMultiply2(Temp, InverseRotation);
+
+	// Equivalence of above can be shown to be:
+	// http://people.csail.mit.edu/bkph/articles/Quaternions.pdf
+	// V' = V + 2w(Q x V) + (2Q x (Q x V))
+	// refactor:
+	// V' = V + w(2(Q x V)) + (Q x (2(Q x V)))
+	// T = 2(Q x V);
+	// V' = V + w*(T) + (Q x T)
+
+	const uniform FVector4f QW = VectorReplicate(Quat, 3);
+	FVector4f T = VectorCross(Quat, VectorW0);
+	T = VectorAdd(T, T);
+	const FVector4f VTemp0 = VectorMultiplyAdd(QW, T, VectorW0);
+	const FVector4f VTemp1 = VectorCross(Quat, T);
+	const FVector4f Rotated = VectorAdd(VTemp0, VTemp1);
+	return Rotated;
+}
+
 inline uniform FVector4d VectorQuaternionRotateVector(const uniform FVector4d& Quat, const uniform FVector4d& VectorW0)
 {
 	// Q * V * Q.Inverse
diff --git a/Engine/Source/Runtime/Core/Public/Math/Transform.isph b/Engine/Source/Runtime/Core/Public/Math/Transform.isph
--- a/Engine/Source/Runtime/Core/Public/Math/Transform.isph
+++ b/Engine/Source/Runtime/Core/Public/Math/Transform.isph
@@ -279,6 +279,14 @@ inline FVector3f TransformVector(const FTransform3f &T, const FVector3f& V)
 	return RotatedVec;
 }
 
+inline FVector3f TransformVector(const uniform FTransform3f &T, const FVector3f& V)
+{
+	//RotatedVec = Q.Rotate(Scale*V.X, Scale*V.Y, Scale*V.Z, 0.f)
+	const FVector3f ScaledVec = SetVector(T.Scale3D) * V;
+	const FVector3f RotatedVec = VectorQuaternionRotateVector(T.Rotation, ScaledVec);
+	return RotatedVec;
+}
+
 inline FVector TransformVector(const uniform FTransform &T, const FVector& V)
 {
 	//RotatedVec = Q.Rotate(Scale*V.X, Scale*V.Y, Scale*V.Z, 0.f)
@@ -344,6 +352,22 @@ inline FVector TransformPosition(const uniform FTransform &T, const FVector& V)
 	return SetVector(TranslatedVec);
 }
 
+inline FVector3f TransformPosition(const uniform FTransform3f &T, const FVector3f& V)
+{
+	FVector4f InputVectorW0 = SetVector4(V.V[0], V.V[1], V.V[2], 0);
+
+	//Transform using QST is following
+	//QST(P) = Q.Rotate(S*P) + T where Q = quaternion, S = scale, T = translation
+
+	//RotatedVec = Q.Rotate(Scale*V.X, Scale*V.Y, Scale*V.Z, 0.f)
+	const FVector4f ScaledVec = VectorMultiply(T.Scale3D, InputVectorW0);
+	const FVector4f RotatedVec = VectorQuaternionRotateVector(T.Rotation, ScaledVec);
+
+	const FVector4f TranslatedVec = VectorAdd(RotatedVec, T.Translation);
+
+	return SetVector(TranslatedVec);
+}
+
 inline uniform FVector3d TransformPosition(const uniform FTransform3d &T, const uniform FVector3d& V)
 {
 	uniform FVector4d InputVectorW0 = SetVector4(V.V[0], V.V[1], V.V[2], 0);
diff --git a/Engine/Source/Runtime/Core/Public/Math/Vector.isph b/Engine/Source/Runtime/Core/Public/Math/Vector.isph
--- a/Engine/Source/Runtime/Core/Public/Math/Vector.isph
+++ b/Engine/Source/Runtime/Core/Public/Math/Vector.isph
@@ -819,6 +819,18 @@ inline FVector4f operator+(const FVector4f &A, const FVector4f &B)
 	return Result;
 }
 
+inline FVector4f operator+(const FVector4f &A, const uniform FVector4f &B)
+{
+	FVector4f Result;
+
+	for(uniform int i = 0; i < 4; i++)
+	{
+		Result.V[i] = A.V[i] + B.V[i];
+	}
+
+	return Result;
+}
+
 inline uniform FVector4d operator+(const uniform FVector4d &A, const uniform FVector4d &B)
 {
 	varying double S0, S1, Result;
@@ -931,6 +943,18 @@ inline FVector4f operator*(const FVector4f &A, const FVector4f &B)
 	return Result;
 }
 
+inline FVector4f operator*(const uniform FVector4f &A, const FVector4f &B)
+{
+	FVector4f Result;
+
+	for(uniform int i = 0; i < 4; i++)
+	{
+		Result.V[i] = A.V[i] * B.V[i];
+	}
+
+	return Result;
+}
+
 inline FVector4 operator*(const uniform FVector4 &A, const FVector4 &B)
 {
 	FVector4 Result;
@@ -2437,6 +2461,11 @@ inline FVector4 VectorAdd(const FVector4 &A, const uniform FVector4 &B)
 	return A + B;
 }
 
+inline FVector4f VectorAdd(const FVector4f &A, const uniform FVector4f &B)
+{
+	return A + B;
+}
+
 inline uniform FVector4d VectorAdd(const uniform FVector4d &A, const uniform FVector4d &B)
 {
 	return A + B;
@@ -2497,6 +2526,11 @@ inline FVector4 VectorMultiply(const uniform FVector4 &A, const FVector4 &B)
 	return A * B;
 }
 
+inline FVector4f VectorMultiply(const uniform FVector4f &A, const FVector4f &B)
+{
+	return A * B;
+}
+
 inline uniform FVector4d VectorMultiply(const uniform FVector4d &A, const uniform FVector4d &B)
 {
 	return A * B;
@@ -2585,6 +2619,18 @@ inline FVector4 VectorMultiplyAdd(const uniform FVector4& A, const FVector4 &B,
 	return Result;
 }
 
+inline FVector4f VectorMultiplyAdd(const uniform FVector4f& A, const FVector4f &B, const FVector4f &C)
+{
+	FVector4f Result;
+
+	Result.V[0] = A.V[0] * B.V[0] + C.V[0];
+	Result.V[1] = A.V[1] * B.V[1] + C.V[1];
+	Result.V[2] = A.V[2] * B.V[2] + C.V[2];
+	Result.V[3] = A.V[3] * B.V[3] + C.V[3];
+
+	return Result;
+}
+
 inline FVector4 VectorMultiplyAdd(const FReal F, const FVector4 &B, const FVector4 &C)
 {
 	FVector4 Result;
@@ -3342,6 +3388,16 @@ inline FVector4f VectorCross(const FVector4f& Vec1, const FVector4f& Vec2)
 	return VectorSwizzle(Tmp2, 1,2,0,3);
 }
 
+inline FVector4f VectorCross(const uniform FVector4f& Vec1, const FVector4f& Vec2)
+{
+	FVector4f Tmp0 = VectorSwizzle(Vec2, 1,2,0,3);
+	FVector4f Tmp1 = VectorSwizzle(Vec1, 1,2,0,3);
+	Tmp0 = Tmp0 * Vec1;
+	Tmp1 = Tmp1 * Vec2;
+	FVector4f Tmp2 = Tmp0 - Tmp1;
+	return VectorSwizzle(Tmp2, 1,2,0,3);
+}
+
 inline FVector4 VectorCross(const uniform FVector4& Vec1, const FVector4& Vec2)
 {
 	FVector4 Tmp0 = VectorSwizzle(Vec2, 1,2,0,3);
