Chaos Cloth: Handle duplicate triangles when building internal derived pattern data.

The ChaosClothingPatternData code was written before we did the rework of pattern data in general. I was able to simplify it significantly for the case where the 2d and 3d triangles match.

Beyond that, it will now emit a warning and discard the pattern data rather than check/crash if it's invalid.

diff --git a/Engine/Plugins/Experimental/ChaosCloth/Source/Chaos/Private/ChaosCloth/ChaosClothingPatternData.cpp b/Engine/Plugins/Experimental/ChaosCloth/Source/Chaos/Private/ChaosCloth/ChaosClothingPatternData.cpp
--- a/Engine/Plugins/Experimental/ChaosCloth/Source/Chaos/Private/ChaosCloth/ChaosClothingPatternData.cpp
+++ b/Engine/Plugins/Experimental/ChaosCloth/Source/Chaos/Private/ChaosCloth/ChaosClothingPatternData.cpp
@@ -1,5 +1,6 @@
 // Copyright Epic Games, Inc. All Rights Reserved.
 #include "ChaosCloth/ChaosClothingPatternData.h"
+#include "ChaosClothPrivate.h"
 #include "Containers/ArrayView.h"
 #include "Chaos/TriangleMesh.h"
 
@@ -33,7 +34,17 @@ static void InitializeTriangleMesh(FTriangleMesh& TriangleMesh, const TConstArra
 			 static_cast<int32>(Indices[Index + 2]) });
 	}
 
-	TriangleMesh.Init(MoveTemp(Elements), 0, NumParticles - 1);  // Init with the Offset to avoid discrepancies since the Triangle Mesh only relies on the used indices
+	constexpr bool bCullDegenerateElementsFalse = false;
+	TriangleMesh.Init(MoveTemp(Elements), 0, NumParticles - 1, bCullDegenerateElementsFalse);  // Init with the Offset to avoid discrepancies since the Triangle Mesh only relies on the used indices
+}
+
+void FClothingPatternData::Reset()
+{
+	// Resets all of the pattern data to empty.
+	PatternPositions = TConstArrayView<FVector2f>();
+	PatternToWeldedIndices = TConstArrayView<uint32>();
+	InitializeTriangleMesh(PatternTriangleMesh, TConstArrayView<uint32>(), 0);
+	WeldedFaceVertexPatternPositions.Reset();
 }
 
 void FClothingPatternData::GenerateDerivedPatternData(
@@ -41,6 +52,15 @@ void FClothingPatternData::GenerateDerivedPatternData(
 	const TConstArrayView<uint32>& Indices,
 	const TConstArrayView<uint32>& PatternIndices)
 {
+	if (Indices.Num() != PatternIndices.Num())
+	{
+		UE_LOG(LogChaosCloth, Warning, TEXT("Invalid pattern data. Num pattern indices %d != Num welded indices %d. Panel-based cloth is disabled for this LOD."), PatternIndices.Num(), Indices.Num());
+
+		// Reset all of the pattern data and return
+		Reset();
+		return;
+	}
+
 	const int32 NumPatternParticles = PatternPositions.Num();
 	InitializeTriangleMesh(PatternTriangleMesh, PatternIndices, NumPatternParticles);
 
@@ -67,24 +87,47 @@ void FClothingPatternData::GenerateDerivedPatternData(
 
 		TArray<TVec3<int32>> WeldedToPatternFaceVertexIndices;
 		WeldedToPatternFaceVertexIndices.Init(TVec3<int32>(INDEX_NONE), NoOffsetWeldedMesh.GetNumElements());
-		for (const TVec3<int32>& PatternElement : PatternElements)
+		for(int32 PatternElemIdx = 0; PatternElemIdx < PatternElements.Num(); ++PatternElemIdx)
 		{
+			const TVec3<int32>& PatternElement = PatternElements[PatternElemIdx];
+
 			// Find equivalent element in the welded mesh
 			const TVec3<int32> WeldedPatternElement = TVec3<int32>(PatternToWeldedIndices[PatternElement[0]], PatternToWeldedIndices[PatternElement[1]], PatternToWeldedIndices[PatternElement[2]]);
 			const TVec3<int32> SortedWeldedPatternElement = SortedElement(WeldedPatternElement);
-			const TArray<int32>& PossibleTriangles = WeldedPointToTriangleMap[WeldedPatternElement[0]];
 
+			// Cloth asset 2D and 3D elements should match each other, but this wasn't always a requirement.
 			int32 TriangleIndex = INDEX_NONE;
-			for (int32 Triangle : PossibleTriangles)
+			if (SortedWeldedPatternElement == SortedWeldedElements[PatternElemIdx] &&
+				WeldedToPatternFaceVertexIndices[PatternElemIdx] == TVec3<int32>(INDEX_NONE))
+			{
+				TriangleIndex = PatternElemIdx;
+			}
+			else
 			{
-				const TVec3<int32>& TriangleElements = SortedWeldedElements[Triangle];
-				if (TriangleElements == SortedWeldedPatternElement)
+				// Handle if that's not the case and try to find the matching triangle.
+				// Assume that there is some one-to-one correspondence between triangles, so pick a triangle that hasn't already been found.
+				const TArray<int32>& PossibleTriangles = WeldedPointToTriangleMap[WeldedPatternElement[0]];
+
+				for (int32 Triangle : PossibleTriangles)
 				{
-					TriangleIndex = Triangle;
-					break;
+					const TVec3<int32>& TriangleElements = SortedWeldedElements[Triangle];
+					if (TriangleElements == SortedWeldedPatternElement && 
+						WeldedToPatternFaceVertexIndices[Triangle] == TVec3<int32>(INDEX_NONE))
+					{
+						TriangleIndex = Triangle;
+						break;
+					}
 				}
 			}
-			check(TriangleIndex != INDEX_NONE);
+			if (TriangleIndex == INDEX_NONE)
+			{
+				// We failed to find a corresponding triangle. This pattern data is invalid.
+				UE_LOG(LogChaosCloth, Warning, TEXT("Invalid pattern data. Failed to find a matching welded triangle. Panel-based cloth is disabled for this LOD."));
+
+				// Reset all of the pattern data and return
+				Reset();
+				return;
+			}
 
 			const TVec3<int32>& WeldedElement = WeldedElements[TriangleIndex];
 			auto MatchIndex = [&WeldedPatternElement, &WeldedElement](int32 WeldedIndex)
@@ -106,7 +149,7 @@ void FClothingPatternData::GenerateDerivedPatternData(
 
 			// TODO: checkSlow
 			const TVec3<int32>& WeldedElement = WeldedElements[ElemIdx];
-			check(!(FaceVertex == TVec3<int32>(INDEX_NONE)));
+			check(!(FaceVertex == TVec3<int32>(INDEX_NONE))); // This should be OK to be a check since we already know that NumTriangles match between pattern and welded and each pattern triangle corresponds with a unique welded triangle.
 			check(PatternToWeldedIndices[FaceVertex[0]] == WeldedElement[0]);
 			check(PatternToWeldedIndices[FaceVertex[1]] == WeldedElement[1]);
 			check(PatternToWeldedIndices[FaceVertex[2]] == WeldedElement[2]);
diff --git a/Engine/Plugins/Experimental/ChaosCloth/Source/Chaos/Public/ChaosCloth/ChaosClothingPatternData.h b/Engine/Plugins/Experimental/ChaosCloth/Source/Chaos/Public/ChaosCloth/ChaosClothingPatternData.h
--- a/Engine/Plugins/Experimental/ChaosCloth/Source/Chaos/Public/ChaosCloth/ChaosClothingPatternData.h
+++ b/Engine/Plugins/Experimental/ChaosCloth/Source/Chaos/Public/ChaosCloth/ChaosClothingPatternData.h
@@ -11,8 +11,8 @@ struct FClothingPatternData
 {
 	FTriangleMesh PatternTriangleMesh;
 
-	const TConstArrayView<FVector2f> PatternPositions;
-	const TConstArrayView<uint32> PatternToWeldedIndices;
+	TConstArrayView<FVector2f> PatternPositions;
+	TConstArrayView<uint32> PatternToWeldedIndices;
 	TArray<TVec3<FVec2f>> WeldedFaceVertexPatternPositions; // This will be empty if there are no PatternPositions
 
 	FClothingPatternData(
@@ -23,6 +23,8 @@ struct FClothingPatternData
 		const TConstArrayView<uint32>& InPatternToWeldedIndices);
 
 private:
+	void Reset();
+
 	void GenerateDerivedPatternData(
 		int32 InNumParticles,
 		const TConstArrayView<uint32>& InIndices,
