Softs Evolution with linear system solver.

Currently only XPBD springs, gravity, and collision proximity forces are supported. Very experimental.

This uses Eigen for the linear solve. No optimization has been done to make building the matrix efficient. Mostly for testing right now as it's super slow and not particularly stable!

Also updated XPBD spring stiffness and damping to be actually resolution independent now that I actually have a ground truth. XPBD has weird approximations that make it so tuning it to be resolution independent is an unknown problem, so whatever I had done before wasn't actually correct or generalizable.

diff --git a/Engine/Plugins/Experimental/ChaosCloth/Source/Chaos/Private/ChaosCloth/ChaosClothConstraints.cpp b/Engine/Plugins/Experimental/ChaosCloth/Source/Chaos/Private/ChaosCloth/ChaosClothConstraints.cpp
--- a/Engine/Plugins/Experimental/ChaosCloth/Source/Chaos/Private/ChaosCloth/ChaosClothConstraints.cpp
+++ b/Engine/Plugins/Experimental/ChaosCloth/Source/Chaos/Private/ChaosCloth/ChaosClothConstraints.cpp
@@ -124,6 +124,8 @@ FClothConstraints::FClothConstraints()
 	, NumExternalForceRules(0)
 	, NumPreSubstepConstraintRules(0)
 	, NumCollisionConstraintRules(0)
+	, NumUpdateLinearSystemRules(0)
+	, NumUpdateLinearSystemCollisionsRules(0)
 
 	, OldAnimationPositions_Deprecated(nullptr)
 	, ConstraintInitOffset(INDEX_NONE)
@@ -559,7 +561,6 @@ void FClothConstraints::CreateGSRules()
 	}
 
 #endif
-
 }
 
 void FClothConstraints::AddRules(
@@ -664,6 +665,7 @@ void FClothConstraints::CreateSelfCollisionConstraints(const Softs::FCollectionP
 			NumParticles,
 			TriangleMesh,
 			ConfigProperties);
+		++NumConstraintInits;
 
 		SelfCollisionConstraints = MakeShared<Softs::FPBDCollisionSpringConstraints>(
 			ParticleOffset,
@@ -672,23 +674,23 @@ void FClothConstraints::CreateSelfCollisionConstraints(const Softs::FCollectionP
 			AnimationPositions,
 			MoveTemp(DisabledCollisionElements),
 			ConfigProperties);
-
-		++NumConstraintInits;
 		++NumPostCollisionConstraintRules;
 
 		SelfIntersectionConstraints = MakeShared<Softs::FPBDTriangleMeshIntersections>(
 			ParticleOffset,
 			NumParticles,
 			TriangleMesh);
+		++NumPostprocessingConstraintRules;
+
 		if (Evolution)
 		{
-			++NumPreSubstepConstraintRules;
+			++NumPreSubstepConstraintRules; // Contour minimization
+			++NumUpdateLinearSystemRules;
 		}
 		else
 		{
-			++NumConstraintInits;
+			++NumConstraintInits; // Contour minimization
 		}
-		++NumPostprocessingConstraintRules;
 	}
 	else if (Softs::FPBDSelfCollisionSphereConstraints::IsEnabled(ConfigProperties))
 	{
@@ -746,6 +748,7 @@ void FClothConstraints::CreateStretchConstraints(
 				WeightMaps,
 				ConfigProperties,
 				/*bTrimKinematicConstraints =*/ true);
+			++NumUpdateLinearSystemRules;
 		}
 		else
 		{
@@ -759,7 +762,7 @@ void FClothConstraints::CreateStretchConstraints(
 				/*bTrimKinematicConstraints =*/ true);
 		}
 
-		++NumConstraintInits;  // Uses init to update the property tables
+		++NumConstraintInits;
 		++NumConstraintRules;
 	}
 	else if (Softs::FPBDEdgeSpringConstraints::IsEnabled(ConfigProperties))
@@ -889,6 +892,7 @@ void FClothConstraints::CreateBendingConstraints(
 				WeightMaps,
 				ConfigProperties,
 				/*bTrimKinematicConstraints =*/ true);
+			++NumUpdateLinearSystemRules;
 		}
 		else
 		{
@@ -901,7 +905,7 @@ void FClothConstraints::CreateBendingConstraints(
 				/*bTrimKinematicConstraints =*/ true);
 		}
 
-		++NumConstraintInits;  // Uses init to update the property tables
+		++NumConstraintInits;
 		++NumConstraintRules;
 	}
 	else if (Softs::FPBDBendingSpringConstraints::IsEnabled(ConfigProperties))
@@ -962,7 +966,7 @@ void FClothConstraints::CreateAreaConstraints(
 				/*bTrimKinematicConstraints =*/ true);
 		}
 
-		++NumConstraintInits;  // Uses init to update the property tables
+		++NumConstraintInits;
 		++NumConstraintRules;
 	}
 	else if (Softs::FPBDAreaSpringConstraints::IsEnabled(ConfigProperties))
@@ -1146,6 +1150,7 @@ void FClothConstraints::CreateExternalForces(
 			);
 
 		++NumExternalForceRules;
+		++NumUpdateLinearSystemRules;
 	}
 }
 
@@ -1167,44 +1172,53 @@ void FClothConstraints::CreateCollisionConstraint(
 			CollisionPhis
 			);
 		++NumCollisionConstraintRules;
+		++NumUpdateLinearSystemCollisionsRules;
 	}
 }
 
 void FClothConstraints::CreateForceBasedRules()
 {
 	Evolution->AllocatePreSubstepParallelInitRange(ParticleRangeId, NumPreSubstepInits);
-	Evolution->AllocateExternalForceRulesRange(ParticleRangeId, NumExternalForceRules);
-	Evolution->AllocateConstraintParallelInitsRange(ParticleRangeId, NumConstraintInits);
+	Evolution->AllocatePBDExternalForceRulesRange(ParticleRangeId, NumExternalForceRules);
+	Evolution->AllocatePostInitialGuessParallelInitRange(ParticleRangeId, NumConstraintInits);
 	Evolution->AllocatePreSubstepConstraintRulesRange(ParticleRangeId, NumPreSubstepConstraintRules);
-	Evolution->AllocatePerIterationConstraintRulesRange(ParticleRangeId, NumConstraintRules);
-	Evolution->AllocatePerIterationCollisionConstraintRulesRange(ParticleRangeId, NumCollisionConstraintRules);
-	Evolution->AllocatePerIterationPostCollisionsConstraintRulesRange(ParticleRangeId, NumPostCollisionConstraintRules);
+	Evolution->AllocatePerIterationPBDConstraintRulesRange(ParticleRangeId, NumConstraintRules);
+	Evolution->AllocatePerIterationCollisionPBDConstraintRulesRange(ParticleRangeId, NumCollisionConstraintRules);
+	Evolution->AllocatePerIterationPostCollisionsPBDConstraintRulesRange(ParticleRangeId, NumPostCollisionConstraintRules);
+	Evolution->AllocateUpdateLinearSystemRulesRange(ParticleRangeId, NumUpdateLinearSystemRules);
+	Evolution->AllocateUpdateLinearSystemCollisionsRulesRange(ParticleRangeId, NumUpdateLinearSystemCollisionsRules);
 	Evolution->AllocatePostSubstepConstraintRulesRange(ParticleRangeId, NumPostprocessingConstraintRules);
 
-	TArrayView<Softs::FEvolution::PreSubstepParallelInitFunc> PreSubstepParallelInits =
+	TArrayView<Softs::FEvolution::ParallelInitFunc> PreSubstepParallelInits =
 		Evolution->GetPreSubstepParallelInitRange(ParticleRangeId);
-	TArrayView<Softs::FEvolution::ExternalForceRuleFunc> ExternalForceRules =
-		Evolution->GetExternalForceRulesRange(ParticleRangeId);
-	TArrayView<Softs::FEvolution::ConstraintParallelInitFunc> ConstraintParallelInits =
-		Evolution->GetConstraintParallelInitsRange(ParticleRangeId);
+	TArrayView<Softs::FEvolution::PBDConstraintRuleFunc> ExternalForceRules =
+		Evolution->GetPBDExternalForceRulesRange(ParticleRangeId);
+	TArrayView<Softs::FEvolution::ParallelInitFunc> PostInitialGuessParallelInits =
+		Evolution->GetPostInitialGuessParallelInitRange(ParticleRangeId);
 	TArrayView<Softs::FEvolution::ConstraintRuleFunc> PreSubstepConstraintRules =
 		Evolution->GetPreSubstepConstraintRulesRange(ParticleRangeId);
-	TArrayView<Softs::FEvolution::ConstraintRuleFunc> PerIterationConstraintRules =
-		Evolution->GetPerIterationConstraintRulesRange(ParticleRangeId);
-	TArrayView<Softs::FEvolution::CollisionConstraintRuleFunc> PerIterationCollisionConstraintRules =
-		Evolution->GetPerIterationCollisionConstraintRulesRange(ParticleRangeId);
-	TArrayView<Softs::FEvolution::ConstraintRuleFunc> PerIterationPostCollisionsConstraintRules =
-		Evolution->GetPerIterationPostCollisionsConstraintRulesRange(ParticleRangeId);
+	TArrayView<Softs::FEvolution::PBDConstraintRuleFunc> PerIterationConstraintRules =
+		Evolution->GetPerIterationPBDConstraintRulesRange(ParticleRangeId);
+	TArrayView<Softs::FEvolution::PBDCollisionConstraintRuleFunc> PerIterationCollisionConstraintRules =
+		Evolution->GetPerIterationCollisionPBDConstraintRulesRange(ParticleRangeId);
+	TArrayView<Softs::FEvolution::PBDConstraintRuleFunc> PerIterationPostCollisionsConstraintRules =
+		Evolution->GetPerIterationPostCollisionsPBDConstraintRulesRange(ParticleRangeId);
+	TArrayView<Softs::FEvolution::UpdateLinearSystemFunc> UpdateLinearSystemRules =
+		Evolution->GetUpdateLinearSystemRulesRange(ParticleRangeId);
+	TArrayView<Softs::FEvolution::UpdateLinearSystemCollisionsFunc> UpdateLinearSystemCollisionsRules =
+		Evolution->GetUpdateLinearSystemCollisionsRulesRange(ParticleRangeId);
 	TArrayView<Softs::FEvolution::ConstraintRuleFunc> PostSubstepConstraintRules =
 		Evolution->GetPostSubstepConstraintRulesRange(ParticleRangeId);
 
 	int32 PreSubstepInitsIndex = 0;
 	int32 ExternalForceRulesIndex = 0;
-	int32 ConstraintInitsIndex = 0;
+	int32 PostInitialGuessInitsIndex = 0;
 	int32 PreSubstepConstraintRulesIndex = 0;
 	int32 ConstraintRuleIndex = 0;
 	int32 CollisionConstraintRulesIndex = 0;
 	int32 PostCollisionConstraintRulesIndex = 0;
+	int32 UpdateLinearSystemRulesIndex = 0;
+	int32 UpdateLinearSystemCollisionsRulesIndex = 0;
 	int32 PostprocessingConstraintRulesIndex = 0;
 
 	if (ExternalForces)
@@ -1214,12 +1228,17 @@ void FClothConstraints::CreateForceBasedRules()
 		{
 			ExternalForces->Apply(Particles, Dt);
 		};
+		UpdateLinearSystemRules[UpdateLinearSystemRulesIndex++] =
+			[this](const Softs::FSolverParticlesRange& Particles, const FSolverReal Dt, Softs::FEvolutionLinearSystem& LinearSystem)
+		{
+			ExternalForces->UpdateLinearSystem(Particles, Dt, LinearSystem);
+		};
 	}
 
 	if (VelocityAndPressureField)
 	{
 		PreSubstepParallelInits[PreSubstepInitsIndex++] =
-			[this](const Softs::FSolverParticlesRange& Particles, const Softs::FSolverReal Dt)
+			[this](const Softs::FSolverParticlesRange& Particles, const Softs::FSolverReal Dt, const Softs::ESolverMode SolverMode)
 		{
 			VelocityAndPressureField->UpdateForces(Particles, Dt);
 		};
@@ -1229,6 +1248,8 @@ void FClothConstraints::CreateForceBasedRules()
 		{
 			VelocityAndPressureField->Apply(Particles, Dt);
 		};
+
+		// TODO Linear System
 	}
 
 	if (PerSolverField)
@@ -1264,176 +1285,112 @@ void FClothConstraints::CreateForceBasedRules()
 				}
 			}
 		};
-	}
-	if (XStretchBiasConstraints)
-	{
-		ConstraintParallelInits[ConstraintInitsIndex++] =
-			[this](const Softs::FSolverParticlesRange& /*Particles*/, const Softs::FSolverReal Dt)
-		{
-			XStretchBiasConstraints->Init();
-			XStretchBiasConstraints->ApplyProperties(Dt, Evolution->GetIterations());
-		};
 
-		PerIterationConstraintRules[ConstraintRuleIndex++] =
-			[this](Softs::FSolverParticlesRange& Particles, const Softs::FSolverReal Dt)
-		{
-			XStretchBiasConstraints->Apply(Particles, Dt);
-		};
+		// TODO: Linear System
 	}
-	if (XEdgeConstraints)
-	{
-		ConstraintParallelInits[ConstraintInitsIndex++] =
-			[this](const Softs::FSolverParticlesRange& /*Particles*/, const Softs::FSolverReal Dt)
-		{
-			XEdgeConstraints->Init();
-			XEdgeConstraints->ApplyProperties(Dt, Evolution->GetIterations());
-		};
 
-		PerIterationConstraintRules[ConstraintRuleIndex++] =
-			[this](Softs::FSolverParticlesRange& Particles, const Softs::FSolverReal Dt)
-		{
-			XEdgeConstraints->Apply(Particles, Dt);
-		};
-	}
-	if (EdgeConstraints)
-	{
-		ConstraintParallelInits[ConstraintInitsIndex++] =
-			[this](const Softs::FSolverParticlesRange& /*Particles*/, const Softs::FSolverReal Dt)
-		{
-			EdgeConstraints->ApplyProperties(Dt, Evolution->GetIterations());
-		};
-		PerIterationConstraintRules[ConstraintRuleIndex++] =
-			[this](Softs::FSolverParticlesRange& Particles, const Softs::FSolverReal Dt)
-		{
-			EdgeConstraints->Apply(Particles, Dt);
-		};
-	}
-	if (XBendingConstraints)
-	{
-		ConstraintParallelInits[ConstraintInitsIndex++] =
-			[this](const Softs::FSolverParticlesRange& /*Particles*/, const Softs::FSolverReal Dt)
-		{
-			XBendingConstraints->Init();
-			XBendingConstraints->ApplyProperties(Dt, Evolution->GetIterations());
-		};
-		PerIterationConstraintRules[ConstraintRuleIndex++] =
-			[this](Softs::FSolverParticlesRange& Particles, const Softs::FSolverReal Dt)
-		{
-			XBendingConstraints->Apply(Particles, Dt);
-		};
-	}
-	if (BendingConstraints)
-	{
-		ConstraintParallelInits[ConstraintInitsIndex++] =
-			[this](const Softs::FSolverParticlesRange& /*Particles*/, const Softs::FSolverReal Dt)
-		{
-			BendingConstraints->ApplyProperties(Dt, Evolution->GetIterations());
-		};
-		PerIterationConstraintRules[ConstraintRuleIndex++] =
-			[this](Softs::FSolverParticlesRange& Particles, const Softs::FSolverReal Dt)
-		{
-			BendingConstraints->Apply(Particles, Dt);
-		};
-	}
-	if (BendingElementConstraints)
-	{
-		ConstraintParallelInits[ConstraintInitsIndex++] =
-			[this](const Softs::FSolverParticlesRange& Particles, const Softs::FSolverReal Dt)
-		{
-			BendingElementConstraints->Init(Particles);
-			BendingElementConstraints->ApplyProperties(Dt, Evolution->GetIterations());
-		};
-		PerIterationConstraintRules[ConstraintRuleIndex++] =
-			[this](Softs::FSolverParticlesRange& Particles, const Softs::FSolverReal Dt)
-		{
-			BendingElementConstraints->Apply(Particles, Dt);
-		};
-	}
-	if (XBendingElementConstraints)
-	{
-		ConstraintParallelInits[ConstraintInitsIndex++] =
-			[this](const Softs::FSolverParticlesRange& Particles, const Softs::FSolverReal Dt)
-		{
-			XBendingElementConstraints->Init(Particles);
-			XBendingElementConstraints->ApplyProperties(Dt, Evolution->GetIterations());
-		};
-		PerIterationConstraintRules[ConstraintRuleIndex++] =
-			[this](Softs::FSolverParticlesRange& Particles, const Softs::FSolverReal Dt)
-		{
-			XBendingElementConstraints->Apply(Particles, Dt);
-		};
-	}
-	if (XAnisoBendingElementConstraints)
-	{
-		ConstraintParallelInits[ConstraintInitsIndex++] =
-			[this](const Softs::FSolverParticlesRange& Particles, const Softs::FSolverReal Dt)
-		{
-			XAnisoBendingElementConstraints->Init(Particles);
-			XAnisoBendingElementConstraints->ApplyProperties(Dt, Evolution->GetIterations());
-		};
-		PerIterationConstraintRules[ConstraintRuleIndex++] =
-			[this](Softs::FSolverParticlesRange& Particles, const Softs::FSolverReal Dt)
-		{
-			XAnisoBendingElementConstraints->Apply(Particles, Dt);
-		};
-	}
-	if (XAreaConstraints)
-	{
-		ConstraintParallelInits[ConstraintInitsIndex++] =
-			[this](const Softs::FSolverParticlesRange& /*Particles*/, const Softs::FSolverReal Dt)
-		{
-			XAreaConstraints->Init();
-			XAreaConstraints->ApplyProperties(Dt, Evolution->GetIterations());
-		};
-		PerIterationConstraintRules[ConstraintRuleIndex++] =
-			[this](Softs::FSolverParticlesRange& Particles, const Softs::FSolverReal Dt)
-		{
-			XAreaConstraints->Apply(Particles, Dt);
-		};
-	}
-	if (AreaConstraints)
-	{
-		ConstraintParallelInits[ConstraintInitsIndex++] =
-			[this](const Softs::FSolverParticlesRange& /*Particles*/, const Softs::FSolverReal Dt)
-		{
-			AreaConstraints->ApplyProperties(Dt, Evolution->GetIterations());
-		};
-		PerIterationConstraintRules[ConstraintRuleIndex++] =
-			[this](Softs::FSolverParticlesRange& Particles, const Softs::FSolverReal Dt)
-		{
-			AreaConstraints->Apply(Particles, Dt);
-		};
-	}
-	if (MaximumDistanceConstraints)
-	{
-		PerIterationConstraintRules[ConstraintRuleIndex++] =
-			[this](Softs::FSolverParticlesRange& Particles, const Softs::FSolverReal Dt)
-		{
-			MaximumDistanceConstraints->Apply(Particles, Dt);
-		};
-	}
-	if (BackstopConstraints)
-	{
-		PerIterationConstraintRules[ConstraintRuleIndex++] =
-			[this](Softs::FSolverParticlesRange& Particles, const Softs::FSolverReal Dt)
-		{
-			BackstopConstraints->Apply(Particles, Dt);
-		};
-	}
-	if (AnimDriveConstraints)
-	{
-		ConstraintParallelInits[ConstraintInitsIndex++] =
-			[this](const Softs::FSolverParticlesRange& /*Particles*/, const Softs::FSolverReal Dt)
-		{
-			AnimDriveConstraints->ApplyProperties(Dt, Evolution->GetIterations());
-		};
+#define UE_CHAOS_XPBD_INIT_VOID(Constraint)\
+	PostInitialGuessParallelInits[PostInitialGuessInitsIndex++] = \
+		[this](const Softs::FSolverParticlesRange& /*Particles*/, const Softs::FSolverReal Dt, const Softs::ESolverMode SolverMode)\
+	{\
+		Constraint->ApplyProperties(Dt, Evolution->GetIterations()); \
+		if ((SolverMode & Softs::ESolverMode::PBD) != Softs::ESolverMode::None)\
+		{\
+			Constraint->Init(); \
+		}\
+	};
+#define UE_CHAOS_XPBD_INIT_PARTICLES(Constraint)\
+	PostInitialGuessParallelInits[PostInitialGuessInitsIndex++] = \
+		[this](const Softs::FSolverParticlesRange& Particles, const Softs::FSolverReal Dt, const Softs::ESolverMode SolverMode)\
+	{\
+		Constraint->ApplyProperties(Dt, Evolution->GetIterations()); \
+		if ((SolverMode & Softs::ESolverMode::PBD) != Softs::ESolverMode::None)\
+		{\
+			Constraint->Init(Particles); \
+		}\
+	};
+#define UE_CHAOS_PBD_INIT(Constraint)\
+	PostInitialGuessParallelInits[PostInitialGuessInitsIndex++] = \
+		[this](const Softs::FSolverParticlesRange& /*Particles*/, const Softs::FSolverReal Dt, const Softs::ESolverMode SolverMode)\
+	{\
+		if ((SolverMode & Softs::ESolverMode::PBD) != Softs::ESolverMode::None)\
+		{\
+			Constraint->ApplyProperties(Dt, Evolution->GetIterations()); \
+		}\
+	};
+#define UE_CHAOS_APPLY_CONSTRAINT(Constraint)\
+	PerIterationConstraintRules[ConstraintRuleIndex++] = \
+		[this](Softs::FSolverParticlesRange& Particles, const Softs::FSolverReal Dt)\
+	{\
+		Constraint->Apply(Particles, Dt); \
+	}; 
+#define UE_CHAOS_UPDATE_LINEAR_SYSTEM(Constraint)\
+	UpdateLinearSystemRules[UpdateLinearSystemRulesIndex++] =\
+		[this](const Softs::FSolverParticlesRange& Particles, const FSolverReal Dt, Softs::FEvolutionLinearSystem& LinearSystem)\
+	{\
+		Constraint->UpdateLinearSystem(Particles, Dt, LinearSystem);\
+	};
+
+#define UE_CHAOS_ADD_XPBD_CONSTRAINT_NO_LINEAR_SYSTEM(Constraint) \
+	if (Constraint)\
+	{\
+		UE_CHAOS_XPBD_INIT_VOID(Constraint)\
+		UE_CHAOS_APPLY_CONSTRAINT(Constraint)\
+	}
+
+#define UE_CHAOS_ADD_XPBD_CONSTRAINT_INIT_PARTICLES_NO_LINEAR_SYSTEM(Constraint) \
+	if (Constraint)\
+	{\
+		UE_CHAOS_XPBD_INIT_PARTICLES(Constraint)\
+		UE_CHAOS_APPLY_CONSTRAINT(Constraint)\
+	}
+
+#define UE_CHAOS_ADD_XPBD_CONSTRAINT(Constraint) \
+	if (Constraint)\
+	{\
+		UE_CHAOS_XPBD_INIT_VOID(Constraint)\
+		UE_CHAOS_APPLY_CONSTRAINT(Constraint)\
+		UE_CHAOS_UPDATE_LINEAR_SYSTEM(Constraint)\
+	}
+
+#define UE_CHAOS_ADD_PBD_CONSTRAINT(Constraint) \
+	if (Constraint)\
+	{\
+		UE_CHAOS_PBD_INIT(Constraint)\
+		UE_CHAOS_APPLY_CONSTRAINT(Constraint)\
+	}
+
+#define UE_CHAOS_ADD_PBD_CONSTRAINT_NO_PROPERTIES(Constraint) \
+	if (Constraint)\
+	{\
+		UE_CHAOS_APPLY_CONSTRAINT(Constraint)\
+	}
+	
+	UE_CHAOS_ADD_XPBD_CONSTRAINT_NO_LINEAR_SYSTEM(XStretchBiasConstraints);
+	UE_CHAOS_ADD_XPBD_CONSTRAINT(XEdgeConstraints);
+	UE_CHAOS_ADD_PBD_CONSTRAINT(EdgeConstraints);
+	UE_CHAOS_ADD_XPBD_CONSTRAINT(XBendingConstraints);
+	UE_CHAOS_ADD_PBD_CONSTRAINT(BendingConstraints);
+	UE_CHAOS_ADD_PBD_CONSTRAINT(BendingElementConstraints);
+	UE_CHAOS_ADD_XPBD_CONSTRAINT_INIT_PARTICLES_NO_LINEAR_SYSTEM(XBendingElementConstraints);
+	UE_CHAOS_ADD_XPBD_CONSTRAINT_INIT_PARTICLES_NO_LINEAR_SYSTEM(XAnisoBendingElementConstraints);
+	UE_CHAOS_ADD_XPBD_CONSTRAINT_NO_LINEAR_SYSTEM(XAreaConstraints);
+	UE_CHAOS_ADD_PBD_CONSTRAINT(AreaConstraints);
+	UE_CHAOS_ADD_PBD_CONSTRAINT_NO_PROPERTIES(MaximumDistanceConstraints);
+	UE_CHAOS_ADD_PBD_CONSTRAINT_NO_PROPERTIES(BackstopConstraints);
+	UE_CHAOS_ADD_PBD_CONSTRAINT(AnimDriveConstraints);
+
+#undef UE_CHAOS_ADD_PBD_CONSTRAINT_NO_PROPERTIES
+#undef UE_CHAOS_ADD_PBD_CONSTRAINT
+#undef UE_CHAOS_ADD_XPBD_CONSTRAINT
+#undef UE_CHAOS_ADD_XPBD_CONSTRAINT_INIT_PARTICLES_NO_LINEAR_SYSTEM
+#undef UE_CHAOS_ADD_XPBD_CONSTRAINT_NO_LINEAR_SYSTEM
+#undef UE_CHAOS_UPDATE_LINEAR_SYSTEM
+#undef UE_CHAOS_APPLY_CONSTRAINT
+#undef UE_CHAOS_PBD_INIT
+#undef UE_CHAOS_XPBD_INIT_PARTICLES
+#undef UE_CHAOS_XPBD_INIT_VOID
 
-		PerIterationConstraintRules[ConstraintRuleIndex++] =
-			[this](Softs::FSolverParticlesRange& Particles, const Softs::FSolverReal Dt)
-		{
-			AnimDriveConstraints->Apply(Particles, Dt);
-		};
-	}
 	if (CollisionConstraint)
 	{
 		PerIterationCollisionConstraintRules[CollisionConstraintRulesIndex++] =
@@ -1441,11 +1398,17 @@ void FClothConstraints::CreateForceBasedRules()
 		{
 			CollisionConstraint->Apply(Particles, Dt, CollisionParticles);
 		};
+
+		UpdateLinearSystemCollisionsRules[UpdateLinearSystemCollisionsRulesIndex++] =
+			[this](const Softs::FSolverParticlesRange& Particles, const Softs::FSolverReal Dt, const TArray<Softs::FSolverCollisionParticlesRange>& CollisionParticles, Softs::FEvolutionLinearSystem& LinearSystem)
+		{
+			CollisionConstraint->UpdateLinearSystem(Particles, Dt, CollisionParticles, LinearSystem);
+		};
 	}
 	if (SelfCollisionInit && SelfCollisionConstraints)
 	{
-		ConstraintParallelInits[ConstraintInitsIndex++] =
-			[this](const Softs::FSolverParticlesRange& Particles, const Softs::FSolverReal Dt)
+		PostInitialGuessParallelInits[PostInitialGuessInitsIndex++] =
+			[this](const Softs::FSolverParticlesRange& Particles, const Softs::FSolverReal Dt, const Softs::ESolverMode SolverMode)
 		{
 			// Thickness * 2 to account for collision radius for both particles
 			SelfCollisionInit->Init(Particles, SelfCollisionConstraints->GetThickness() * (Softs::FSolverReal)2.f);
@@ -1457,11 +1420,17 @@ void FClothConstraints::CreateForceBasedRules()
 		{
 			SelfCollisionConstraints->Apply(Particles, Dt);
 		};
+
+		UpdateLinearSystemRules[UpdateLinearSystemRulesIndex++] =
+			[this](const Softs::FSolverParticlesRange& Particles, const FSolverReal Dt, Softs::FEvolutionLinearSystem& LinearSystem)
+		{
+			SelfCollisionConstraints->UpdateLinearSystem(Particles, Dt, LinearSystem);
+		};
 	}
 	if (SelfCollisionSphereConstraints)
 	{
-		ConstraintParallelInits[ConstraintInitsIndex++] =
-			[this](const Softs::FSolverParticlesRange& Particles, const Softs::FSolverReal /*Dt*/)
+		PostInitialGuessParallelInits[PostInitialGuessInitsIndex++] =
+			[this](const Softs::FSolverParticlesRange& Particles, const Softs::FSolverReal Dt, const Softs::ESolverMode SolverMode)
 		{
 			SelfCollisionSphereConstraints->Init(Particles);
 		};
@@ -1476,13 +1445,13 @@ void FClothConstraints::CreateForceBasedRules()
 	if (SelfCollisionInit && SelfIntersectionConstraints)
 	{
 		PreSubstepConstraintRules[PreSubstepConstraintRulesIndex++] =
-			[this](Softs::FSolverParticlesRange& Particles, const Softs::FSolverReal Dt)
+			[this](Softs::FSolverParticlesRange& Particles, const Softs::FSolverReal Dt, const Softs::ESolverMode SolverMode)
 		{
 			SelfIntersectionConstraints->Apply(Particles, SelfCollisionInit->GetContourMinimizationIntersections(), Dt);
 		};
 
 		PostSubstepConstraintRules[PostprocessingConstraintRulesIndex++] =
-			[this](Softs::FSolverParticlesRange& Particles, const Softs::FSolverReal Dt)
+			[this](Softs::FSolverParticlesRange& Particles, const Softs::FSolverReal Dt, const Softs::ESolverMode SolverMode)
 		{
 			const int32 NumContourIterations = SelfCollisionInit->GetNumContourMinimizationPostSteps();
 			for (int32 Iter = 0; Iter < NumContourIterations; ++Iter)
@@ -1495,22 +1464,27 @@ void FClothConstraints::CreateForceBasedRules()
 	if (LongRangeConstraints)
 	{
 		PreSubstepConstraintRules[PreSubstepConstraintRulesIndex++] =
-			[this](Softs::FSolverParticlesRange& Particles, const Softs::FSolverReal Dt)
+			[this](Softs::FSolverParticlesRange& Particles, const Softs::FSolverReal Dt, const Softs::ESolverMode SolverMode)
 		{
-			// Only doing one iteration.
-			constexpr int32 NumLRAIterations = 1;
-			LongRangeConstraints->ApplyProperties(Dt, NumLRAIterations);
-			LongRangeConstraints->Apply(Particles, Dt);  // Run the LRA constraint only once per timestep
+			if ((SolverMode & Softs::ESolverMode::PBD) != Softs::ESolverMode::None)
+			{
+				// Only doing one iteration.
+				constexpr int32 NumLRAIterations = 1;
+				LongRangeConstraints->ApplyProperties(Dt, NumLRAIterations);
+				LongRangeConstraints->Apply(Particles, Dt);  // Run the LRA constraint only once per timestep
+			}
 		};
 	}
 
 	check(PreSubstepInitsIndex == NumPreSubstepInits);
 	check(ExternalForceRulesIndex == NumExternalForceRules);
-	check(ConstraintInitsIndex == NumConstraintInits);
+	check(PostInitialGuessInitsIndex == NumConstraintInits)
 	check(PreSubstepConstraintRulesIndex == NumPreSubstepConstraintRules);
 	check(ConstraintRuleIndex == NumConstraintRules);
 	check(CollisionConstraintRulesIndex == NumCollisionConstraintRules);
 	check(PostCollisionConstraintRulesIndex == NumPostCollisionConstraintRules);
+	check(UpdateLinearSystemRulesIndex == NumUpdateLinearSystemRules);
+	check(UpdateLinearSystemCollisionsRulesIndex == NumUpdateLinearSystemCollisionsRules);
 	check(PostprocessingConstraintRulesIndex == NumPostprocessingConstraintRules);
 }
 
diff --git a/Engine/Plugins/Experimental/ChaosCloth/Source/Chaos/Private/ChaosCloth/ChaosClothingSimulationSolver.cpp b/Engine/Plugins/Experimental/ChaosCloth/Source/Chaos/Private/ChaosCloth/ChaosClothingSimulationSolver.cpp
--- a/Engine/Plugins/Experimental/ChaosCloth/Source/Chaos/Private/ChaosCloth/ChaosClothingSimulationSolver.cpp
+++ b/Engine/Plugins/Experimental/ChaosCloth/Source/Chaos/Private/ChaosCloth/ChaosClothingSimulationSolver.cpp
@@ -2051,6 +2051,16 @@ int32 FClothingSimulationSolver::GetNumUsedIterations() const
 	return Evolution ? Evolution->GetIterations() : PBDEvolution->GetIterations();
 }
 
+int32 FClothingSimulationSolver::GetNumLinearSolverIterations(int32 ParticleRangeId) const
+{
+	return Evolution ? Evolution->GetLastLinearSolveIterations(ParticleRangeId) : 0;
+}
+
+FRealSingle FClothingSimulationSolver::GetLinearSolverError(int32 ParticleRangeId) const
+{
+	return Evolution ? Evolution->GetLastLinearSolveError(ParticleRangeId) : 0.f;
+}
+
 FBoxSphereBounds FClothingSimulationSolver::CalculateBounds() const
 {
 	TRACE_CPUPROFILER_EVENT_SCOPE(FClothingSimulationSolver_CalculateBounds);
diff --git a/Engine/Plugins/Experimental/ChaosCloth/Source/Chaos/Public/ChaosCloth/ChaosClothConstraints.h b/Engine/Plugins/Experimental/ChaosCloth/Source/Chaos/Public/ChaosCloth/ChaosClothConstraints.h
--- a/Engine/Plugins/Experimental/ChaosCloth/Source/Chaos/Public/ChaosCloth/ChaosClothConstraints.h
+++ b/Engine/Plugins/Experimental/ChaosCloth/Source/Chaos/Public/ChaosCloth/ChaosClothConstraints.h
@@ -256,6 +256,8 @@ PRAGMA_ENABLE_DEPRECATION_WARNINGS
 		int32 NumExternalForceRules;
 		int32 NumPreSubstepConstraintRules;
 		int32 NumCollisionConstraintRules;
+		int32 NumUpdateLinearSystemRules;
+		int32 NumUpdateLinearSystemCollisionsRules;
 		//~ End Force-based solver only fields
 
 		//~ Begin PBD solver only fields
diff --git a/Engine/Plugins/Experimental/ChaosCloth/Source/Chaos/Public/ChaosCloth/ChaosClothingSimulationSolver.h b/Engine/Plugins/Experimental/ChaosCloth/Source/Chaos/Public/ChaosCloth/ChaosClothingSimulationSolver.h
--- a/Engine/Plugins/Experimental/ChaosCloth/Source/Chaos/Public/ChaosCloth/ChaosClothingSimulationSolver.h
+++ b/Engine/Plugins/Experimental/ChaosCloth/Source/Chaos/Public/ChaosCloth/ChaosClothingSimulationSolver.h
@@ -108,6 +108,10 @@ namespace Chaos
 
 		// Return the actual of number of iterations used by the Evolution solver after the update (different from the number of iterations, depends on frame rate)
 		CHAOSCLOTH_API int32 GetNumUsedIterations() const;
+		// Return the actual number of linear solver iterations used by the Evolution solver after the update (force-based solver only)
+		CHAOSCLOTH_API int32 GetNumLinearSolverIterations(int32 ParticleRangeId) const;
+		// Return the final error of the linear solver after the update (force-based solver only)
+		CHAOSCLOTH_API FRealSingle GetLinearSolverError(int32 ParticleRangeId) const;
 
 		CHAOSCLOTH_API FBoxSphereBounds CalculateBounds() const;
 		// ---- End of the object management functions ----
diff --git a/Engine/Plugins/Experimental/ChaosClothAsset/Source/ChaosClothAssetEngine/Private/ChaosClothAsset/ClothSimulationProxy.cpp b/Engine/Plugins/Experimental/ChaosClothAsset/Source/ChaosClothAssetEngine/Private/ChaosClothAsset/ClothSimulationProxy.cpp
--- a/Engine/Plugins/Experimental/ChaosClothAsset/Source/ChaosClothAssetEngine/Private/ChaosClothAsset/ClothSimulationProxy.cpp
+++ b/Engine/Plugins/Experimental/ChaosClothAsset/Source/ChaosClothAssetEngine/Private/ChaosClothAsset/ClothSimulationProxy.cpp
@@ -265,10 +265,25 @@ PRAGMA_ENABLE_DEPRECATION_WARNINGS
 		// Update particle counts (could have changed if lod changed)
 		NumKinematicParticles = 0;
 		NumDynamicParticles = 0;
+		int32 FirstActiveClothParticleRangeId = INDEX_NONE;
 		for (const TUniquePtr<FClothingSimulationCloth>& Cloth : Cloths)
 		{
 			NumKinematicParticles += Cloth->GetNumActiveKinematicParticles();
 			NumDynamicParticles += Cloth->GetNumActiveDynamicParticles();
+			if (FirstActiveClothParticleRangeId == INDEX_NONE && Cloth->GetNumActiveDynamicParticles() > 0)
+			{
+				FirstActiveClothParticleRangeId = Cloth->GetParticleRangeId(Solver.Get ());
+			}
+		}
+		if (FirstActiveClothParticleRangeId != INDEX_NONE)
+		{
+			LastLinearSolveError = Solver->GetLinearSolverError(FirstActiveClothParticleRangeId);
+			LastLinearSolveIterations = Solver->GetNumLinearSolverIterations(FirstActiveClothParticleRangeId);
+		}
+		else
+		{
+			LastLinearSolveError = 0.f;
+			LastLinearSolveIterations = 0;
 		}
 
 		// Visualization
diff --git a/Engine/Plugins/Experimental/ChaosClothAsset/Source/ChaosClothAssetEngine/Public/ChaosClothAsset/ClothSimulationProxy.h b/Engine/Plugins/Experimental/ChaosClothAsset/Source/ChaosClothAssetEngine/Public/ChaosClothAsset/ClothSimulationProxy.h
--- a/Engine/Plugins/Experimental/ChaosClothAsset/Source/ChaosClothAssetEngine/Public/ChaosClothAsset/ClothSimulationProxy.h
+++ b/Engine/Plugins/Experimental/ChaosClothAsset/Source/ChaosClothAssetEngine/Public/ChaosClothAsset/ClothSimulationProxy.h
@@ -66,6 +66,8 @@ namespace UE::Chaos::ClothAsset
 		int32 GetNumDynamicParticles() const { return NumDynamicParticles; }
 		int32 GetNumIterations() const { return NumIterations; }
 		int32 GetNumSubsteps() const { return NumSubsteps; }
+		int32 GetNumLinearSolveIterations() const { return LastLinearSolveIterations; }
+		float GetLinearSolveError() const { return LastLinearSolveError; }
 		float GetSimulationTime() const { return SimulationTime; }
 		bool IsTeleported() const { return bIsTeleported; }
 
@@ -127,6 +129,8 @@ PRAGMA_ENABLE_DEPRECATION_WARNINGS
 		std::atomic<int32> NumSubsteps;
 		std::atomic<float> SimulationTime;
 		std::atomic<bool> bIsTeleported;
+		std::atomic<int32> LastLinearSolveIterations; // For single cloth only.
+		std::atomic<float> LastLinearSolveError; // For single cloth only.
 
 		mutable bool bHasInvalidReferenceBoneTransforms = false;
 
diff --git a/Engine/Plugins/Experimental/ChaosClothAssetEditor/Source/ChaosClothAssetDataflowNodes/Private/ChaosClothAsset/SimulationCollisionConfigNode.cpp b/Engine/Plugins/Experimental/ChaosClothAssetEditor/Source/ChaosClothAssetDataflowNodes/Private/ChaosClothAsset/SimulationCollisionConfigNode.cpp
--- a/Engine/Plugins/Experimental/ChaosClothAssetEditor/Source/ChaosClothAssetDataflowNodes/Private/ChaosClothAsset/SimulationCollisionConfigNode.cpp
+++ b/Engine/Plugins/Experimental/ChaosClothAssetEditor/Source/ChaosClothAssetDataflowNodes/Private/ChaosClothAsset/SimulationCollisionConfigNode.cpp
@@ -17,4 +17,5 @@ void FChaosClothAssetSimulationCollisionConfigNode::AddProperties(Dataflow::FCon
 	UE_CHAOS_CLOTHASSET_SIMULATIONCONFIG_SETPROPERTY(CollisionThickness);
 	UE_CHAOS_CLOTHASSET_SIMULATIONCONFIG_SETPROPERTY(FrictionCoefficient);
 	UE_CHAOS_CLOTHASSET_SIMULATIONCONFIG_SETPROPERTYBOOL(UseCCD);
+	UE_CHAOS_CLOTHASSET_SIMULATIONCONFIG_SETPROPERTY(ProximityStiffness);
 }
diff --git a/Engine/Plugins/Experimental/ChaosClothAssetEditor/Source/ChaosClothAssetDataflowNodes/Private/ChaosClothAsset/SimulationSelfCollisionConfigNode.cpp b/Engine/Plugins/Experimental/ChaosClothAssetEditor/Source/ChaosClothAssetDataflowNodes/Private/ChaosClothAsset/SimulationSelfCollisionConfigNode.cpp
--- a/Engine/Plugins/Experimental/ChaosClothAssetEditor/Source/ChaosClothAssetDataflowNodes/Private/ChaosClothAsset/SimulationSelfCollisionConfigNode.cpp
+++ b/Engine/Plugins/Experimental/ChaosClothAssetEditor/Source/ChaosClothAssetDataflowNodes/Private/ChaosClothAsset/SimulationSelfCollisionConfigNode.cpp
@@ -25,4 +25,5 @@ void FChaosClothAssetSimulationSelfCollisionConfigNode::AddProperties(Dataflow::
 	UE_CHAOS_CLOTHASSET_SIMULATIONCONFIG_SETPROPERTYBOOL(UseContourMinimization);
 	UE_CHAOS_CLOTHASSET_SIMULATIONCONFIG_SETPROPERTY(NumContourMinimizationPostSteps);
 	UE_CHAOS_CLOTHASSET_SIMULATIONCONFIG_SETPROPERTYBOOL(UseGlobalPostStepContours);
+	UE_CHAOS_CLOTHASSET_SIMULATIONCONFIG_SETPROPERTY(SelfCollisionProximityStiffness);
 }
diff --git a/Engine/Plugins/Experimental/ChaosClothAssetEditor/Source/ChaosClothAssetDataflowNodes/Public/ChaosClothAsset/SimulationCollisionConfigNode.h b/Engine/Plugins/Experimental/ChaosClothAssetEditor/Source/ChaosClothAssetDataflowNodes/Public/ChaosClothAsset/SimulationCollisionConfigNode.h
--- a/Engine/Plugins/Experimental/ChaosClothAssetEditor/Source/ChaosClothAssetDataflowNodes/Public/ChaosClothAsset/SimulationCollisionConfigNode.h
+++ b/Engine/Plugins/Experimental/ChaosClothAssetEditor/Source/ChaosClothAssetDataflowNodes/Public/ChaosClothAsset/SimulationCollisionConfigNode.h
@@ -21,6 +21,10 @@ public:
 	UPROPERTY(EditAnywhere, Category = "Collision Properties", Meta = (UIMin = "0", UIMax = "1", ClampMin = "0", ClampMax = "10"))
 	float FrictionCoefficient = 0.8f;
 
+	/** Stiffness for proximity repulsion forces (Force-based solver only). Units = kg cm/ s^2 (same as XPBD springs)*/
+	UPROPERTY(EditAnywhere, Category = "Proximity Force Properties", Meta = (UIMin = "0", UIMax = "10000", ClampMin = "0", ClampMax = "10000000"))
+	float ProximityStiffness = 100.f;
+
 	/**
 	 * Use continuous collision detection (CCD) to prevent any missed collisions between fast moving particles and colliders.
 	 * This has a negative effect on performance compared to when resolving collision without using CCD.
diff --git a/Engine/Plugins/Experimental/ChaosClothAssetEditor/Source/ChaosClothAssetDataflowNodes/Public/ChaosClothAsset/SimulationSelfCollisionConfigNode.h b/Engine/Plugins/Experimental/ChaosClothAssetEditor/Source/ChaosClothAssetDataflowNodes/Public/ChaosClothAsset/SimulationSelfCollisionConfigNode.h
--- a/Engine/Plugins/Experimental/ChaosClothAssetEditor/Source/ChaosClothAssetDataflowNodes/Public/ChaosClothAsset/SimulationSelfCollisionConfigNode.h
+++ b/Engine/Plugins/Experimental/ChaosClothAssetEditor/Source/ChaosClothAssetDataflowNodes/Public/ChaosClothAsset/SimulationSelfCollisionConfigNode.h
@@ -17,12 +17,12 @@ public:
 	UPROPERTY(EditAnywhere, Category = "Self-Collision Properties", meta = (UIMin = "0", UIMax = "100", ClampMin = "0", ClampMax = "1000", EditCondition = "bUseSelfCollisions"))
 	float SelfCollisionThickness = 0.5f;
 
-	/** The stiffness of the springs used to control self collision. */
-	UPROPERTY(EditAnywhere, Category = "Self-Collision Properties", meta = (UIMin = "0", UIMax = "100", ClampMin = "0", ClampMax = "1000"))
+	/** The stiffness of the springs used to control self collision (PBD Solver). */
+	UPROPERTY(EditAnywhere, Category = "Self-Collision Properties", meta = (UIMin = "0", UIMax = "1", ClampMin = "0", ClampMax = "1"))
 	float SelfCollisionStiffness = 0.5f;
 
 	/** Friction coefficient for cloth - cloth interaction. */
-	UPROPERTY(EditAnywhere, Category = "Self-Collision Properties", meta = (UIMin = "0", UIMax = "1", ClampMin = "0", ClampMax = "10"))
+	UPROPERTY(EditAnywhere, Category = "Self-Collision Properties", meta = (UIMin = "0", UIMax = "1", ClampMin = "0", ClampMax = "1"))
 	float SelfCollisionFriction = 0.0f;
 
 	/** Enable self intersection resolution. This will try to fix any cloth intersections that are not handled by collision repulsions. */
@@ -45,6 +45,10 @@ public:
 	UPROPERTY(EditAnywhere, Category = "Experimental", meta = (EditCondition = "bUseSelfIntersections && NumContourMinimizationPostSteps > 0"))
 	bool bUseGlobalPostStepContours = true;
 
+	/** The stiffness of the proximity repulsions used to control self collision (Force-based Solver). Units = kg cm/ s^2 (same as XPBD springs) */
+	UPROPERTY(EditAnywhere, Category = "Experimental", meta = (UIMin = "0", UIMax = "10000", ClampMin = "0", ClampMax = "10000000"))
+	float SelfCollisionProximityStiffness = 1.f;
+
 	FChaosClothAssetSimulationSelfCollisionConfigNode(const Dataflow::FNodeParameters& InParam, FGuid InGuid = FGuid::NewGuid());
 
 private:
diff --git a/Engine/Plugins/Experimental/ChaosClothAssetEditor/Source/ChaosClothAssetDataflowNodes/Public/ChaosClothAsset/SimulationSolverConfigNode.h b/Engine/Plugins/Experimental/ChaosClothAssetEditor/Source/ChaosClothAssetDataflowNodes/Public/ChaosClothAsset/SimulationSolverConfigNode.h
--- a/Engine/Plugins/Experimental/ChaosClothAssetEditor/Source/ChaosClothAssetDataflowNodes/Public/ChaosClothAsset/SimulationSolverConfigNode.h
+++ b/Engine/Plugins/Experimental/ChaosClothAssetEditor/Source/ChaosClothAssetDataflowNodes/Public/ChaosClothAsset/SimulationSolverConfigNode.h
@@ -43,7 +43,7 @@ public:
 	bool bEnableForceBasedSolver = false;
 	
 	/**
-	 * Number of Newton iterations for force-based solver
+	 * Number of Newton iterations for force-based solver. Prototype only--very few constraints support this.
 	 */
 	UPROPERTY(EditAnywhere, Category = Experimental, meta = (UIMin = "1", UIMax = "10", ClampMin = "0", ClampMax = "100", EditCondition = "bEnableForceBasedSolver"))
 	int32 NumNewtonIterations = 1;
@@ -51,7 +51,7 @@ public:
 	/**
 	 * Max number of CG Iterations per linear solve
 	 */
-	UPROPERTY(EditAnywhere, Category = Experimental, meta = (UIMin = "1", UIMax = "100", ClampMin = "1", ClampMax = "1000", EditCondition = "bEnableForceBasedSolver"))
+	UPROPERTY(EditAnywhere, Category = Experimental, meta = (UIMin = "1", UIMax = "100", ClampMin = "1", EditCondition = "bEnableForceBasedSolver"))
 	int32 MaxNumCGIterations = 50;
 
 	/**
diff --git a/Engine/Plugins/Experimental/ChaosClothAssetEditor/Source/ChaosClothAssetDataflowNodes/Public/ChaosClothAsset/SimulationXPBDBendingSpringConfigNode.h b/Engine/Plugins/Experimental/ChaosClothAssetEditor/Source/ChaosClothAssetDataflowNodes/Public/ChaosClothAsset/SimulationXPBDBendingSpringConfigNode.h
--- a/Engine/Plugins/Experimental/ChaosClothAssetEditor/Source/ChaosClothAssetDataflowNodes/Public/ChaosClothAsset/SimulationXPBDBendingSpringConfigNode.h
+++ b/Engine/Plugins/Experimental/ChaosClothAssetEditor/Source/ChaosClothAssetDataflowNodes/Public/ChaosClothAsset/SimulationXPBDBendingSpringConfigNode.h
@@ -20,7 +20,7 @@ public:
 	 * are interpolated with the per particle weight to make the final value used for the simulation.
 	 * Otherwise all particles are considered to have a zero weight, and only the Low value is meaningful.
 	 */
-	UPROPERTY(EditAnywhere, Category = "XPBDBendingSpring Properties", Meta = (UIMin = "0", UIMax = "10000", ClampMin = "0", ClampMax = "10000000"))
+	UPROPERTY(EditAnywhere, Category = "XPBDBendingSpring Properties", Meta = (UIMin = "0", UIMax = "10000", ClampMin = "0", ClampMax = "1000000000"))
 	FChaosClothAssetWeightedValue XPBDBendingSpringStiffness = { true, 100.f, 100.f, TEXT("XPBDBendingSpringStiffness") };
 
 	/**
diff --git a/Engine/Plugins/Experimental/ChaosClothAssetEditor/Source/ChaosClothAssetDataflowNodes/Public/ChaosClothAsset/SimulationXPBDEdgeSpringConfigNode.h b/Engine/Plugins/Experimental/ChaosClothAssetEditor/Source/ChaosClothAssetDataflowNodes/Public/ChaosClothAsset/SimulationXPBDEdgeSpringConfigNode.h
--- a/Engine/Plugins/Experimental/ChaosClothAssetEditor/Source/ChaosClothAssetDataflowNodes/Public/ChaosClothAsset/SimulationXPBDEdgeSpringConfigNode.h
+++ b/Engine/Plugins/Experimental/ChaosClothAssetEditor/Source/ChaosClothAssetDataflowNodes/Public/ChaosClothAsset/SimulationXPBDEdgeSpringConfigNode.h
@@ -20,7 +20,7 @@ public:
 	 * are interpolated with the per particle weight to make the final value used for the simulation.
 	 * Otherwise all particles are considered to have a zero weight, and only the Low value is meaningful.
 	 */
-	UPROPERTY(EditAnywhere, Category = "XPBDEdgeSpring Properties", Meta = (UIMin = "0", UIMax = "10000", ClampMin = "0", ClampMax = "10000000"))
+	UPROPERTY(EditAnywhere, Category = "XPBDEdgeSpring Properties", Meta = (UIMin = "0", UIMax = "10000", ClampMin = "0", ClampMax = "1000000000"))
 	FChaosClothAssetWeightedValue XPBDEdgeSpringStiffness = { true, 100.f, 100.f, TEXT("XPBDEdgeSpringStiffness") };
 
 	/**
diff --git a/Engine/Plugins/Experimental/ChaosClothAssetEditor/Source/ChaosClothAssetEditor/Private/ChaosClothAsset/ClothEditorSimulationVisualization.cpp b/Engine/Plugins/Experimental/ChaosClothAssetEditor/Source/ChaosClothAssetEditor/Private/ChaosClothAsset/ClothEditorSimulationVisualization.cpp
--- a/Engine/Plugins/Experimental/ChaosClothAssetEditor/Source/ChaosClothAssetEditor/Private/ChaosClothAsset/ClothEditorSimulationVisualization.cpp
+++ b/Engine/Plugins/Experimental/ChaosClothAssetEditor/Source/ChaosClothAssetEditor/Private/ChaosClothAsset/ClothEditorSimulationVisualization.cpp
@@ -46,6 +46,22 @@ static FText GetSimulationStatisticsString(const FClothSimulationProxy* SimProxy
 		{
 			TextValue = ConcatenateLine(TextValue, FText::Format(LOCTEXT("NumIterations", "Iterations: {0}"), NumIterations));
 		}
+		if (const int32 NumLinearSolveIterations = SimProxy->GetNumLinearSolveIterations())
+		{
+			TextValue = ConcatenateLine(TextValue, FText::Format(LOCTEXT("NumCGIterations", "CGIterations: {0}"), NumLinearSolveIterations));
+		}
+		if (const float LinearSolveError = SimProxy->GetLinearSolveError())
+		{
+			FNumberFormattingOptions NumberFormatOptions;
+			NumberFormatOptions.AlwaysSign = false;
+			NumberFormatOptions.UseGrouping = false;
+			NumberFormatOptions.RoundingMode = ERoundingMode::HalfFromZero;
+			NumberFormatOptions.MinimumIntegralDigits = 1;
+			NumberFormatOptions.MaximumIntegralDigits = 6;
+			NumberFormatOptions.MinimumFractionalDigits = 2;
+			NumberFormatOptions.MaximumFractionalDigits = 6;
+			TextValue = ConcatenateLine(TextValue, FText::Format(LOCTEXT("CGError", "CGError: {0}"), FText::AsNumber(LinearSolveError, &NumberFormatOptions)));
+		}
 		if (const float SimulationTime = SimProxy->GetSimulationTime())
 		{
 			FNumberFormattingOptions NumberFormatOptions;
diff --git a/Engine/Source/Runtime/Experimental/Chaos/Chaos.Build.cs b/Engine/Source/Runtime/Experimental/Chaos/Chaos.Build.cs
--- a/Engine/Source/Runtime/Experimental/Chaos/Chaos.Build.cs
+++ b/Engine/Source/Runtime/Experimental/Chaos/Chaos.Build.cs
@@ -22,8 +22,13 @@ namespace UnrealBuildTool.Rules
 				"GeometryCore"
 				}
             );
+			PrivateDependencyModuleNames.AddRange(
+				new string[] {
+				"Eigen"
+				}
+			);
 
-            PublicDefinitions.Add("COMPILE_WITHOUT_UNREAL_SUPPORT=0");
+			PublicDefinitions.Add("COMPILE_WITHOUT_UNREAL_SUPPORT=0");
 			PrivateDefinitions.Add("CHAOS_INCLUDE_LEVEL_1=1");
 
 			if (Target.bUseChaosMemoryTracking == true)
diff --git a/Engine/Source/Runtime/Experimental/Chaos/Private/Chaos/Math/BlockSparseLinearSystem.cpp b/Engine/Source/Runtime/Experimental/Chaos/Private/Chaos/Math/BlockSparseLinearSystem.cpp
new file mode 100644
--- /dev/null
+++ b/Engine/Source/Runtime/Experimental/Chaos/Private/Chaos/Math/BlockSparseLinearSystem.cpp
@@ -0,0 +1,201 @@
+// Copyright Epic Games, Inc. All Rights Reserved.
+
+#include "Chaos/Math/BlockSparseLinearSystem.h"
+#include "ChaosCheck.h"
+#include "ChaosLog.h"
+#include "Chaos/Framework/Parallel.h"
+#include "GenericPlatform/GenericPlatformMath.h"
+#include "Misc/AssertionMacros.h"
+#include "Chaos/Vector.h"
+#include "HAL/IConsoleManager.h"
+
+#if defined(_MSC_VER) && USING_CODE_ANALYSIS
+#pragma warning(push)
+#pragma warning(disable : 6011)
+#pragma warning(disable : 6387)
+#pragma warning(disable : 6313)
+#pragma warning(disable : 6294)
+#endif
+PRAGMA_DEFAULT_VISIBILITY_START
+THIRD_PARTY_INCLUDES_START
+#include <Eigen/Sparse>
+#include <Eigen/Core>
+#include <Eigen/Dense>
+THIRD_PARTY_INCLUDES_END
+PRAGMA_DEFAULT_VISIBILITY_END
+#if defined(_MSC_VER) && USING_CODE_ANALYSIS
+#pragma warning(pop)
+#endif
+
+
+namespace Chaos {
+enum EPreconditionerType
+{
+	Diagonal = 0,
+	IncompleteCholesky = 1
+};
+static int32 PreconditionerType = EPreconditionerType::Diagonal;
+FAutoConsoleVariableRef CVarChaosLinearSystemPreconditioner(TEXT("p.Chaos.LinearSystem.Preconditioner"), PreconditionerType, TEXT("0 = Diagonal, 1 = IncompleteCholesky"));
+
+	template<typename T, int32 BlockSize>
+	struct TBlockSparseSymmetricLinearSystem<T, BlockSize>::FPimpl
+	{
+		// Just doing Eigen sparse matrix for now and building from triplets. Will investigate block symmetric later
+		typedef typename Eigen::SparseMatrix<T, Eigen::RowMajor> SparseMatrixType;
+		typedef typename Eigen::Triplet<T> MatrixTripletType;
+		typedef typename Eigen::Matrix<T, Eigen::Dynamic, 1> RealVectorType;
+		typedef typename Eigen::Map<RealVectorType> VectorType;
+		typedef typename Eigen::Map<const RealVectorType> ConstVectorType;
+
+		SparseMatrixType AMatrix;
+
+		// TripletBuilderNum is the actual number of entries.
+		// The TArray Num may be larger as it represents the amount of reserved space.
+		TArray<MatrixTripletType> MatrixBuilderTriplets;
+		std::atomic<int32> TripletBuilderNum;
+
+		FPimpl()
+			: TripletBuilderNum(0)
+		{}
+
+		void Reset(int32 NumRows)
+		{
+			TripletBuilderNum.store(0);
+			MatrixBuilderTriplets.Reset();
+			AMatrix.resize(NumRows*BlockSize, NumRows*BlockSize);
+		}
+		
+		void ReserveForParallelAdd(int32 NumDiagEntries, int32 NumOffDiagEntries)
+		{
+			MatrixBuilderTriplets.SetNum(TripletBuilderNum.load() + BlockSize * BlockSize* (NumDiagEntries + NumOffDiagEntries * 2));
+		}
+
+		void AddMatrixEntry(int32 Index0, int32 Index1, const Chaos::PMatrix<T, BlockSize, BlockSize>& AEntry)
+		{
+			if (Index0 == Index1)
+			{
+				int32 BuilderIndex = TripletBuilderNum.fetch_add(9);
+				for (int32 I = 0; I < BlockSize; ++I)
+				{
+					for (int32 J = 0; J < BlockSize; ++J)
+					{
+						check(FMath::IsNearlyEqual(AEntry.M[I][J], AEntry.M[J][I]));
+						check(BlockSize * Index0 + I < AMatrix.rows());
+						check(BlockSize * Index1 + J < AMatrix.cols());
+						MatrixBuilderTriplets[BuilderIndex++] = MatrixTripletType(BlockSize * Index0 + I, BlockSize * Index1 + J, AEntry.M[I][J]);
+					}
+				}
+			}
+			else
+			{
+				// Need to add symmetric part
+				int32 BuilderIndex = TripletBuilderNum.fetch_add(18);
+				for (int32 I = 0; I < BlockSize; ++I)
+				{
+					for (int32 J = 0; J < BlockSize; ++J)
+					{
+						check(BlockSize * Index0 + I < AMatrix.rows());
+						check(BlockSize * Index1 + J < AMatrix.cols());
+						check(BlockSize * Index0 + I < AMatrix.cols());
+						check(BlockSize * Index1 + J < AMatrix.rows());
+						MatrixBuilderTriplets[BuilderIndex++] = MatrixTripletType(BlockSize * Index0 + I, BlockSize * Index1 + J, AEntry.M[I][J]);
+						MatrixBuilderTriplets[BuilderIndex++] = MatrixTripletType(BlockSize * Index1 + J, BlockSize * Index0 + I, AEntry.M[I][J]);
+					}
+				}
+			}
+		}
+
+		void BuildMatrix()
+		{
+			constexpr bool bAllowShrinking = false;
+			MatrixBuilderTriplets.SetNum(TripletBuilderNum.load(), bAllowShrinking);
+			TRACE_CPUPROFILER_EVENT_SCOPE(ChaosBlockSparseLinearSystem_BuildMatrix);
+			AMatrix.setFromTriplets(MatrixBuilderTriplets.GetData(), MatrixBuilderTriplets.GetData() + MatrixBuilderTriplets.Num());
+		}
+
+		template<typename Preconditioner>
+		bool Solve(const TConstArrayView<TVector<T, BlockSize>>& RHS, const TArrayView<TVector<T, BlockSize>>& Result, const int32 MaxNumCGIterations, const T CGResidualTolerance, bool bCheckResidual, int32* OutIterations, T* OutError)
+		{
+			TRACE_CPUPROFILER_EVENT_SCOPE(ChaosBlockSparseLinearSystem_CGSolve);
+			check(Result.Num() == RHS.Num());
+			ConstVectorType WrappedRHS((const T*)RHS.GetData(), RHS.Num() * BlockSize);
+			VectorType WrappedResult((T*)Result.GetData(), Result.Num() * BlockSize);
+
+			// Just use Eigen's CG for now.
+			Eigen::ConjugateGradient<SparseMatrixType, Eigen::Lower | Eigen::Upper, Preconditioner> CG;
+			CG.setMaxIterations(MaxNumCGIterations);
+			CG.setTolerance(CGResidualTolerance);
+			// bCheckResidual is unused... Eigen always uses it
+			CG.compute(AMatrix);
+			WrappedResult = CG.solve(WrappedRHS);
+
+			if (OutIterations)
+			{
+				*OutIterations = (int32)CG.iterations();
+			}
+			if (OutError)
+			{
+				*OutError = (T)CG.error();
+			}
+			return CG.info() == Eigen::ComputationInfo::Success;
+		}
+	};
+
+	template<typename T, int32 BlockSize>
+	TBlockSparseSymmetricLinearSystem<T, BlockSize>::TBlockSparseSymmetricLinearSystem()
+		:Pimpl(MakePimpl<FPimpl>())
+	{}
+	template<typename T, int32 BlockSize>
+	TBlockSparseSymmetricLinearSystem<T, BlockSize>::TBlockSparseSymmetricLinearSystem(TBlockSparseSymmetricLinearSystem&& Other)
+		: Pimpl(MoveTemp(Other.Pimpl))
+	{}
+	template<typename T, int32 BlockSize>
+	TBlockSparseSymmetricLinearSystem<T,BlockSize>& TBlockSparseSymmetricLinearSystem<T, BlockSize>::operator=(TBlockSparseSymmetricLinearSystem&& Other)
+	{
+		Pimpl = MoveTemp(Other.Pimpl);
+		return *this;
+	}
+
+	template<typename T, int32 BlockSize>
+	void TBlockSparseSymmetricLinearSystem<T, BlockSize>::Reset(int32 NumRows)
+	{
+		Pimpl->Reset(NumRows);
+	}
+
+	template<typename T, int32 BlockSize>
+	void TBlockSparseSymmetricLinearSystem<T, BlockSize>::ReserveForParallelAdd(int32 NumDiagEntries, int32 NumOffDiagEntries)
+	{
+		Pimpl->ReserveForParallelAdd(NumDiagEntries, NumOffDiagEntries);
+	}
+
+	template<typename T, int32 BlockSize>
+	void TBlockSparseSymmetricLinearSystem<T, BlockSize>::AddMatrixEntry(int32 Index0, int32 Index1, const Chaos::PMatrix<T, BlockSize, BlockSize>& AEntry)
+	{
+		Pimpl->AddMatrixEntry(Index0, Index1, AEntry);
+	}
+
+	template<typename T, int32 BlockSize>
+	void TBlockSparseSymmetricLinearSystem<T, BlockSize>::FinalizeSystem()
+	{
+		Pimpl->BuildMatrix();
+	}
+
+	template<typename T, int32 BlockSize>
+	bool TBlockSparseSymmetricLinearSystem<T, BlockSize>::Solve(const TConstArrayView<TVector<T, BlockSize>>& RHS, const TArrayView<TVector<T, BlockSize>>& Result,
+		const int32 MaxNumCGIterations, const T CGResidualTolerance, bool bCheckResidual, 
+		int32* OptionalOutIterations, T* OptionalOutError) const
+	{
+		switch (PreconditionerType)
+		{
+		case EPreconditionerType::Diagonal:
+			return Pimpl->template Solve<Eigen::DiagonalPreconditioner<T>>(RHS, Result, MaxNumCGIterations, CGResidualTolerance, bCheckResidual, OptionalOutIterations, OptionalOutError);
+		case EPreconditionerType::IncompleteCholesky:
+			return Pimpl->template Solve<Eigen::IncompleteCholesky<T, Eigen::Upper | Eigen::Lower>>(RHS, Result, MaxNumCGIterations, CGResidualTolerance, bCheckResidual, OptionalOutIterations, OptionalOutError);
+		default:
+			checkNoEntry();
+		}
+		return false;
+	}
+	template class TBlockSparseSymmetricLinearSystem<FRealSingle, 3>;
+	template class TBlockSparseSymmetricLinearSystem<FRealDouble, 3>;
+}
diff --git a/Engine/Source/Runtime/Experimental/Chaos/Private/Chaos/PBDCollisionSpringConstraintsBase.cpp b/Engine/Source/Runtime/Experimental/Chaos/Private/Chaos/PBDCollisionSpringConstraintsBase.cpp
--- a/Engine/Source/Runtime/Experimental/Chaos/Private/Chaos/PBDCollisionSpringConstraintsBase.cpp
+++ b/Engine/Source/Runtime/Experimental/Chaos/Private/Chaos/PBDCollisionSpringConstraintsBase.cpp
@@ -5,6 +5,7 @@
 #include "Chaos/TriangleCollisionPoint.h"
 #include "Chaos/TriangleMesh.h"
 #include "Chaos/PBDSoftsSolverParticles.h"
+#include "Chaos/SoftsEvolutionLinearSystem.h"
 #include "Chaos/SoftsSolverParticlesRange.h"
 
 #if !COMPILE_WITHOUT_UNREAL_SUPPORT
@@ -39,10 +40,12 @@ FPBDCollisionSpringConstraintsBase::FPBDCollisionSpringConstraintsBase(
 	TSet<TVec2<int32>>&& InDisabledCollisionElements,
 	const FSolverReal InThickness,
 	const FSolverReal InStiffness,
-	const FSolverReal InFrictionCoefficient)
+	const FSolverReal InFrictionCoefficient,
+	const FSolverReal InProximityStiffness)
 	: Thickness(InThickness)
 	, Stiffness(InStiffness)
 	, FrictionCoefficient(InFrictionCoefficient)
+	, ProximityStiffness(InProximityStiffness)
 	, TriangleMesh(InTriangleMesh)
 	, Elements(InTriangleMesh.GetSurfaceElements())
 	, ReferencePositions(InReferencePositions)
@@ -103,6 +106,10 @@ void FPBDCollisionSpringConstraintsBase::Init(const SolverParticlesOrRange& Part
 			[this, &Spatial, &Particles, &ConstraintIndex, HeightSq, MaxConnectionsPerPoint, &VertexGIAColors, &TriangleGIAColors, &ReferencePositionsView](int32 i)
 			{
 				const int32 Index = i + Offset;
+				if (Particles.InvM(Index) == (FSolverReal)0.)
+				{
+					return;
+				}
 				constexpr FSolverReal ExtraThicknessMult = 1.5f;
 
 				TArray< TTriangleCollisionPoint<FSolverReal> > Result;
@@ -215,37 +222,37 @@ template void CHAOS_API FPBDCollisionSpringConstraintsBase::Init<FTriangleMesh::
 	const TConstArrayView<FPBDTriangleMeshCollisions::FGIAColor>& VertexGIAColors, const TArray<FPBDTriangleMeshCollisions::FGIAColor>& TriangleGIAColors);
 
 template<typename SolverParticlesOrRange>
-FSolverVec3 FPBDCollisionSpringConstraintsBase::GetDelta(const SolverParticlesOrRange& Particles, const int32 i) const
+FSolverVec3 FPBDCollisionSpringConstraintsBase::GetDelta(const SolverParticlesOrRange& Particles, const int32 ConstraintIndex) const
 {
-	const TVec4<int32>& Constraint = Constraints[i];
-	const int32 i1 = Constraint[0];
-	const int32 i2 = Constraint[1];
-	const int32 i3 = Constraint[2];
-	const int32 i4 = Constraint[3];
+	const TVec4<int32>& Constraint = Constraints[ConstraintIndex];
+	const int32 Index1 = Constraint[0];
+	const int32 Index2 = Constraint[1];
+	const int32 Index3 = Constraint[2];
+	const int32 Index4 = Constraint[3];
 
 	const FSolverReal TrianglePointInvM =
-		Particles.InvM(i2) * Barys[i][0] +
-		Particles.InvM(i3) * Barys[i][1] +
-		Particles.InvM(i4) * Barys[i][2];
+		Particles.InvM(Index2) * Barys[ConstraintIndex][0] +
+		Particles.InvM(Index3) * Barys[ConstraintIndex][1] +
+		Particles.InvM(Index4) * Barys[ConstraintIndex][2];
 
-	const FSolverReal CombinedMass = Particles.InvM(i1) + TrianglePointInvM;
+	const FSolverReal CombinedMass = Particles.InvM(Index1) + TrianglePointInvM;
 	if (CombinedMass <= (FSolverReal)1e-7)
 	{
 		return FSolverVec3(0);
 	}
 
-	const FSolverVec3& P1 = Particles.P(i1);
-	const FSolverVec3& P2 = Particles.P(i2);
-	const FSolverVec3& P3 = Particles.P(i3);
-	const FSolverVec3& P4 = Particles.P(i4);
+	const FSolverVec3& P1 = Particles.P(Index1);
+	const FSolverVec3& P2 = Particles.P(Index2);
+	const FSolverVec3& P3 = Particles.P(Index3);
+	const FSolverVec3& P4 = Particles.P(Index4);
 
 	const FSolverReal Height = Thickness + Thickness;
-	const FSolverVec3 P = Barys[i][0] * P2 + Barys[i][1] * P3 + Barys[i][2] * P4;
+	const FSolverVec3 P = Barys[ConstraintIndex][0] * P2 + Barys[ConstraintIndex][1] * P3 + Barys[ConstraintIndex][2] * P4;
 	const FSolverVec3 Difference = P1 - P;
 
 	// Normal repulsion with friction
 	const TTriangle<FSolverReal> Triangle(P2, P3, P4);
-	const FSolverVec3 Normal = FlipNormal[i] ? -Triangle.GetNormal() : Triangle.GetNormal();
+	const FSolverVec3 Normal = FlipNormal[ConstraintIndex] ? -Triangle.GetNormal() : Triangle.GetNormal();
 
 	const FSolverReal NormalDifference = Difference.Dot(Normal);
 	if (NormalDifference > Height)
@@ -258,9 +265,9 @@ FSolverVec3 FPBDCollisionSpringConstraintsBase::GetDelta(const SolverParticlesOr
 
 	if (FrictionCoefficient > 0)
 	{
-		const FSolverVec3& X1 = Particles.X(i1);
-		const FSolverVec3 X = Barys[i][0] * Particles.X(i2) + Barys[i][1] * Particles.X(i3) + Barys[i][2] * Particles.X(i4);
-		const FSolverVec3 RelativeDisplacement = (P1 - X1) - (P - X) + (Particles.InvM(i1) - TrianglePointInvM) * RepulsionDelta;
+		const FSolverVec3& X1 = Particles.X(Index1);
+		const FSolverVec3 X = Barys[ConstraintIndex][0] * Particles.X(Index2) + Barys[ConstraintIndex][1] * Particles.X(Index3) + Barys[ConstraintIndex][2] * Particles.X(Index4);
+		const FSolverVec3 RelativeDisplacement = (P1 - X1) - (P - X) + (Particles.InvM(Index1) - TrianglePointInvM) * RepulsionDelta;
 		const FSolverVec3 RelativeDisplacementTangent = RelativeDisplacement - RelativeDisplacement.Dot(Normal) * Normal;
 		const FSolverReal RelativeDisplacementTangentLength = RelativeDisplacementTangent.Length();
 		const FSolverReal PositionCorrection = FMath::Min(NormalDelta * FrictionCoefficient, RelativeDisplacementTangentLength);
@@ -276,6 +283,68 @@ FSolverVec3 FPBDCollisionSpringConstraintsBase::GetDelta(const SolverParticlesOr
 template CHAOS_API FSolverVec3 FPBDCollisionSpringConstraintsBase::GetDelta(const FSolverParticles& Particles, const int32 i) const;
 template CHAOS_API FSolverVec3 FPBDCollisionSpringConstraintsBase::GetDelta(const FSolverParticlesRange& Particles, const int32 i) const;
 
+void FPBDCollisionSpringConstraintsBase::UpdateLinearSystem(const FSolverParticlesRange& Particles, const FSolverReal Dt, FEvolutionLinearSystem& LinearSystem) const
+{
+	LinearSystem.ReserveForParallelAdd(Constraints.Num() * 4, Constraints.Num() * 3);
+	for (int32 Index = 0; Index < Constraints.Num(); ++Index)
+	{
+		const TVector<int32, 4>& Constraint = Constraints[Index];
+		const int32 Index1 = Constraint[0];
+		const int32 Index2 = Constraint[1];
+		const int32 Index3 = Constraint[2];
+		const int32 Index4 = Constraint[3];
+		const FSolverVec3& P1 = Particles.P(Index1);
+		const FSolverVec3& P2 = Particles.P(Index2);
+		const FSolverVec3& P3 = Particles.P(Index3);
+		const FSolverVec3& P4 = Particles.P(Index4);
+
+		const FSolverReal Height = Thickness + Thickness;
+		const FSolverVec3 P = Barys[Index][0] * P2 + Barys[Index][1] * P3 + Barys[Index][2] * P4;
+		const FSolverVec3 Difference = P1 - P;
+
+		// Normal repulsion with some stiction
+		const TTriangle<FSolverReal> Triangle(P2, P3, P4);
+		const FSolverVec3 Normal = FlipNormal[Index] ? -Triangle.GetNormal() : Triangle.GetNormal();
+
+		const FSolverReal NormalDifference = Difference.Dot(Normal);
+		if (NormalDifference > Height)
+		{
+			continue;
+		}
+
+		const FSolverReal NormalDelta = Height - NormalDifference;
+
+		const FSolverVec3 Force = ProximityStiffness * NormalDelta * Normal;
+		const FSolverMatrix33 DfDx = -ProximityStiffness * (((FSolverReal)1. - FrictionCoefficient) * FSolverMatrix33::OuterProduct(Normal, Normal) + FSolverMatrix33(FrictionCoefficient, FrictionCoefficient, FrictionCoefficient));
+
+		LinearSystem.AddForce(Particles, Force, Index1, Dt);
+		LinearSystem.AddSymmetricForceDerivative(Particles, &DfDx, nullptr, Index1, Index1, Dt);
+		if (Particles.InvM(Index2) > (FSolverReal)0.)
+		{
+			LinearSystem.AddForce(Particles, -Barys[Index][0] * Force, Index2, Dt);
+			FSolverMatrix33 DfDxScaled = -Barys[Index][0] * DfDx;
+			LinearSystem.AddSymmetricForceDerivative(Particles, &DfDxScaled, nullptr, Index1, Index2, Dt);
+			DfDxScaled *= -Barys[Index][0];
+			LinearSystem.AddSymmetricForceDerivative(Particles, &DfDxScaled, nullptr, Index2, Index2, Dt);
+		}
+		if (Particles.InvM(Index3) > (FSolverReal)0.)
+		{
+			LinearSystem.AddForce(Particles, -Barys[Index][1] * Force, Index3, Dt);
+			FSolverMatrix33 DfDxScaled = -Barys[Index][1] * DfDx;
+			LinearSystem.AddSymmetricForceDerivative(Particles, &DfDxScaled, nullptr, Index1, Index3, Dt);
+			DfDxScaled *= -Barys[Index][1];
+			LinearSystem.AddSymmetricForceDerivative(Particles, &DfDxScaled, nullptr, Index3, Index3, Dt);
+		}
+		if (Particles.InvM(Index4) > (FSolverReal)0.)
+		{
+			LinearSystem.AddForce(Particles, -Barys[Index][2] * Force, Index4, Dt);
+			FSolverMatrix33 DfDxScaled = -Barys[Index][2] * DfDx;
+			LinearSystem.AddSymmetricForceDerivative(Particles, &DfDxScaled, nullptr, Index1, Index4, Dt);
+			DfDxScaled *= -Barys[Index][2];
+			LinearSystem.AddSymmetricForceDerivative(Particles, &DfDxScaled, nullptr, Index4, Index4, Dt);
+		}
+	}
+}
 }  // End namespace Chaos::Softs
 
 #endif
diff --git a/Engine/Source/Runtime/Experimental/Chaos/Private/Chaos/PBDSoftBodyCollisionConstraint.cpp b/Engine/Source/Runtime/Experimental/Chaos/Private/Chaos/PBDSoftBodyCollisionConstraint.cpp
--- a/Engine/Source/Runtime/Experimental/Chaos/Private/Chaos/PBDSoftBodyCollisionConstraint.cpp
+++ b/Engine/Source/Runtime/Experimental/Chaos/Private/Chaos/PBDSoftBodyCollisionConstraint.cpp
@@ -1,8 +1,9 @@
 // Copyright Epic Games, Inc. All Rights Reserved.
 
 #include "Chaos/PBDSoftBodyCollisionConstraint.h"
-#include "Chaos/WeightedLatticeImplicitObject.h"
 #include "Chaos/Levelset.h"
+#include "Chaos/SoftsEvolutionLinearSystem.h"
+#include "Chaos/WeightedLatticeImplicitObject.h"
 #include "ChaosStats.h"
 #include "HAL/IConsoleManager.h"
 #include "Chaos/Framework/Parallel.h"
@@ -302,7 +303,7 @@ void FPBDSoftBodyCollisionConstraintBase::ApplyInternal(FSolverParticlesRange& P
 					FWeightedLatticeImplicitObject::FEmbeddingCoordinate SurfaceCoord;
 					Phi = (FSolverReal)LevelSet->PhiWithNormalAndSurfacePoint(RigidSpacePosition, ImplicitNormal, SurfaceCoord);
 					Penetration = CollisionThickness - Phi; // This is related to the Normal impulse
-					if (Penetration > (FSolverReal)0.)
+					if (bWithFriction && Penetration > (FSolverReal)0.)
 					{
 						const int32 StrongestBone = SurfaceCoord.GreatestInfluenceBone(LevelSet->GetBoneData());
 						if (StrongestBone != INDEX_NONE)
@@ -403,44 +404,47 @@ void FPBDSoftBodyCollisionConstraintBase::ApplyInternalCCD(FSolverParticlesRange
 					const FSolverReal Penetration = FMath::Max((FSolverReal)0., FSolverVec3::DotProduct(NormalWorld, Direction)) + (FSolverReal)UE_THRESH_POINT_ON_PLANE;
 
 					PAndInvM[Index].P += Penetration * NormalWorld;
-
-					// Friction
-					int32 VelocityBone = CollisionIndex;
-					if (const TWeightedLatticeImplicitObject<FLevelSet>* LevelSet = CollisionParticlesRange.GetGeometry(CollisionIndex)->GetObject< TWeightedLatticeImplicitObject<FLevelSet> >())
+					
+					if constexpr (bWithFriction)
 					{
-						TArray<FWeightedLatticeImplicitObject::FEmbeddingCoordinate> Coordinates;
-						LevelSet->GetEmbeddingCoordinates(PointPair.First, Coordinates, false);
-						int32 ClosestCoordIndex = INDEX_NONE;
-						double ClosestCoordPhi = UE_BIG_NUMBER;
-						for (int32 CoordIndex = 0; CoordIndex < Coordinates.Num(); ++CoordIndex)
+						// Friction
+						int32 VelocityBone = CollisionIndex;
+						if (const TWeightedLatticeImplicitObject<FLevelSet>* LevelSet = CollisionParticlesRange.GetGeometry(CollisionIndex)->GetObject< TWeightedLatticeImplicitObject<FLevelSet> >())
 						{
-							FVec3 NormalUnused;
-							const double CoordPhi = FMath::Abs(LevelSet->GetEmbeddedObject()->PhiWithNormal(Coordinates[CoordIndex].UndeformedPosition(LevelSet->GetGrid()), NormalUnused));
-							if (CoordPhi < ClosestCoordPhi)
+							TArray<FWeightedLatticeImplicitObject::FEmbeddingCoordinate> Coordinates;
+							LevelSet->GetEmbeddingCoordinates(PointPair.First, Coordinates, false);
+							int32 ClosestCoordIndex = INDEX_NONE;
+							double ClosestCoordPhi = UE_BIG_NUMBER;
+							for (int32 CoordIndex = 0; CoordIndex < Coordinates.Num(); ++CoordIndex)
 							{
-								ClosestCoordIndex = CoordIndex;
-								ClosestCoordPhi = CoordPhi;
+								FVec3 NormalUnused;
+								const double CoordPhi = FMath::Abs(LevelSet->GetEmbeddedObject()->PhiWithNormal(Coordinates[CoordIndex].UndeformedPosition(LevelSet->GetGrid()), NormalUnused));
+								if (CoordPhi < ClosestCoordPhi)
+								{
+									ClosestCoordIndex = CoordIndex;
+									ClosestCoordPhi = CoordPhi;
+								}
 							}
-						}
-						if (ClosestCoordIndex != INDEX_NONE)
-						{
-							const int32 StrongestBone = Coordinates[ClosestCoordIndex].GreatestInfluenceBone(LevelSet->GetBoneData());
-							if (StrongestBone != INDEX_NONE)
+							if (ClosestCoordIndex != INDEX_NONE)
 							{
-								VelocityBone = LevelSet->GetSolverBoneIndices()[StrongestBone];
+								const int32 StrongestBone = Coordinates[ClosestCoordIndex].GreatestInfluenceBone(LevelSet->GetBoneData());
+								if (StrongestBone != INDEX_NONE)
+								{
+									VelocityBone = LevelSet->GetSolverBoneIndices()[StrongestBone];
+								}
 							}
 						}
-					}
 
-					const FSolverVec3 VectorToPoint = PAndInvM[Index].P - CollisionParticlesRange.X(VelocityBone);
-					const FSolverVec3 RelativeDisplacement = (PAndInvM[Index].P - X[Index]) - (CollisionParticlesRange.V(VelocityBone) + FSolverVec3::CrossProduct(CollisionParticlesRange.W(VelocityBone), VectorToPoint)) * Dt;  // This corresponds to the tangential velocity multiplied by dt (friction will drive this to zero if it is high enough)
-					const FSolverVec3 RelativeDisplacementTangent = RelativeDisplacement - FSolverVec3::DotProduct(RelativeDisplacement, NormalWorld) * NormalWorld;  // Project displacement into the tangential plane
-					const FSolverReal RelativeDisplacementTangentLength = RelativeDisplacementTangent.Size();
-					if (RelativeDisplacementTangentLength >= UE_SMALL_NUMBER)
-					{
-						const FSolverReal PositionCorrection = FMath::Min<FSolverReal>(Penetration * FrictionCoefficient, RelativeDisplacementTangentLength);
-						const FSolverReal CorrectionRatio = PositionCorrection / RelativeDisplacementTangentLength;
-						PAndInvM[Index].P -= CorrectionRatio * RelativeDisplacementTangent;
+						const FSolverVec3 VectorToPoint = PAndInvM[Index].P - CollisionParticlesRange.X(VelocityBone);
+						const FSolverVec3 RelativeDisplacement = (PAndInvM[Index].P - X[Index]) - (CollisionParticlesRange.V(VelocityBone) + FSolverVec3::CrossProduct(CollisionParticlesRange.W(VelocityBone), VectorToPoint)) * Dt;  // This corresponds to the tangential velocity multiplied by dt (friction will drive this to zero if it is high enough)
+						const FSolverVec3 RelativeDisplacementTangent = RelativeDisplacement - FSolverVec3::DotProduct(RelativeDisplacement, NormalWorld) * NormalWorld;  // Project displacement into the tangential plane
+						const FSolverReal RelativeDisplacementTangentLength = RelativeDisplacementTangent.Size();
+						if (RelativeDisplacementTangentLength >= UE_SMALL_NUMBER)
+						{
+							const FSolverReal PositionCorrection = FMath::Min<FSolverReal>(Penetration * FrictionCoefficient, RelativeDisplacementTangentLength);
+							const FSolverReal CorrectionRatio = PositionCorrection / RelativeDisplacementTangentLength;
+							PAndInvM[Index].P -= CorrectionRatio * RelativeDisplacementTangent;
+						}
 					}
 				}
 			}
@@ -507,6 +511,73 @@ void FPBDSoftBodyCollisionConstraintBase::ApplyInternalISPC(FSolverParticlesRang
 #endif  // #if INTEL_ISPC
 }
 
+void FPBDSoftBodyCollisionConstraintBase::UpdateLinearSystem(const FSolverParticlesRange& Particles, const FSolverReal Dt, const TArray<FSolverCollisionParticlesRange>& CollisionParticles, FEvolutionLinearSystem& LinearSystem) const
+{
+	TRACE_CPUPROFILER_EVENT_SCOPE(FPBDSoftBodyCollisionConstraint_UpdateLinearSystem);
+
+	if (CollisionParticles.IsEmpty() || ProximityStiffness == (FSolverReal)0.f)
+	{
+		return;
+	}
+
+	// Just going to allocate enough space for all possible collisions. 
+	LinearSystem.ReserveForParallelAdd(Particles.GetRangeSize(), 0);
+
+	// Just proximity forces for now
+	const FPAndInvM* const PAndInvM = Particles.GetPAndInvM().GetData();
+	const FSolverReal ClampedFriction = FMath::Clamp(FrictionCoefficient, (FSolverReal)0., (FSolverReal)1.);
+	PhysicsParallelFor(Particles.GetRangeSize(), [this, Dt, ClampedFriction, &Particles, &PAndInvM, &CollisionParticles, &LinearSystem](int32 Index)
+	{
+		if (PAndInvM[Index].InvM == (FSolverReal)0.)
+		{
+			return;
+		}
+
+		bool bAddForce = false;
+		FSolverVec3 Force(0.f);
+		FSolverMatrix33 DfDx(0.f);
+		for (const FSolverCollisionParticlesRange& CollisionParticlesRange : CollisionParticles)
+		{
+			for (int32 CollisionIndex = 0; CollisionIndex < CollisionParticlesRange.GetRangeSize(); ++CollisionIndex)
+			{
+				if (CollisionParticlesRange.GetGeometry(CollisionIndex)->GetType() == Chaos::ImplicitObjectType::WeightedLatticeBone)
+				{
+					continue;
+				}
+
+				const FSolverRigidTransform3 Frame(CollisionParticlesRange.X(CollisionIndex), CollisionParticlesRange.R(CollisionIndex));
+				const FVec3 RigidSpacePosition(Frame.InverseTransformPosition(PAndInvM[Index].P));  // PhiWithNormal requires FReal based arguments
+				FVec3 ImplicitNormal;                                                                // since implicits don't use FSolverReal
+				const FSolverReal Phi = (FSolverReal)CollisionParticlesRange.GetGeometry(CollisionIndex)->PhiWithNormal(RigidSpacePosition, ImplicitNormal);
+				const FSolverReal Penetration = CollisionThickness - Phi; // This is related to the Normal impulse
+				const FSolverVec3 Normal(ImplicitNormal);
+
+				if (Penetration > (FSolverReal)0.)
+				{
+					bAddForce = true;
+
+					const FSolverVec3 NormalWorld = Frame.TransformVector(Normal);
+
+					// Repulsion force
+					Force += ProximityStiffness * Penetration * NormalWorld;
+
+					// Blend between a zero-length spring (stiction) and repulsion force based on friction
+					// DfDx = -ProximityStiffness * ((1-FrictionCoefficient)*OuterProduct(N,N) + FrictionCoefficient * Identity)
+					// Nothing here to match velocities... not sure if it's necessary, but this is a very stable force at least unlike any velocity-based thing.
+
+					DfDx += -ProximityStiffness * (((FSolverReal)1. - ClampedFriction) * FSolverMatrix33::OuterProduct(NormalWorld, NormalWorld) + FSolverMatrix33(ClampedFriction, ClampedFriction, ClampedFriction));
+				}
+			}
+		}
+
+		if (bAddForce)
+		{
+			LinearSystem.AddForce(Particles, Force, Index, Dt);
+			LinearSystem.AddSymmetricForceDerivative(Particles, &DfDx, nullptr, Index, Index, Dt);
+		}
+	});
+}
+
 void FPBDSoftBodyCollisionConstraint::SetProperties(const FCollectionPropertyConstFacade& PropertyCollection)
 {
 	if (IsCollisionThicknessMutable(PropertyCollection))
@@ -521,6 +592,10 @@ void FPBDSoftBodyCollisionConstraint::SetProperties(const FCollectionPropertyCon
 	{
 		bUseCCD = GetUseCCD(PropertyCollection);
 	}
+	if (IsProximityStiffnessMutable(PropertyCollection))
+	{
+		ProximityStiffness = GetProximityStiffness(PropertyCollection);
+	}
 }
 
 }  // End namespace Chaos::Softs
diff --git a/Engine/Source/Runtime/Experimental/Chaos/Private/Chaos/PBDSofts.isph b/Engine/Source/Runtime/Experimental/Chaos/Private/Chaos/PBDSofts.isph
--- a/Engine/Source/Runtime/Experimental/Chaos/Private/Chaos/PBDSofts.isph
+++ b/Engine/Source/Runtime/Experimental/Chaos/Private/Chaos/PBDSofts.isph
@@ -1,5 +1,7 @@
 // Copyright Epic Games, Inc. All Rights Reserved.
 
+#ifndef __PBDSOFTS_ISPH__
+#define __PBDSOFTS_ISPH__
 #define EXPLICIT_VECTOR4 1
 
 #include "Math/Vector.isph"
@@ -11,4 +13,5 @@ inline void UnzipPandInvM( const FVector4f& PandInvM, FVector3f& P, float& InvM
 	P.V[2] = PandInvM.V[2];
 
 	InvM = PandInvM.V[3];
-}
\ No newline at end of file
+}
+#endif
diff --git a/Engine/Source/Runtime/Experimental/Chaos/Private/Chaos/SoftsEvolution.cpp b/Engine/Source/Runtime/Experimental/Chaos/Private/Chaos/SoftsEvolution.cpp
--- a/Engine/Source/Runtime/Experimental/Chaos/Private/Chaos/SoftsEvolution.cpp
+++ b/Engine/Source/Runtime/Experimental/Chaos/Private/Chaos/SoftsEvolution.cpp
@@ -1,6 +1,7 @@
 // Copyright Epic Games, Inc. All Rights Reserved.
 #include "Chaos/SoftsEvolution.h"
 #include "Chaos/PerParticleDampVelocity.h"
+#include "Chaos/SoftsEvolutionLinearSystem.h"
 
 namespace Chaos::Softs {
 
@@ -12,6 +13,8 @@ constexpr FSolverReal SolverFrequency = (FSolverReal)60.;
 constexpr int32 NumIterations = 1;
 constexpr int32 MinNumIterations = 1;
 constexpr int32 MaxNumIterations = 10;
+constexpr bool bEnableForceBasedSolver = false;
+constexpr int32 NumNewtonIterations = 1;
 constexpr bool bDoQuasistatics = false;
 }
 namespace EvolutionSoftBodyDefault
@@ -122,15 +125,41 @@ static void QuasistaticUpdate(FSolverParticlesRange& Particles)
 	}
 }
 
+template<bool bUpdateX>
 static void PostStepUpdate(FSolverParticlesRange& Particles, const FSolverReal Dt)
 {
 	const FPAndInvM* const PAndInvM = Particles.GetPAndInvM().GetData();
 	FSolverVec3* const X = Particles.XArray().GetData();
 	FSolverVec3* const V = Particles.GetV().GetData();
+	FSolverVec3* const VPrev = Particles.GetVPrev().GetData();
 	for (int32 Index = 0; Index < Particles.GetRangeSize(); ++Index)
 	{
-		X[Index] = PAndInvM[Index].P;
+		VPrev[Index] = V[Index];
 		V[Index] = (PAndInvM[Index].P - X[Index]) / Dt;
+		if constexpr (bUpdateX)
+		{
+			X[Index] = PAndInvM[Index].P;
+		}
+	}
+}
+
+static void CopyVtoVPrev(FSolverParticlesRange& Particles)
+{
+	const FSolverVec3* const V = Particles.GetV().GetData();
+	FSolverVec3* const VPrev = Particles.GetVPrev().GetData();
+	for (int32 Index = 0; Index < Particles.GetRangeSize(); ++Index)
+	{
+		VPrev[Index] = V[Index];
+	}
+}
+
+static void CopyPToX(FSolverParticlesRange& Particles)
+{
+	const FPAndInvM* const PAndInvM = Particles.GetPAndInvM().GetData();
+	FSolverVec3* const X = Particles.XArray().GetData();
+	for (int32 Index = 0; Index < Particles.GetRangeSize(); ++Index)
+	{
+		X[Index] = PAndInvM[Index].P;
 	}
 }
 }
@@ -139,7 +168,6 @@ FEvolution::FEvolution(const FCollectionPropertyConstFacade& Properties)
 	: Time((FSolverReal)0.f)
 {
 	SetSolverProperties(Properties);
-
 	Particles.AddArray(&ParticleDampings);
 }
 
@@ -153,9 +181,12 @@ void FEvolution::Reset()
 
 	Groups.Reset();
 
-	ConstParticleRules.Reset();
-	ParticleRules.Reset();
-	CollisionRules.Reset();
+	ParallelInitRules.Reset();
+	ConstraintRules.Reset();
+	PBDConstraintRules.Reset();
+	PBDCollisionConstraintRules.Reset();
+	UpdateLinearSystemRules.Reset();
+	UpdateLinearSystemCollisionsRules.Reset();
 
 	KinematicUpdate.Reset();
 	CollisionKinematicUpdate.Reset();
@@ -294,18 +325,13 @@ TArray<FSolverCollisionParticlesRange> FEvolution::GetActiveCollisionParticles(u
 
 void FEvolution::SetSolverProperties(const FCollectionPropertyConstFacade& Properties)
 {
-	MaxNumIterations = FMath::Max(
-		Properties.GetValue<int32>(TEXT("MaxNumIterations"), Private::EvolutionSolverDefault::MaxNumIterations),
-		Private::EvolutionSolverDefault::MinNumIterations);
-
-	NumIterations = FMath::Clamp(
-		Properties.GetValue<int32>(TEXT("NumIterations"), Private::EvolutionSolverDefault::NumIterations),
-		Private::EvolutionSolverDefault::MinNumIterations,
-		MaxNumIterations);
-
-	bDoQuasistatics = Properties.GetValue<bool>(TEXT("DoQuasistatics"), Private::EvolutionSolverDefault::bDoQuasistatics);
-	
-	SolverFrequency = (FSolverReal)Properties.GetValue<FRealSingle>(TEXT("SolverFrequency"), Private::EvolutionSolverDefault::SolverFrequency);
+	bEnableForceBasedSolver = GetEnableForceBasedSolver(Properties, Private::EvolutionSolverDefault::bEnableForceBasedSolver);
+	MaxNumIterations = FMath::Min(GetMaxNumIterations(Properties, Private::EvolutionSolverDefault::MaxNumIterations), Private::EvolutionSolverDefault::MinNumIterations);
+	NumIterations = GetNumIterations(Properties, Private::EvolutionSolverDefault::NumIterations);
+	bDoQuasistatics = GetDoQuasistatics(Properties, Private::EvolutionSolverDefault::bDoQuasistatics);
+	SolverFrequency = GetSolverFrequency(Properties, Private::EvolutionSolverDefault::SolverFrequency);
+	NumNewtonIterations = GetNumNewtonIterations(Properties, Private::EvolutionSolverDefault::NumNewtonIterations);
+	LinearSystemParameters.SetProperties(Properties, NumIterations > 0);
 }
 
 void FEvolution::AdvanceOneTimeStep(const FSolverReal Dt, const FSolverReal TimeDependentIterationMultiplier)
@@ -315,9 +341,10 @@ void FEvolution::AdvanceOneTimeStep(const FSolverReal Dt, const FSolverReal Time
 	// Advance time
 	Time += Dt;
 
-	const int32 TimeDependentNumIterations = bDisableTimeDependentNumIterations ?
-		NumIterations :
-		FMath::Clamp(FMath::RoundToInt32(SolverFrequency * Dt * TimeDependentIterationMultiplier * (Softs::FSolverReal)NumIterations), Private::EvolutionSolverDefault::MinNumIterations, MaxNumIterations);
+	// Keep NumIterations == 0 as zero, but otherwise, clamp to at least 1.
+	const int32 TimeDependentNumIterations = (bDisableTimeDependentNumIterations || NumIterations == 0) ? NumIterations :
+		FMath::Clamp(FMath::RoundToInt32(SolverFrequency * Dt * TimeDependentIterationMultiplier * (Softs::FSolverReal)NumIterations),
+			Private::EvolutionSolverDefault::MinNumIterations, MaxNumIterations);
 
 	const TArray<uint32> ActiveGroupsArray = ActiveGroups.Array();
 
@@ -338,16 +365,26 @@ void FEvolution::AdvanceOneTimeStepInternal(const FSolverReal Dt, const int32 Ti
 	{
 		check(SoftBodies.Active[SoftBodyId]);
 
+		const bool bDoNewtonUpdate = bEnableForceBasedSolver && NumNewtonIterations > 0 && LinearSystemParameters.MaxNumCGIterations > 0 && SoftBodies.UpdateLinearSystemRules[SoftBodyId].GetRangeSize() > 0;
+
+		// Do at least one PBD iteration if not doing a newton update.
+		const int32 NumPBDIterations = bDoNewtonUpdate ? TimeDependentNumIterations : FMath::Max(1, TimeDependentNumIterations);
+
+		const ESolverMode SolverMode =
+			(NumPBDIterations > 0 ? ESolverMode::PBD : ESolverMode::None) |
+			(bDoNewtonUpdate ? ESolverMode::ForceBased : ESolverMode::None);
+		check(SolverMode != ESolverMode::None);
+
 		// PreSubstepParallelInits + LocalDamping
 		const bool bDoLocalDamping = SoftBodies.LocalDampings[SoftBodyId] > (FSolverReal)0.;
 		FPerParticleDampVelocity DampVelocityRule(SoftBodies.LocalDampings[SoftBodyId]);
 		{
 			TRACE_CPUPROFILER_EVENT_SCOPE(FSoftsEvolution_AdvanceOneTimeStepInternal_PreSubstepParallelInits);
 
-			TConstArrayView<PreSubstepParallelInitFunc> PreSubstepParallelInits = SoftBodies.PreSubstepParallelInits[SoftBodyId].GetConstArrayView();
+			TConstArrayView<ParallelInitFunc> PreSubstepParallelInits = SoftBodies.PreSubstepParallelInits[SoftBodyId].GetConstArrayView();
 			const int32 NumPreSubstepInits = PreSubstepParallelInits.Num() + (bDoLocalDamping ? 1 : 0);
 			PhysicsParallelFor(NumPreSubstepInits,
-				[this, &PreSubstepParallelInits, &DampVelocityRule, Dt, SoftBodyId](int32 FuncIdx)
+				[this, &PreSubstepParallelInits, &DampVelocityRule, Dt, SoftBodyId, SolverMode](int32 FuncIdx)
 			{
 				if (FuncIdx == PreSubstepParallelInits.Num())
 				{
@@ -355,7 +392,7 @@ void FEvolution::AdvanceOneTimeStepInternal(const FSolverReal Dt, const int32 Ti
 				}
 				else
 				{
-					PreSubstepParallelInits[FuncIdx](SoftBodies.ParticleRanges[SoftBodyId], Dt);
+					PreSubstepParallelInits[FuncIdx](SoftBodies.ParticleRanges[SoftBodyId], Dt, SolverMode);
 				}
 			}
 			);
@@ -369,109 +406,155 @@ void FEvolution::AdvanceOneTimeStepInternal(const FSolverReal Dt, const int32 Ti
 			KinematicUpdate(SoftBodies.ParticleRanges[SoftBodyId], Dt, Time);
 		}
 		{
-			// Dynamic initial guess (Euler step external forces and apply damping)
-			TRACE_CPUPROFILER_EVENT_SCOPE(FSoftsEvolution_AdvanceOneTimeStepInternal_InitialGuess);
-			if (bDoQuasistatics)
+			// Collision kinematic update
+			TRACE_CPUPROFILER_EVENT_SCOPE(FSoftsEvolution_AdvanceOneTimeStepInternal_CollisionKinematicUpdate);
+			check(CollisionKinematicUpdate);
+			for (FSolverCollisionParticlesRange& ActiveRange : ActiveCollisionRanges)
 			{
-				Private::QuasistaticUpdate(SoftBodies.ParticleRanges[SoftBodyId]);
+				CollisionKinematicUpdate(ActiveRange, Dt, Time);
 			}
-			else
-			{
-				// ExternalForceRules
-				TConstArrayView<ExternalForceRuleFunc> ExternalForceRules = SoftBodies.ExternalForceRules[SoftBodyId].GetConstArrayView();
-				for (const ExternalForceRuleFunc& ExternalForceRule : ExternalForceRules)
-				{
-					ExternalForceRule(SoftBodies.ParticleRanges[SoftBodyId], Dt);
-				}
-
-				// Euler Step Velocity
-				Private::EulerStepVelocity(SoftBodies.ParticleRanges[SoftBodyId], Dt);
-
-				if (bDoLocalDamping)
-				{
-					Private::DampLocalVelocity(SoftBodies.ParticleRanges[SoftBodyId], DampVelocityRule);
-				}
+		}
 
-				if (SoftBodies.UsePerParticleDamping[SoftBodyId])
+		// PBD initial guess
+		if (NumPBDIterations > 0)
+		{
+			{
+				// Dynamic initial guess (Euler step external forces and apply damping)
+				TRACE_CPUPROFILER_EVENT_SCOPE(FSoftsEvolution_AdvanceOneTimeStepInternal_InitialGuess);
+				if (bDoQuasistatics)
 				{
-					Private::EulerStepPositionWithGlobalDampingArray(SoftBodies.ParticleRanges[SoftBodyId], ParticleDampings, SolverFrequency, Dt);
+					Private::QuasistaticUpdate(SoftBodies.ParticleRanges[SoftBodyId]);
 				}
 				else
 				{
-					Private::EulerStepPositionWithGlobalDamping(SoftBodies.ParticleRanges[SoftBodyId], SoftBodies.GlobalDampings[SoftBodyId], SolverFrequency, Dt);
+					// ExternalForceRules
+					TConstArrayView<PBDConstraintRuleFunc> ExternalForceRules = SoftBodies.PBDExternalForceRules[SoftBodyId].GetConstArrayView();
+					for (const PBDConstraintRuleFunc& ExternalForceRule : ExternalForceRules)
+					{
+						ExternalForceRule(SoftBodies.ParticleRanges[SoftBodyId], Dt);
+					}
+
+					// Euler Step Velocity
+					Private::EulerStepVelocity(SoftBodies.ParticleRanges[SoftBodyId], Dt);
+
+					if (bDoLocalDamping)
+					{
+						Private::DampLocalVelocity(SoftBodies.ParticleRanges[SoftBodyId], DampVelocityRule);
+					}
+
+					if (SoftBodies.UsePerParticleDamping[SoftBodyId])
+					{
+						Private::EulerStepPositionWithGlobalDampingArray(SoftBodies.ParticleRanges[SoftBodyId], ParticleDampings, SolverFrequency, Dt);
+					}
+					else
+					{
+						Private::EulerStepPositionWithGlobalDamping(SoftBodies.ParticleRanges[SoftBodyId], SoftBodies.GlobalDampings[SoftBodyId], SolverFrequency, Dt);
+					}
 				}
 			}
 		}
-
-		{
-			// Collision kinematic update
-			TRACE_CPUPROFILER_EVENT_SCOPE(FSoftsEvolution_AdvanceOneTimeStepInternal_CollisionKinematicUpdate);
-			check(CollisionKinematicUpdate);
-			for (FSolverCollisionParticlesRange& ActiveRange : ActiveCollisionRanges)
-			{
-				CollisionKinematicUpdate(ActiveRange, Dt, Time);
-			}
-		}
-
 		{
 			// Parallel Inits
-			TRACE_CPUPROFILER_EVENT_SCOPE(FSoftsEvolution_AdvanceOneTimeStepInternal_ConstraintParallelInits);
-			TConstArrayView<ConstraintParallelInitFunc> ConstraintParallelInits = SoftBodies.ConstraintParallelInits[SoftBodyId].GetConstArrayView();
+			TRACE_CPUPROFILER_EVENT_SCOPE(FSoftsEvolution_AdvanceOneTimeStepInternal_PostInitialGuessParallelInits);
+			TConstArrayView<ParallelInitFunc> ConstraintParallelInits = SoftBodies.PostInitialGuessParallelInits[SoftBodyId].GetConstArrayView();
 			PhysicsParallelFor(ConstraintParallelInits.Num(),
-				[this, &ConstraintParallelInits, Dt, SoftBodyId](int32 Index)
+				[this, &ConstraintParallelInits, Dt, SoftBodyId, SolverMode](int32 Index)
 			{
-				ConstraintParallelInits[Index](SoftBodies.ParticleRanges[SoftBodyId], Dt);
+				ConstraintParallelInits[Index](SoftBodies.ParticleRanges[SoftBodyId], Dt, SolverMode);
 			});
 		}
-
 		{
 			// PreSubstep rules
 			TRACE_CPUPROFILER_EVENT_SCOPE(FSoftsEvolution_AdvanceOneTimeStepInternal_PreSubstepConstraintRules);
 			TConstArrayView<ConstraintRuleFunc> PreSubstepConstraintRules = SoftBodies.PreSubstepConstraintRules[SoftBodyId].GetConstArrayView();
 			for (const ConstraintRuleFunc& ConstraintRule : PreSubstepConstraintRules)
 			{
-				ConstraintRule(SoftBodies.ParticleRanges[SoftBodyId], Dt);
+				ConstraintRule(SoftBodies.ParticleRanges[SoftBodyId], Dt, SolverMode);
 			}
 		}
-
+		if (NumPBDIterations > 0)
 		{
-			// Iteration loop
-			TRACE_CPUPROFILER_EVENT_SCOPE(FSoftsEvolution_AdvanceOneTimeStepInternal_IterationLoop);
-			for (int32 i = 0; i < NumIterations; ++i)
 			{
-				TConstArrayView<ConstraintRuleFunc> PerIterationConstraintRules = SoftBodies.PerIterationConstraintRules[SoftBodyId].GetConstArrayView();
-				for (const ConstraintRuleFunc& ConstraintRule : PerIterationConstraintRules)
+				// Iteration loop
+				TRACE_CPUPROFILER_EVENT_SCOPE(FSoftsEvolution_AdvanceOneTimeStepInternal_IterationLoop);
+				for (int32 i = 0; i < NumPBDIterations; ++i)
 				{
-					ConstraintRule(SoftBodies.ParticleRanges[SoftBodyId], Dt);
+					TConstArrayView<PBDConstraintRuleFunc> PerIterationConstraintRules = SoftBodies.PerIterationPBDConstraintRules[SoftBodyId].GetConstArrayView();
+					for (const PBDConstraintRuleFunc& ConstraintRule : PerIterationConstraintRules)
+					{
+						ConstraintRule(SoftBodies.ParticleRanges[SoftBodyId], Dt);
+					}
+
+					TConstArrayView<PBDCollisionConstraintRuleFunc> PerIterationCollisionConstraintRules = SoftBodies.PerIterationCollisionPBDConstraintRules[SoftBodyId].GetConstArrayView();
+					for (const PBDCollisionConstraintRuleFunc& ConstraintRule : PerIterationCollisionConstraintRules)
+					{
+						ConstraintRule(SoftBodies.ParticleRanges[SoftBodyId], Dt, ActiveCollisionRanges);
+					}
+
+					TConstArrayView<PBDConstraintRuleFunc> PerIterationPostCollisionsConstraintRules = SoftBodies.PerIterationPostCollisionsPBDConstraintRules[SoftBodyId].GetConstArrayView();
+					for (const PBDConstraintRuleFunc& ConstraintRule : PerIterationPostCollisionsConstraintRules)
+					{
+						ConstraintRule(SoftBodies.ParticleRanges[SoftBodyId], Dt);
+					}
 				}
+			}
+			{
+				// Post Step Update
+				TRACE_CPUPROFILER_EVENT_SCOPE(FSoftsEvolution_AdvanceOneTimeStepInternal_PostStepUpdate);
+				if (bDoNewtonUpdate)
+				{
+					Private::PostStepUpdate<false>(SoftBodies.ParticleRanges[SoftBodyId], Dt);
+				}
+				else
+				{
+					Private::PostStepUpdate<true>(SoftBodies.ParticleRanges[SoftBodyId], Dt);
+				}
+			}
+		}
+		else
+		{
+			Private::CopyVtoVPrev(SoftBodies.ParticleRanges[SoftBodyId]);
+		}
+		
+		// Linear system update... currently we don't have any cross-softbody constraints. 
+		// If we did, the linear system would need to live at the group level.
+		if (bDoNewtonUpdate)
+		{
+			TConstArrayView<UpdateLinearSystemFunc> LinearSystemRules = SoftBodies.UpdateLinearSystemRules[SoftBodyId].GetConstArrayView();
+			check(LinearSystemRules.Num());
+			TConstArrayView<UpdateLinearSystemCollisionsFunc> LinearSystemCollisionRules = SoftBodies.UpdateLinearSystemCollisionsRules[SoftBodyId].GetConstArrayView();
+			check(NumNewtonIterations > 0);
+			for (int32 NewtonIteration = 0; NewtonIteration < NumNewtonIterations; ++NewtonIteration)
+			{
+				TRACE_CPUPROFILER_EVENT_SCOPE(FSoftsEvolution_AdvanceOneTimeStepInternal_NewtonIteration);
+				SoftBodies.LinearSystems[SoftBodyId].Init(SoftBodies.ParticleRanges[SoftBodyId], Dt, NewtonIteration == 0, LinearSystemParameters);
 
-				TConstArrayView<CollisionConstraintRuleFunc> PerIterationCollisionConstraintRules = SoftBodies.PerIterationCollisionConstraintRules[SoftBodyId].GetConstArrayView();
-				for (const CollisionConstraintRuleFunc& ConstraintRule : PerIterationCollisionConstraintRules)
+				for (const UpdateLinearSystemFunc& LinearSystemRule : LinearSystemRules)
 				{
-					ConstraintRule(SoftBodies.ParticleRanges[SoftBodyId], Dt, ActiveCollisionRanges);
+					LinearSystemRule(SoftBodies.ParticleRanges[SoftBodyId], Dt, SoftBodies.LinearSystems[SoftBodyId]);
 				}
 
-				TConstArrayView<ConstraintRuleFunc> PerIterationPostCollisionsConstraintRules = SoftBodies.PerIterationPostCollisionsConstraintRules[SoftBodyId].GetConstArrayView();
-				for (const ConstraintRuleFunc& ConstraintRule : PerIterationPostCollisionsConstraintRules)
+				for (const UpdateLinearSystemCollisionsFunc& LinearSystemCollisionRule : LinearSystemCollisionRules)
 				{
-					ConstraintRule(SoftBodies.ParticleRanges[SoftBodyId], Dt);
+					LinearSystemCollisionRule(SoftBodies.ParticleRanges[SoftBodyId], Dt, ActiveCollisionRanges, SoftBodies.LinearSystems[SoftBodyId]);
 				}
-			}
 
-			{
-				// Post Substep Rules
-				TRACE_CPUPROFILER_EVENT_SCOPE(FSoftsEvolution_AdvanceOneTimeStepInternal_PostSubstep);
-				TConstArrayView<ConstraintRuleFunc> PostSubstepConstraintRules = SoftBodies.PostSubstepConstraintRules[SoftBodyId].GetConstArrayView();
-				for (const ConstraintRuleFunc& ConstraintRule : PostSubstepConstraintRules)
+				if (!SoftBodies.LinearSystems[SoftBodyId].Solve(SoftBodies.ParticleRanges[SoftBodyId], Dt))
 				{
-					ConstraintRule(SoftBodies.ParticleRanges[SoftBodyId], Dt);
+					// Linear system update failed... just bail for now
+					break;
 				}
 			}
+			Private::CopyPToX(SoftBodies.ParticleRanges[SoftBodyId]);
+		}
 
+		{
+			// Post Substep Rules
+			TRACE_CPUPROFILER_EVENT_SCOPE(FSoftsEvolution_AdvanceOneTimeStepInternal_PostSubstep);
+			TConstArrayView<ConstraintRuleFunc> PostSubstepConstraintRules = SoftBodies.PostSubstepConstraintRules[SoftBodyId].GetConstArrayView();
+			for (const ConstraintRuleFunc& ConstraintRule : PostSubstepConstraintRules)
 			{
-				// Post Step Update
-				TRACE_CPUPROFILER_EVENT_SCOPE(FSoftsEvolution_AdvanceOneTimeStepInternal_PostStepUpdate);
-				Private::PostStepUpdate(SoftBodies.ParticleRanges[SoftBodyId], Dt);
+				ConstraintRule(SoftBodies.ParticleRanges[SoftBodyId], Dt, SolverMode);
 			}
 		}
 	}
diff --git a/Engine/Source/Runtime/Experimental/Chaos/Private/Chaos/SoftsEvolutionLinearSystem.cpp b/Engine/Source/Runtime/Experimental/Chaos/Private/Chaos/SoftsEvolutionLinearSystem.cpp
new file mode 100644
--- /dev/null
+++ b/Engine/Source/Runtime/Experimental/Chaos/Private/Chaos/SoftsEvolutionLinearSystem.cpp
@@ -0,0 +1,239 @@
+// Copyright Epic Games, Inc. All Rights Reserved.
+#include "Chaos/SoftsEvolutionLinearSystem.h"
+#include "Chaos/SoftsSolverParticlesRange.h"
+
+namespace Chaos::Softs
+{
+
+FEvolutionLinearSystem::FEvolutionLinearSystem(const FSolverParticlesRange& Particles)
+{
+	CalculateCompactIndices(Particles);
+	RHS.Init(FSolverVec3::ZeroVector, NumCompactIndices);
+}
+
+void FEvolutionLinearSystem::CalculateCompactIndices(const FSolverParticlesRange& Particles)
+{
+	NumCompactIndices = 0;
+	CompactifiedIndices.SetNumUninitialized(Particles.GetRangeSize());
+	for (int32 Index = 0; Index < Particles.GetRangeSize(); ++Index)
+	{
+		if (Particles.InvM(Index) == (FSolverReal)0.)
+		{
+			CompactifiedIndices[Index] = INDEX_NONE;
+		}
+		else
+		{
+			CompactifiedIndices[Index] = NumCompactIndices++;
+		}
+	}
+}
+
+void FEvolutionLinearSystem::Init(const FSolverParticlesRange& Particles, const FSolverReal Dt, bool bInFirstNewtonIteration, const FEvolutionLinearSystemSolverParameters& Params)
+{
+	Parameters = &Params;
+	if (CompactifiedIndices.Num() != Particles.Size())
+	{
+		CalculateCompactIndices(Particles);
+	}
+
+	TRACE_CPUPROFILER_EVENT_SCOPE(FEvolutionLinearSystem_Init);
+	// Note on variable naming convention
+	// X = Last time position
+	// P = Current time position guess
+	// VPrev = Last time velocity
+	// V = Current time velocity guess
+	//
+	// When solving for DeltaV,
+	// Solve M*(VNext - VPrev)/Dt = F, (XNext - X)/Dt = VNext
+	// Linearize about current guess of P and V, so 
+	// VNext = V + DeltaV, XNext = P + DeltaX
+	//  
+	// M*(V + DeltaV - VPrev) / Dt = F + DfDx * DeltaX + DfDv * DeltaV
+	// (P + DeltaX - X)/Dt = V + DeltaV
+	// Substitute DeltaX = (V + DeltaV) * Dt - (P - X), multiply by Dt.
+	// 
+	// (M - DfDx * Dt * Dt - DfDv * Dt) * DeltaV = F * Dt + M * (VPrev - V) + DfDx * (X - P + V * Dt) * Dt
+
+	// When using XPBD Initial Guess, or we're not the first Newton Iteration, P and V have been updated such that
+	// (1) P = X + V * Dt
+	// When first Newton Iteration, no XPBD,
+	// (2) P = X, V = VPrev
+	// Thus, in case (1), the RHS is
+	// F * Dt + M * (VPrev - V)
+	// In case (2), the RHS is
+	// F * Dt + DfDx * V * Dt * Dt
+
+	// When doing quasistatics and solve for DeltaX
+	// Solving F = 0
+	// Linearize about XNext = P + DeltaX
+	//
+	// -DfDx * DeltaX = F
+	// 
+	bDfDxTimesVTerm = !Parameters->bDoQuasistatics && !Parameters->bXPBDInitialGuess && bInFirstNewtonIteration;
+
+	// Initialize RHS
+	if (bDfDxTimesVTerm || Parameters->bDoQuasistatics)
+	{
+		RHS.Init(FSolverVec3::ZeroVector, NumCompactIndices);
+	}
+	else
+	{
+		check(CompactifiedIndices.Num() == Particles.GetRangeSize());
+
+		RHS.SetNumUninitialized(NumCompactIndices);
+		const FSolverReal* const M = Particles.GetM().GetData();
+		const FSolverVec3* const V = Particles.GetV().GetData();
+		const FSolverVec3* const VPrev = Particles.GetVPrev().GetData();
+		for (int32 Index = 0; Index < Particles.GetRangeSize(); ++Index)
+		{
+			if (CompactifiedIndices[Index] != INDEX_NONE)
+			{
+				RHS[CompactifiedIndices[Index]] = M[Index] * (VPrev[Index] - V[Index]);
+			}
+		}
+	}
+
+	// Initialize matrix as M when not doing quasistatics
+	Matrix.Reset(NumCompactIndices);
+	if (!Parameters->bDoQuasistatics)
+	{
+		Matrix.ReserveForParallelAdd(NumCompactIndices, 0);
+		const FSolverReal* const M = Particles.GetM().GetData();
+		for (int32 Index = 0; Index < Particles.GetRangeSize(); ++Index)
+		{
+			const int32 CompactIndex = CompactifiedIndices[Index];
+			if (CompactIndex != INDEX_NONE)
+			{
+				Matrix.AddMatrixEntry(CompactIndex, CompactIndex, FSolverMatrix33(M[Index], M[Index], M[Index]));
+			}
+		}
+	}
+}
+
+void FEvolutionLinearSystem::AddForce(const FSolverParticlesRange& Particles, const FSolverVec3& Force, int32 ParticleIndex, const FSolverReal Dt)
+{
+	const int32 CompactIndex = CompactifiedIndices[ParticleIndex];
+	
+	const FSolverReal DtScale = Parameters->bDoQuasistatics ? 1.f : Dt;
+	// RHS force contribution = F * Dt;
+	if (CompactIndex != INDEX_NONE)
+	{
+		RHS[CompactIndex] += Force * DtScale;
+	}
+}
+
+void FEvolutionLinearSystem::AddSymmetricForceDerivative(const FSolverParticlesRange& Particles, const FSolverMatrix33* const Df1Dx2, const FSolverMatrix33* const Df1Dv2, int32 ParticleIndex1, int32 ParticleIndex2, const FSolverReal Dt)
+{
+	check(Parameters);
+
+	const int32 CompactIndex1 = CompactifiedIndices[ParticleIndex1];
+	const int32 CompactIndex2 = CompactifiedIndices[ParticleIndex2];
+	if (Parameters->bDoQuasistatics)
+	{
+		if (CompactIndex1 != INDEX_NONE && CompactIndex2 != INDEX_NONE && Df1Dx2)
+		{
+			// -DfDx
+			Matrix.AddMatrixEntry(CompactIndex1, CompactIndex2, -(*Df1Dx2));
+		}
+	}
+	else
+	{
+		if(CompactIndex1 != INDEX_NONE && CompactIndex2 != INDEX_NONE)
+		{
+			// - DfDx * Dt * Dt - DfDv * Dt
+			FSolverMatrix33 MatrixToAdd;
+			if (Df1Dx2)
+			{
+				MatrixToAdd = (*Df1Dx2) * (-Dt * Dt);
+				if (Df1Dv2)
+				{
+					MatrixToAdd += (*Df1Dv2) * (-Dt);
+				}
+			}
+			else if (Df1Dv2)
+			{
+				MatrixToAdd = (*Df1Dv2) * (-Dt);
+			}
+			else
+			{
+				return;
+			}
+			Matrix.AddMatrixEntry(CompactIndex1, CompactIndex2, MatrixToAdd);
+		} // else Dirichlet boundaries with DeltaV (or DeltaX) = 0.
+
+		if (bDfDxTimesVTerm && Df1Dx2)
+		{
+			// DfDx * V * Dt * Dt
+			if (CompactIndex1 != INDEX_NONE)
+			{
+				RHS[CompactIndex1] += Df1Dx2->TransformVector(Particles.V(ParticleIndex1)) * (Dt * Dt);
+			}
+			if (CompactIndex2 != INDEX_NONE && CompactIndex1 != CompactIndex2)
+			{
+				RHS[CompactIndex2] += (*Df1Dx2) * Particles.V(ParticleIndex2) * (Dt * Dt);
+			}
+		}
+	}
+}
+
+bool FEvolutionLinearSystem::Solve(FSolverParticlesRange& Particles, const FSolverReal Dt)
+{
+	TRACE_CPUPROFILER_EVENT_SCOPE(FEvolutionLinearSystem_Solve);
+	check(Parameters);
+	Matrix.FinalizeSystem();
+
+	TArray<FSolverVec3> DeltaCompact;
+	DeltaCompact.Init(FSolverVec3::ZeroVector, NumCompactIndices);
+
+	const bool bSuccess = Matrix.Solve(TConstArrayView<FSolverVec3>(RHS), TArrayView<FSolverVec3>(DeltaCompact), Parameters->MaxNumCGIterations, Parameters->CGResidualTolerance, Parameters->bCheckCGResidual, &LastSolveIterations, &LastSolveError);
+	if (bSuccess)
+	{
+		if (Parameters->bDoQuasistatics)
+		{
+			// Solved for DeltaX
+			FPAndInvM* const PAndInvM = Particles.GetPAndInvM().GetData();
+			for (int32 ParticleIndex = 0; ParticleIndex < Particles.GetRangeSize(); ++ParticleIndex)
+			{
+				if (CompactifiedIndices[ParticleIndex] != INDEX_NONE)
+				{
+					PAndInvM[ParticleIndex].P += DeltaCompact[CompactifiedIndices[ParticleIndex]];
+				}
+			}
+		}
+		else
+		{
+			// Solved for DeltaV
+			if (bDfDxTimesVTerm)
+			{
+				const FSolverVec3* const X = Particles.XArray().GetData();
+				FPAndInvM* const PAndInvM = Particles.GetPAndInvM().GetData();
+				FSolverVec3* const V = Particles.GetV().GetData();
+				for (int32 ParticleIndex = 0; ParticleIndex < Particles.GetRangeSize(); ++ParticleIndex)
+				{
+					if (CompactifiedIndices[ParticleIndex] != INDEX_NONE)
+					{
+						V[ParticleIndex] += DeltaCompact[CompactifiedIndices[ParticleIndex]];
+						PAndInvM[ParticleIndex].P = X[ParticleIndex] + V[ParticleIndex] * Dt;
+					}
+				}
+			}
+			else
+			{
+				FPAndInvM* const PAndInvM = Particles.GetPAndInvM().GetData();
+				FSolverVec3* const V = Particles.GetV().GetData();
+				for (int32 ParticleIndex = 0; ParticleIndex < Particles.GetRangeSize(); ++ParticleIndex)
+				{
+					if (CompactifiedIndices[ParticleIndex] != INDEX_NONE)
+					{
+						V[ParticleIndex] += DeltaCompact[CompactifiedIndices[ParticleIndex]];
+						PAndInvM[ParticleIndex].P += DeltaCompact[CompactifiedIndices[ParticleIndex]] * Dt;
+					}
+				}
+			}
+		}
+		return true;
+	}
+	return false;
+}
+
+}  // End namespace Chaos::Softs
diff --git a/Engine/Source/Runtime/Experimental/Chaos/Private/Chaos/SoftsSpring.isph b/Engine/Source/Runtime/Experimental/Chaos/Private/Chaos/SoftsSpring.isph
new file mode 100644
--- /dev/null
+++ b/Engine/Source/Runtime/Experimental/Chaos/Private/Chaos/SoftsSpring.isph
@@ -0,0 +1,63 @@
+// Copyright Epic Games, Inc. All Rights Reserved.
+
+#ifndef __SOFTSSPRING_ISPH__
+#define __SOFTSSPRING_ISPH__
+
+#define EXPLICIT_VECTOR4 1
+
+#include "Math/Vector.isph"
+#include "Chaos/PBDSofts.isph"
+
+inline float SafeNormalize(FVector3f &Direction)
+{
+	const float SizeSquared = VectorSizeSquared(Direction);
+	const float Size = sqrt(SizeSquared);
+	Direction = VectorSelect((SizeSquared < FLOAT_KINDA_SMALL_NUMBER), FloatForwardVector, Direction / Size);
+	return (SizeSquared < FLOAT_KINDA_SMALL_NUMBER) ? FLOAT_ZERO : Size;
+}
+
+template <typename StiffnessFloatType>
+inline FVector3f GetXPBDSpringDelta(const FVector3f& P1, const FVector3f& P2, const varying float InvM1, const varying float InvM2,
+	 const uniform float Dt, const varying float Dist, const varying float Lambda, const StiffnessFloatType Stiffness, const uniform float MinStiffness, float& DLambda)
+{
+	FVector3f Direction = P1 - P2;
+	const float Distance = SafeNormalize(Direction);
+	const varying float Offset = Distance - Dist;
+	
+	const varying float AlphaInv = Stiffness * Dt * Dt;
+	
+	const float CombinedInvMass = InvM2 + InvM1;
+
+	const float ComputedDLambda = (AlphaInv * Offset - Lambda) / (AlphaInv * CombinedInvMass + (uniform float)1.f);
+	DLambda = select(Stiffness < MinStiffness, 0.f, ComputedDLambda);
+
+	const FVector3f ComputedDelta = DLambda * Direction;
+
+	return VectorSelect((InvM2 == FLOAT_ZERO && InvM1 == FLOAT_ZERO), FloatZeroVector, ComputedDelta);
+}
+
+template <typename StiffnessFloatType>
+inline FVector3f GetXPBDSpringDeltaWithDamping(const FVector3f& P1, const FVector3f& P2, const FVector3f& X1, const FVector3f& X2,
+	const varying float InvM1, const varying float InvM2,
+	const uniform float Dt, const varying float Dist, const varying float Lambda, const StiffnessFloatType Stiffness, const StiffnessFloatType DampingRatio, const uniform float MinStiffness, float& DLambda)
+{
+	FVector3f Direction = P1 - P2;
+	const float Distance = SafeNormalize(Direction);
+	const varying float Offset = Distance - Dist;
+	const varying float OneOverDist = select(Dist< FLOAT_SMALL_NUMBER, 1.f, 1.f / Dist);
+	
+	const varying float AlphaInv = Stiffness * Dt * Dt;
+	
+	const float CombinedInvMass = InvM2 + InvM1;
+	const varying float Damping = (uniform float)2.f * DampingRatio * sqrt(Stiffness / CombinedInvMass) * OneOverDist;
+	const varying float BetaDt = Damping * Dt;
+	const varying FVector3f RelativeVelocityTimesDt = P1 - X1 - P2 + X2;
+
+	const float ComputedDLambda = (AlphaInv * Offset - Lambda + BetaDt * VectorDot(Direction, RelativeVelocityTimesDt)) / ((AlphaInv + BetaDt)*CombinedInvMass + (uniform float)1.f);
+	DLambda = select(Stiffness < MinStiffness, 0.f, ComputedDLambda);
+
+	const FVector3f ComputedDelta = DLambda * Direction;
+
+	return VectorSelect((InvM2 == FLOAT_ZERO && InvM1 == FLOAT_ZERO), FloatZeroVector, ComputedDelta);
+}
+#endif
diff --git a/Engine/Source/Runtime/Experimental/Chaos/Private/Chaos/XPBDSpringConstraints.cpp b/Engine/Source/Runtime/Experimental/Chaos/Private/Chaos/XPBDSpringConstraints.cpp
--- a/Engine/Source/Runtime/Experimental/Chaos/Private/Chaos/XPBDSpringConstraints.cpp
+++ b/Engine/Source/Runtime/Experimental/Chaos/Private/Chaos/XPBDSpringConstraints.cpp
@@ -1,8 +1,9 @@
 // Copyright Epic Games, Inc. All Rights Reserved.
 #include "Chaos/XPBDSpringConstraints.h"
-#include "Chaos/PBDSoftsSolverParticles.h"
-#include "Chaos/GraphColoring.h"
 #include "Chaos/Framework/Parallel.h"
+#include "Chaos/GraphColoring.h"
+#include "Chaos/PBDSoftsSolverParticles.h"
+#include "Chaos/SoftsEvolutionLinearSystem.h"
 #include "ChaosStats.h"
 #include "ChaosLog.h"
 #include "HAL/IConsoleManager.h"
@@ -79,16 +80,16 @@ template<typename SolverParticlesOrRange>
 void FXPBDSpringConstraints::ApplyHelper(SolverParticlesOrRange& Particles, const FSolverReal Dt, const int32 ConstraintIndex, const FSolverReal ExpStiffnessValue, const FSolverReal DampingRatioValue) const
 {
 	const TVec2<int32>& Constraint = Constraints[ConstraintIndex];
-	const int32 i1 = Constraint[0];
-	const int32 i2 = Constraint[1];
+	const int32 Index1 = Constraint[0];
+	const int32 Index2 = Constraint[1];
 	const FSolverVec3 Delta =  GetDelta(Particles, Dt, ConstraintIndex, ExpStiffnessValue, DampingRatioValue);
-	if (Particles.InvM(i1) > (FSolverReal)0.)
+	if (Particles.InvM(Index1) > (FSolverReal)0.)
 	{
-		Particles.P(i1) -= Particles.InvM(i1) * Delta;
+		Particles.P(Index1) -= Particles.InvM(Index1) * Delta;
 	}
-	if (Particles.InvM(i2) > (FSolverReal)0.)
+	if (Particles.InvM(Index2) > (FSolverReal)0.)
 	{
-		Particles.P(i2) += Particles.InvM(i2) * Delta;
+		Particles.P(Index2) += Particles.InvM(Index2) * Delta;
 	}
 }
 
@@ -158,7 +159,7 @@ void FXPBDSpringConstraints::Apply(SolverParticlesOrRange& Particles, const FSol
 				{
 					const int32 ColorStart = ConstraintsPerColorStartIndex[ConstraintColorIndex];
 					const int32 ColorSize = ConstraintsPerColorStartIndex[ConstraintColorIndex + 1] - ColorStart;
-					PhysicsParallelFor(ColorSize, [&](const int32 Index)
+					PhysicsParallelFor(ColorSize, [this, &Particles, Dt, ColorStart, ExpStiffnessValue, DampingRatioValue](const int32 Index)
 					{
 						const int32 ConstraintIndex = ColorStart + Index;
 						ApplyHelper(Particles, Dt, ConstraintIndex, ExpStiffnessValue, DampingRatioValue);
@@ -222,7 +223,7 @@ void FXPBDSpringConstraints::Apply(SolverParticlesOrRange& Particles, const FSol
 				{
 					const int32 ColorStart = ConstraintsPerColorStartIndex[ConstraintColorIndex];
 					const int32 ColorSize = ConstraintsPerColorStartIndex[ConstraintColorIndex + 1] - ColorStart;
-					PhysicsParallelFor(ColorSize, [&](const int32 Index)
+					PhysicsParallelFor(ColorSize, [this, &Particles, Dt, ColorStart, StiffnessHasWeightMap, StiffnessNoMap, DampingHasWeightMap, DampingNoMap](const int32 Index)
 					{
 						const int32 ConstraintIndex = ColorStart + Index;
 						const FSolverReal ExpStiffnessValue = StiffnessHasWeightMap ? Stiffness[ConstraintIndex] : StiffnessNoMap;
@@ -264,6 +265,82 @@ void FXPBDSpringConstraints::Apply(SolverParticlesOrRange& Particles, const FSol
 template CHAOS_API void FXPBDSpringConstraints::Apply(FSolverParticles& Particles, const FSolverReal Dt) const;
 template CHAOS_API void FXPBDSpringConstraints::Apply(FSolverParticlesRange& Particles, const FSolverReal Dt) const;
 
+void FXPBDSpringConstraints::UpdateLinearSystem(const FSolverParticlesRange& Particles, const FSolverReal Dt, FEvolutionLinearSystem& LinearSystem) const
+{
+	TRACE_CPUPROFILER_EVENT_SCOPE(FXPBDSpringConstraints_UpdateLinearSystem);
+	LinearSystem.ReserveForParallelAdd(Constraints.Num() * 2, Constraints.Num());
+
+	const bool StiffnessHasWeightMap = Stiffness.HasWeightMap();
+	const bool DampingHasWeightMap = DampingRatio.HasWeightMap();
+	if ((ConstraintsPerColorStartIndex.Num() > 1) && (Constraints.Num() > Chaos_XPBDSpring_ParallelConstraintCount))
+	{
+		const int32 ConstraintColorNum = ConstraintsPerColorStartIndex.Num() - 1;
+		if (!StiffnessHasWeightMap && !DampingHasWeightMap)
+		{
+			const FSolverReal ExpStiffnessValue = (FSolverReal)Stiffness;
+			const FSolverReal DampingRatioValue = (FSolverReal)DampingRatio;
+			if (ExpStiffnessValue < MinStiffness)
+			{
+				return;
+			}
+			for (int32 ConstraintColorIndex = 0; ConstraintColorIndex < ConstraintColorNum; ++ConstraintColorIndex)
+			{
+				const int32 ColorStart = ConstraintsPerColorStartIndex[ConstraintColorIndex];
+				const int32 ColorSize = ConstraintsPerColorStartIndex[ConstraintColorIndex + 1] - ColorStart;
+				PhysicsParallelFor(ColorSize, [this, &Particles, Dt, &LinearSystem, ColorStart, ExpStiffnessValue, DampingRatioValue](const int32 Index)
+				{
+					const int32 ConstraintIndex = ColorStart + Index;
+					Spring::UpdateSpringLinearSystem(Particles, Dt, Constraints[ConstraintIndex], Dists[ConstraintIndex], ExpStiffnessValue, MinStiffness, DampingRatioValue, LinearSystem);
+				});
+			}
+		}
+		else // Has weight maps
+		{
+			const FSolverReal StiffnessNoMap = (FSolverReal)Stiffness;
+			const FSolverReal DampingNoMap = (FSolverReal)DampingRatio;
+			for (int32 ConstraintColorIndex = 0; ConstraintColorIndex < ConstraintColorNum; ++ConstraintColorIndex)
+			{
+				const int32 ColorStart = ConstraintsPerColorStartIndex[ConstraintColorIndex];
+				const int32 ColorSize = ConstraintsPerColorStartIndex[ConstraintColorIndex + 1] - ColorStart;
+				PhysicsParallelFor(ColorSize, [this, &Particles, Dt, &LinearSystem, ColorStart, StiffnessHasWeightMap, StiffnessNoMap, DampingHasWeightMap, DampingNoMap](const int32 Index)
+				{
+					const int32 ConstraintIndex = ColorStart + Index;
+					const FSolverReal ExpStiffnessValue = StiffnessHasWeightMap ? Stiffness[ConstraintIndex] : StiffnessNoMap;
+					const FSolverReal DampingRatioValue = DampingHasWeightMap ? DampingRatio[ConstraintIndex] : DampingNoMap;
+					Spring::UpdateSpringLinearSystem(Particles, Dt, Constraints[ConstraintIndex], Dists[ConstraintIndex], ExpStiffnessValue, MinStiffness, DampingRatioValue, LinearSystem);
+				});
+			}
+		}		
+	}
+	else
+	{
+		if (!StiffnessHasWeightMap && !DampingHasWeightMap)
+		{
+			const FSolverReal ExpStiffnessValue = (FSolverReal)Stiffness;
+			const FSolverReal DampingRatioValue = (FSolverReal)DampingRatio;
+			if (ExpStiffnessValue < MinStiffness)
+			{
+				return;
+			}
+			for (int32 ConstraintIndex = 0; ConstraintIndex < Constraints.Num(); ++ConstraintIndex)
+			{
+				Spring::UpdateSpringLinearSystem(Particles, Dt, Constraints[ConstraintIndex], Dists[ConstraintIndex], ExpStiffnessValue, MinStiffness, DampingRatioValue, LinearSystem);
+			}
+		}
+		else
+		{
+			const FSolverReal StiffnessNoMap = (FSolverReal)Stiffness;
+			const FSolverReal DampingNoMap = (FSolverReal)DampingRatio;
+			for (int32 ConstraintIndex = 0; ConstraintIndex < Constraints.Num(); ++ConstraintIndex)
+			{
+				const FSolverReal ExpStiffnessValue = StiffnessHasWeightMap ? Stiffness[ConstraintIndex] : StiffnessNoMap;
+				const FSolverReal DampingRatioValue = DampingHasWeightMap ? DampingRatio[ConstraintIndex] : DampingNoMap;
+				Spring::UpdateSpringLinearSystem(Particles, Dt, Constraints[ConstraintIndex], Dists[ConstraintIndex], ExpStiffnessValue, MinStiffness, DampingRatioValue, LinearSystem);
+			}
+		}
+	}
+}
+
 void FXPBDEdgeSpringConstraints::SetProperties(
 	const FCollectionPropertyConstFacade& PropertyCollection,
 	const TMap<FString, TConstArrayView<FRealSingle>>& WeightMaps)
diff --git a/Engine/Source/Runtime/Experimental/Chaos/Private/Chaos/XPBDSpringConstraints.ispc b/Engine/Source/Runtime/Experimental/Chaos/Private/Chaos/XPBDSpringConstraints.ispc
--- a/Engine/Source/Runtime/Experimental/Chaos/Private/Chaos/XPBDSpringConstraints.ispc
+++ b/Engine/Source/Runtime/Experimental/Chaos/Private/Chaos/XPBDSpringConstraints.ispc
@@ -4,14 +4,7 @@
 
 #include "Math/Vector.isph"
 #include "Chaos/PBDSofts.isph"
-
-static inline float SafeNormalize(FVector3f &Direction)
-{
-	const float SizeSquared = VectorSizeSquared(Direction);
-	const float Size = sqrt(SizeSquared);
-	Direction = VectorSelect((SizeSquared < FLOAT_KINDA_SMALL_NUMBER), FloatForwardVector, Direction / Size);
-	return (SizeSquared < FLOAT_KINDA_SMALL_NUMBER) ? FLOAT_ZERO : Size;
-}
+#include "Chaos/SoftsSpring.isph"
 
 export void ApplyXPBDSpringConstraints(uniform FVector4f PandInvM[],
 									const uniform FIntVector2 Constraints[],
@@ -34,23 +27,12 @@ export void ApplyXPBDSpringConstraints(uniform FVector4f PandInvM[],
 		varying float InvM1, InvM2;
 		UnzipPandInvM(PandInvM1, P1, InvM1);
 		UnzipPandInvM(PandInvM2, P2, InvM2);
-
-		FVector3f Direction = P1 - P2;
-		const float Distance = SafeNormalize(Direction);
-
+		
 		const varying float Dist = Dists[i];
-		const varying float FinalStiffness = select(Dist < FLOAT_SMALL_NUMBER, Stiffness, Stiffness / Dist);
 		const varying float Lambda = Lambdas[i];
-		const varying float Alpha = (uniform float)1.f / (FinalStiffness * Dt * Dt);
-		const varying float Offset = Distance - Dist;
-
-		const float CombinedInvMass = InvM2 + InvM1;
 
-		const varying float DLambda = (Offset - Alpha * Lambda) / (CombinedInvMass + Alpha);
-
-		const FVector3f ComputedDelta = DLambda * Direction;
-
-		const FVector3f Delta = VectorSelect((InvM2 == FLOAT_ZERO && InvM1 == FLOAT_ZERO), FloatZeroVector, ComputedDelta);
+		varying float DLambda;
+		const varying FVector3f Delta = GetXPBDSpringDelta(P1, P2, InvM1, InvM2, Dt, Dist, Lambda, Stiffness, FLOAT_SMALL_NUMBER, DLambda);
 
 		Lambdas[i] = Lambda + DLambda;
 
@@ -89,28 +71,14 @@ export void ApplyXPBDSpringConstraintsWithDamping(uniform FVector4f PandInvM[],
 		UnzipPandInvM(PandInvM1, P1, InvM1);
 		UnzipPandInvM(PandInvM2, P2, InvM2);
 
-		FVector3f Direction = P1 - P2;
-		const float Distance = SafeNormalize(Direction);
-
 		const varying FVector3f X1 = VectorGather(&X[i1]);
 		const varying FVector3f X2 = VectorGather(&X[i2]);
 
 		const varying float Dist = Dists[i];
-		const varying float FinalStiffness = select(Dist < FLOAT_SMALL_NUMBER, Stiffness, Stiffness / Dist);
 		const varying float Lambda = Lambdas[i];
-		const varying float Alpha = (uniform float)1.f / (FinalStiffness * Dt * Dt);
-		const varying float Offset = Distance - Dist;
-
-		const float CombinedInvMass = InvM2 + InvM1;
-		const varying float Damping = (uniform float)2.f * DampingRatio * sqrt(FinalStiffness / CombinedInvMass);
-		const varying float Gamma = Alpha * Damping * Dt;
-		const varying FVector3f RelativeVelocityTimesDt = P1 - X1 - P2 + X2;
-
-		const varying float DLambda = (Offset - Alpha * Lambda + Gamma * VectorDot(Direction, RelativeVelocityTimesDt)) / ((1.f + Gamma)*CombinedInvMass + Alpha);
-
-		const FVector3f ComputedDelta = DLambda * Direction;
 
-		const FVector3f Delta = VectorSelect((InvM2 == FLOAT_ZERO && InvM1 == FLOAT_ZERO), FloatZeroVector, ComputedDelta);
+		varying float DLambda;
+		const FVector3f Delta = GetXPBDSpringDeltaWithDamping(P1, P2, X1, X2, InvM1, InvM2, Dt, Dist, Lambda, Stiffness, DampingRatio, FLOAT_SMALL_NUMBER, DLambda);
 
 		Lambdas[i] = Lambda + DLambda;
 
@@ -150,30 +118,17 @@ export void ApplyXPBDSpringConstraintsWithWeightMaps(uniform FVector4f PandInvM[
 		UnzipPandInvM(PandInvM1, P1, InvM1);
 		UnzipPandInvM(PandInvM2, P2, InvM2);
 
-		FVector3f Direction = P1 - P2;
-		const float Distance = SafeNormalize(Direction);
-
 		const varying float Dist = Dists[i];
 		const varying float Lambda = Lambdas[i];
-		const varying float Offset = Distance - Dist;
 
 		const varying uint8 StiffnessIndex = StiffnessIndices[i];
 
 		#pragma ignore warning(perf)
 		const varying float Stiffness = StiffnessTable[StiffnessIndex];
 
-		const varying float FinalStiffness = select(Dist < FLOAT_SMALL_NUMBER, Stiffness, Stiffness / Dist);
-
-		const varying float Alpha = (uniform float)1.f / (FinalStiffness * Dt * Dt);
-
-		const float CombinedInvMass = InvM2 + InvM1;
+		varying float DLambda;
+		const varying FVector3f Delta = GetXPBDSpringDelta(P1, P2, InvM1, InvM2, Dt, Dist, Lambda, Stiffness, MinStiffness, DLambda);
 
-		const varying float ComputedDLambda = (Offset - Alpha * Lambda) / (CombinedInvMass + Alpha);
-		const varying float DLambda = select(FinalStiffness < MinStiffness, 0.f, ComputedDLambda);
-
-		const FVector3f ComputedDelta = DLambda * Direction;
-
-		const FVector3f Delta = VectorSelect((InvM2 == FLOAT_ZERO && InvM1 == FLOAT_ZERO), FloatZeroVector, ComputedDelta);
 		Lambdas[i] = Lambda + DLambda;
 
 		if (InvM1 > FLOAT_ZERO)
@@ -219,37 +174,19 @@ export void ApplyXPBDSpringConstraintsWithDampingAndWeightMaps(uniform FVector4f
 		UnzipPandInvM(PandInvM1, P1, InvM1);
 		UnzipPandInvM(PandInvM2, P2, InvM2);
 
-		FVector3f Direction = P1 - P2;
-		const float Distance = SafeNormalize(Direction);
-
 		const varying FVector3f X1 = VectorGather(&X[i1]);
 		const varying FVector3f X2 = VectorGather(&X[i2]);
 
 		const varying float Dist = Dists[i];
 		const varying float Lambda = Lambdas[i];
-		const varying float Offset = Distance - Dist;
-
 
 		#pragma ignore warning(perf)
 		const varying float Stiffness = StiffnessTable[StiffnessIndex];
 		#pragma ignore warning(perf)
 		const varying float DampingRatio = DampingTable[DampingIndex];
-		
-		const varying float FinalStiffness = select(Dist < FLOAT_SMALL_NUMBER, Stiffness, Stiffness / Dist);
-
-		const varying float Alpha = (uniform float)1.f / (FinalStiffness * Dt * Dt);
-
-		const float CombinedInvMass = InvM2 + InvM1;
-		const varying float Damping = (uniform float)2.f * DampingRatio * sqrt(FinalStiffness / CombinedInvMass);
-		const varying float Gamma = Alpha * Damping * Dt;
-		const varying FVector3f RelativeVelocityTimesDt = P1 - X1 - P2 + X2;
-
-		const varying float ComputedDLambda = (Offset - Alpha * Lambda + Gamma * VectorDot(Direction, RelativeVelocityTimesDt)) / ((1.f + Gamma)*CombinedInvMass + Alpha);
-		const varying float DLambda = select(FinalStiffness < MinStiffness, 0.f, ComputedDLambda);
-
-		const FVector3f ComputedDelta = DLambda * Direction;
 
-		const FVector3f Delta = VectorSelect((InvM2 == FLOAT_ZERO && InvM1 == FLOAT_ZERO), FloatZeroVector, ComputedDelta);
+		varying float DLambda;
+		const FVector3f Delta = GetXPBDSpringDeltaWithDamping(P1, P2, X1, X2, InvM1, InvM2, Dt, Dist, Lambda, Stiffness, DampingRatio, MinStiffness, DLambda);
 
 		Lambdas[i] = Lambda + DLambda;
 
diff --git a/Engine/Source/Runtime/Experimental/Chaos/Public/Chaos/CollectionPropertyFacade.h b/Engine/Source/Runtime/Experimental/Chaos/Public/Chaos/CollectionPropertyFacade.h
--- a/Engine/Source/Runtime/Experimental/Chaos/Public/Chaos/CollectionPropertyFacade.h
+++ b/Engine/Source/Runtime/Experimental/Chaos/Public/Chaos/CollectionPropertyFacade.h
@@ -494,8 +494,8 @@ PRAGMA_ENABLE_DEPRECATION_WARNINGS
 	}
 }  // End namespace Chaos
 
-// Use this macro to add shorthands for property getters and direct access through the declared key index
-#define UE_CHAOS_DECLARE_PROPERTYCOLLECTION_NAME(PropertyName, Type) \
+// Use this macro to add shorthands for property getters without a key index
+#define UE_CHAOS_DECLARE_INDEXLESS_PROPERTYCOLLECTION_NAME(PropertyName, Type) \
 	inline static const FName PropertyName##Name = TEXT(#PropertyName); \
 	UE_DEPRECATED(5.3, "PropertyName##String is to be removed as to not be confused with GetPropertyName##String().") \
 	static FString PropertyName##String() { return PropertyName##Name.ToString(); } \
@@ -537,7 +537,11 @@ PRAGMA_ENABLE_DEPRECATION_WARNINGS
 PRAGMA_DISABLE_DEPRECATION_WARNINGS \
 	return PropertyCollection.GetFlags(PropertyName##Name.ToString(), Default); \
 PRAGMA_ENABLE_DEPRECATION_WARNINGS \
-	} \
+	}
+
+// Use this macro to add shorthands for property getters and direct access through the declared key index
+#define UE_CHAOS_DECLARE_PROPERTYCOLLECTION_NAME(PropertyName, Type) \
+	UE_CHAOS_DECLARE_INDEXLESS_PROPERTYCOLLECTION_NAME(PropertyName, Type) \
 	Type GetLow##PropertyName(const FCollectionPropertyConstFacade& PropertyCollection) const \
 	{ \
 		checkSlow(PropertyName##Index == PropertyCollection.GetKeyIndex(PropertyName##Name.ToString())); \
@@ -621,3 +625,4 @@ PRAGMA_ENABLE_DEPRECATION_WARNINGS \
 		explicit F##PropertyName##Index(const FCollectionPropertyConstFacade& PropertyCollection) : Index(PropertyCollection.GetKeyIndex(PropertyName##Name.ToString())) {} \
 		operator int32() const { return Index; } \
 	} PropertyName##Index;
+
diff --git a/Engine/Source/Runtime/Experimental/Chaos/Public/Chaos/Math/BlockSparseLinearSystem.h b/Engine/Source/Runtime/Experimental/Chaos/Public/Chaos/Math/BlockSparseLinearSystem.h
new file mode 100644
--- /dev/null
+++ b/Engine/Source/Runtime/Experimental/Chaos/Public/Chaos/Math/BlockSparseLinearSystem.h
@@ -0,0 +1,41 @@
+// Copyright Epic Games, Inc. All Rights Reserved.
+#pragma once
+#include "Chaos/Core.h"
+#include "Templates/PimplPtr.h"
+
+namespace Chaos {
+
+// Note: currently only explicitly instantiated for 3x3 blocks
+template<typename T, int32 BlockSize>
+class TBlockSparseSymmetricLinearSystem final
+{
+public:
+	CHAOS_API TBlockSparseSymmetricLinearSystem();
+	TBlockSparseSymmetricLinearSystem(const TBlockSparseSymmetricLinearSystem&) = delete;
+	CHAOS_API TBlockSparseSymmetricLinearSystem(TBlockSparseSymmetricLinearSystem&&);
+	~TBlockSparseSymmetricLinearSystem() = default;
+	TBlockSparseSymmetricLinearSystem& operator=(TBlockSparseSymmetricLinearSystem&) = delete;
+	CHAOS_API TBlockSparseSymmetricLinearSystem& operator=(TBlockSparseSymmetricLinearSystem&&);
+
+	// Square system so NumRows == NumCols
+	void Reset(int32 NumRows);
+	void ReserveForParallelAdd(int32 NumDiagEntries, int32 NumOffDiagEntries);
+
+	void AddMatrixEntry(int32 Index0, int32 Index1, const Chaos::PMatrix<T, BlockSize, BlockSize>& AEntry);
+
+	// Call after all Matrix Entries have been added
+	void FinalizeSystem();
+
+	bool Solve(const TConstArrayView<TVector<T, BlockSize>>& RHS, const TArrayView<TVector<T, BlockSize>>& Result,
+		const int32 MaxNumCGIterations,	const T CGResidualTolerance, bool bCheckResidual, 
+		int32* OptionalOutIterations = nullptr, T* OptionalOutError = nullptr) const;
+
+private:
+
+	// Putting everything in here to avoid exposing Eigen
+	struct FPimpl;
+	TPimplPtr<FPimpl> Pimpl;
+};
+
+} // namespace Chaos
+
diff --git a/Engine/Source/Runtime/Experimental/Chaos/Public/Chaos/PBDCollisionSpringConstraints.h b/Engine/Source/Runtime/Experimental/Chaos/Public/Chaos/PBDCollisionSpringConstraints.h
--- a/Engine/Source/Runtime/Experimental/Chaos/Public/Chaos/PBDCollisionSpringConstraints.h
+++ b/Engine/Source/Runtime/Experimental/Chaos/Public/Chaos/PBDCollisionSpringConstraints.h
@@ -16,7 +16,7 @@ class FPBDCollisionSpringConstraints : public FPBDCollisionSpringConstraintsBase
 
 public:
 	static constexpr FSolverReal MinFrictionCoefficient = (FSolverReal)0.;
-	static constexpr FSolverReal MaxFrictionCoefficient = (FSolverReal)10.;
+	static constexpr FSolverReal MaxFrictionCoefficient = (FSolverReal)1.;
 
 	static bool IsEnabled(const FCollectionPropertyConstFacade& PropertyCollection)
 	{
@@ -38,10 +38,12 @@ public:
 			MoveTemp(InDisabledCollisionElements),
 			(FSolverReal)FMath::Max(GetSelfCollisionThickness(PropertyCollection, Base::BackCompatThickness), 0.f),
 			(FSolverReal)FMath::Clamp(GetSelfCollisionStiffness(PropertyCollection, Base::BackCompatStiffness), 0.f, 1.f),
-			FMath::Clamp((FSolverReal)GetSelfCollisionFriction(PropertyCollection, Base::BackCompatFrictionCoefficient), MinFrictionCoefficient, MaxFrictionCoefficient))
+			FMath::Clamp((FSolverReal)GetSelfCollisionFriction(PropertyCollection, Base::BackCompatFrictionCoefficient), MinFrictionCoefficient, MaxFrictionCoefficient),
+			(FSolverReal)GetSelfCollisionProximityStiffness(PropertyCollection, Base::DefaultProximityStiffness))
 		, SelfCollisionThicknessIndex(PropertyCollection)
 		, SelfCollisionStiffnessIndex(PropertyCollection)
 		, SelfCollisionFrictionIndex(PropertyCollection)
+		, SelfCollisionProximityStiffnessIndex(PropertyCollection)
 	{}
 
 	FPBDCollisionSpringConstraints(
@@ -65,6 +67,7 @@ public:
 		, SelfCollisionThicknessIndex(ForceInit)
 		, SelfCollisionStiffnessIndex(ForceInit)
 		, SelfCollisionFrictionIndex(ForceInit)
+		, SelfCollisionProximityStiffnessIndex(ForceInit)
 	{}
 
 	virtual ~FPBDCollisionSpringConstraints() override {}
@@ -85,50 +88,9 @@ public:
 		{
 			FrictionCoefficient = FMath::Clamp((FSolverReal)GetSelfCollisionFriction(PropertyCollection), MinFrictionCoefficient, MaxFrictionCoefficient);
 		}
-	}
-
-	template<typename SolverParticlesOrRange>
-	void Apply(SolverParticlesOrRange& Particles, const FSolverReal Dt, const int32 ConstraintIndex) const
-	{
-		const int32 i = ConstraintIndex;
-		const TVector<int32, 4>& Constraint = Constraints[i];
-		const int32 i1 = Constraint[0];
-		const int32 i2 = Constraint[1];
-		const int32 i3 = Constraint[2];
-		const int32 i4 = Constraint[3];
-		const FSolverVec3 Delta = Base::GetDelta(Particles, i);
-		if (Particles.InvM(i1) > 0)
-		{
-			Particles.P(i1) += Particles.InvM(i1) * Delta;
-		}
-		if (Particles.InvM(i2) > (FSolverReal)0.)
-		{
-			Particles.P(i2) -= Particles.InvM(i2) * Barys[i][0] * Delta;
-		}
-		if (Particles.InvM(i3) > (FSolverReal)0.)
-		{
-			Particles.P(i3) -= Particles.InvM(i3) * Barys[i][1] * Delta;
-		}
-		if (Particles.InvM(i4) > (FSolverReal)0.)
-		{
-			Particles.P(i4) -= Particles.InvM(i4) * Barys[i][2] * Delta;
-		}
-	}
-
-	template<typename SolverParticlesOrRange>
-	void Apply(SolverParticlesOrRange& InParticles, const FSolverReal Dt) const
-	{
-		for (int32 i = 0; i < Constraints.Num(); ++i)
-		{
-			Apply(InParticles, Dt, i);
-		}
-	}
-
-	void Apply(FSolverParticles& InParticles, const FSolverReal Dt, const TArray<int32>& InConstraintIndices) const
-	{
-		for (int32 i : InConstraintIndices)
+		if (IsSelfCollisionProximityStiffnessMutable(PropertyCollection))
 		{
-			Apply(InParticles, Dt, i);
+			ProximityStiffness = GetSelfCollisionProximityStiffness(PropertyCollection);
 		}
 	}
 
@@ -136,10 +98,12 @@ private:
 	using Base::Thickness;
 	using Base::Stiffness;
 	using Base::FrictionCoefficient;
+	using Base::ProximityStiffness;
 
 	UE_CHAOS_DECLARE_PROPERTYCOLLECTION_NAME(SelfCollisionThickness, float);
 	UE_CHAOS_DECLARE_PROPERTYCOLLECTION_NAME(SelfCollisionStiffness, float);
 	UE_CHAOS_DECLARE_PROPERTYCOLLECTION_NAME(SelfCollisionFriction, float);
+	UE_CHAOS_DECLARE_PROPERTYCOLLECTION_NAME(SelfCollisionProximityStiffness, float);
 };
 
 }  // End namespace Chaos::Softs
diff --git a/Engine/Source/Runtime/Experimental/Chaos/Public/Chaos/PBDCollisionSpringConstraintsBase.h b/Engine/Source/Runtime/Experimental/Chaos/Public/Chaos/PBDCollisionSpringConstraintsBase.h
--- a/Engine/Source/Runtime/Experimental/Chaos/Public/Chaos/PBDCollisionSpringConstraintsBase.h
+++ b/Engine/Source/Runtime/Experimental/Chaos/Public/Chaos/PBDCollisionSpringConstraintsBase.h
@@ -24,6 +24,7 @@ public:
 	static constexpr FSolverReal BackCompatThickness = (FSolverReal)1.f;
 	static constexpr FSolverReal BackCompatStiffness = (FSolverReal)0.5f;
 	static constexpr FSolverReal BackCompatFrictionCoefficient = (FSolverReal)0.f;
+	static constexpr FSolverReal DefaultProximityStiffness = (FSolverReal)1.;
 
 	CHAOS_API FPBDCollisionSpringConstraintsBase(
 		const int32 InOffset,
@@ -33,7 +34,8 @@ public:
 		TSet<TVec2<int32>>&& InDisabledCollisionElements,
 		const FSolverReal InThickness = BackCompatThickness,
 		const FSolverReal InStiffness = BackCompatStiffness,
-		const FSolverReal InFrictionCoefficient = BackCompatFrictionCoefficient);
+		const FSolverReal InFrictionCoefficient = BackCompatFrictionCoefficient,
+		const FSolverReal InProximityStiffness = DefaultProximityStiffness);
 
 	virtual ~FPBDCollisionSpringConstraintsBase() {}
 
@@ -55,13 +57,60 @@ public:
 	void SetThickness(FSolverReal InThickness) { Thickness = FMath::Max(InThickness, (FSolverReal)0.);  }
 	void SetFrictionCoefficient(FSolverReal InFrictionCoefficient) { FrictionCoefficient = InFrictionCoefficient; }
 
+	template<typename SolverParticlesOrRange>
+	void Apply(SolverParticlesOrRange& Particles, const FSolverReal Dt, const int32 ConstraintIndex) const
+	{
+		const TVector<int32, 4>& Constraint = Constraints[ConstraintIndex];
+		const int32 Index1 = Constraint[0];
+		const int32 Index2 = Constraint[1];
+		const int32 Index3 = Constraint[2];
+		const int32 Index4 = Constraint[3];
+		const FSolverVec3 Delta = GetDelta(Particles, ConstraintIndex);
+		if (Particles.InvM(Index1) > 0)
+		{
+			Particles.P(Index1) += Particles.InvM(Index1) * Delta;
+		}
+		if (Particles.InvM(Index2) > (FSolverReal)0.)
+		{
+			Particles.P(Index2) -= Particles.InvM(Index2) * Barys[ConstraintIndex][0] * Delta;
+		}
+		if (Particles.InvM(Index3) > (FSolverReal)0.)
+		{
+			Particles.P(Index3) -= Particles.InvM(Index3) * Barys[ConstraintIndex][1] * Delta;
+		}
+		if (Particles.InvM(Index4) > (FSolverReal)0.)
+		{
+			Particles.P(Index4) -= Particles.InvM(Index4) * Barys[ConstraintIndex][2] * Delta;
+		}
+	}
+
+	template<typename SolverParticlesOrRange>
+	void Apply(SolverParticlesOrRange& InParticles, const FSolverReal Dt) const
+	{
+		for (int32 ConstraintIndex = 0; ConstraintIndex < Constraints.Num(); ++ConstraintIndex)
+		{
+			Apply(InParticles, Dt, ConstraintIndex);
+		}
+	}
+
+	void Apply(FSolverParticles& InParticles, const FSolverReal Dt, const TArray<int32>& InConstraintIndices) const
+	{
+		for (int32 ConstraintIndex : InConstraintIndices)
+		{
+			Apply(InParticles, Dt, ConstraintIndex);
+		}
+	}
+
+	CHAOS_API void UpdateLinearSystem(const FSolverParticlesRange& Particles, const FSolverReal Dt, FEvolutionLinearSystem& LinearSystem) const;
+
 protected:
 	TArray<TVec4<int32>> Constraints;
 	TArray<FSolverVec3> Barys;
 	TArray<bool> FlipNormal;
 	FSolverReal Thickness;
-	FSolverReal Stiffness;
+	FSolverReal Stiffness; // (0-1 compliance for PBD)
 	FSolverReal FrictionCoefficient;
+	FSolverReal ProximityStiffness; // (actual spring stiffness for force-based solver)
 
 private:
 	const FTriangleMesh& TriangleMesh;
diff --git a/Engine/Source/Runtime/Experimental/Chaos/Public/Chaos/PBDSoftBodyCollisionConstraint.h b/Engine/Source/Runtime/Experimental/Chaos/Public/Chaos/PBDSoftBodyCollisionConstraint.h
--- a/Engine/Source/Runtime/Experimental/Chaos/Public/Chaos/PBDSoftBodyCollisionConstraint.h
+++ b/Engine/Source/Runtime/Experimental/Chaos/Public/Chaos/PBDSoftBodyCollisionConstraint.h
@@ -31,6 +31,7 @@ public:
 		FSolverReal InCollisionThickness,
 		FSolverReal InFrictionCoefficient,
 		bool bInUseCCD,
+		FSolverReal InProximityStiffness,
 		TArray<bool>* InCollisionParticleCollided = nullptr,
 		TArray<FSolverVec3>* InContacts = nullptr,
 		TArray<FSolverVec3>* InNormals = nullptr,
@@ -40,6 +41,7 @@ public:
 		, CollisionThickness(InCollisionThickness)
 		, FrictionCoefficient(InFrictionCoefficient)
 		, bUseCCD(bInUseCCD)
+		, ProximityStiffness(InProximityStiffness)
 		, CollisionParticleCollided(InCollisionParticleCollided)
 		, Contacts(InContacts)
 		, Normals(InNormals)
@@ -50,12 +52,13 @@ public:
 
 	CHAOS_API void Apply(FSolverParticlesRange& Particles, const FSolverReal Dt, const TArray<FSolverCollisionParticlesRange>& CollisionParticles) const;
 
+	CHAOS_API void UpdateLinearSystem(const FSolverParticlesRange& Particles, const FSolverReal Dt, const TArray<FSolverCollisionParticlesRange>& CollisionParticles, FEvolutionLinearSystem& LinearSystem) const;
+
 private:
 	template<bool bLockAndWriteContacts, bool bWithFriction>
 	void ApplyInternal(FSolverParticlesRange& Particles, const FSolverReal Dt, const TArray<FSolverCollisionParticlesRange>& CollisionParticles) const;
 	template<bool bLockAndWriteContacts, bool bWithFriction>
 	void ApplyInternalCCD(FSolverParticlesRange& Particles, const FSolverReal Dt, const TArray<FSolverCollisionParticlesRange>& CollisionParticles) const;
-
 	template<bool bWithFriction>
 	void ApplyInternalISPC(FSolverParticlesRange& Particles, const FSolverReal Dt, const TArray<FSolverCollisionParticlesRange>& CollisionParticles) const;
 
@@ -64,6 +67,7 @@ protected:
 	FSolverReal CollisionThickness;
 	FSolverReal FrictionCoefficient;
 	bool bUseCCD;
+	FSolverReal ProximityStiffness; // Used by force-based solver
 
 	/**  Used for writing debug contacts */
 	bool bWriteDebugContacts = false;
@@ -83,6 +87,7 @@ class FPBDSoftBodyCollisionConstraint : public FPBDSoftBodyCollisionConstraintBa
 public:
 	static constexpr FSolverReal DefaultCollisionThickness = (FSolverReal)1.;
 	static constexpr FSolverReal DefaultFrictionCoefficient = (FSolverReal)0.8;
+	static constexpr FSolverReal DefaultProximityStiffness = (FSolverReal)100.;
 
 	FPBDSoftBodyCollisionConstraint(
 		const TArray<FSolverRigidTransform3>& InLastCollisionTransforms,
@@ -96,11 +101,13 @@ public:
 			InMeshScale * GetCollisionThickness(PropertyCollection, DefaultCollisionThickness),
 			GetFrictionCoefficient(PropertyCollection, DefaultFrictionCoefficient),
 			GetUseCCD(PropertyCollection, false),
+			GetProximityStiffness(PropertyCollection, DefaultProximityStiffness),
 			InCollisionParticleCollided, InContacts, InNormals, InPhis)
 		, MeshScale(InMeshScale)
 		, CollisionThicknessIndex(PropertyCollection)
 		, FrictionCoefficientIndex(PropertyCollection)
 		, UseCCDIndex(PropertyCollection)
+		, ProximityStiffnessIndex(PropertyCollection)
 	{}
 
 	CHAOS_API void SetProperties(const FCollectionPropertyConstFacade& PropertyCollection);
@@ -110,6 +117,7 @@ private:
 	UE_CHAOS_DECLARE_PROPERTYCOLLECTION_NAME(CollisionThickness, float);
 	UE_CHAOS_DECLARE_PROPERTYCOLLECTION_NAME(FrictionCoefficient, float);
 	UE_CHAOS_DECLARE_PROPERTYCOLLECTION_NAME(UseCCD, bool);
+	UE_CHAOS_DECLARE_PROPERTYCOLLECTION_NAME(ProximityStiffness, float);
 };
 
 }  // End namespace Chaos::Softs
diff --git a/Engine/Source/Runtime/Experimental/Chaos/Public/Chaos/PBDSoftsEvolutionFwd.h b/Engine/Source/Runtime/Experimental/Chaos/Public/Chaos/PBDSoftsEvolutionFwd.h
--- a/Engine/Source/Runtime/Experimental/Chaos/Public/Chaos/PBDSoftsEvolutionFwd.h
+++ b/Engine/Source/Runtime/Experimental/Chaos/Public/Chaos/PBDSoftsEvolutionFwd.h
@@ -81,5 +81,8 @@ namespace Chaos
 		class FVelocityAndPressureField;
 		using FVelocityField UE_DEPRECATED(5.1, "Chaos::Softs::FVelocityField has been renamed FVelocityAndPressureField to match its new behavior.") = FVelocityAndPressureField;
 		class FExternalForces;
+
+		// Linear system solver
+		class FEvolutionLinearSystem;
 	}  // End namespace Softs
 }  // End namespace Chaos
diff --git a/Engine/Source/Runtime/Experimental/Chaos/Public/Chaos/SoftsEvolution.h b/Engine/Source/Runtime/Experimental/Chaos/Public/Chaos/SoftsEvolution.h
--- a/Engine/Source/Runtime/Experimental/Chaos/Public/Chaos/SoftsEvolution.h
+++ b/Engine/Source/Runtime/Experimental/Chaos/Public/Chaos/SoftsEvolution.h
@@ -5,14 +5,24 @@
 #include "Chaos/ArrayCollection.h"
 #include "Chaos/PBDSoftsEvolutionFwd.h"
 #include "Chaos/PBDSoftsSolverParticles.h"
+#include "Chaos/SoftsEvolutionLinearSystem.h"
 #include "Chaos/SoftsSolverParticlesRange.h"
 #include "Chaos/SoftsSolverCollisionParticles.h"
 #include "Chaos/SoftsSolverCollisionParticlesRange.h"
 #include "Chaos/VelocityField.h"
+#include "Misc/EnumClassFlags.h"
 
 namespace Chaos::Softs
 {
 
+enum struct ESolverMode : uint8
+{
+	None = 0,
+	PBD = 1 << 0,
+	ForceBased = 1 << 1
+};
+ENUM_CLASS_FLAGS(ESolverMode);
+
 /**
  * Solver can contain multiple "Groups". Groups do not interact with each other. 
  * They may be in different spaces. They may be solved in parallel, completely independently of each other. 
@@ -62,6 +72,8 @@ public:
 	const FSolverParticlesRange& GetSoftBodyParticles(int32 SoftBodyId) const { return SoftBodies.ParticleRanges[SoftBodyId]; }
 	const TArray<int32>& GetGroupSoftBodies(uint32 GroupId) const { return Groups.SoftBodies[GroupId]; }
 	const TSet<int32>& GetGroupActiveSoftBodies(uint32 GroupId) const { return Groups.ActiveSoftBodies[GroupId]; }
+	int32 GetLastLinearSolveIterations(int32 SoftBodyId) const { return SoftBodies.LinearSystems[SoftBodyId].GetLastSolveIterations(); }
+	FSolverReal GetLastLinearSolveError(int32 SoftBodyId) const { return SoftBodies.LinearSystems[SoftBodyId].GetLastSolveError(); }
 
 	/**
 	 * Add Collision particle range to a group.
@@ -83,74 +95,104 @@ public:
 	void SetKinematicUpdateFunction(KinematicUpdateFunc Func) { KinematicUpdate = Func; }
 	void SetCollisionKinematicUpdateFunction(CollisionKinematicUpdateFunc Func) { CollisionKinematicUpdate = Func; }
 
-	/** Soft Body Rules. Add Ranges to allocate space for your rules. You will get back an ArrayView where you can then set the rules. */
-	typedef TFunction<void(const FSolverParticlesRange&, const FSolverReal)> PreSubstepParallelInitFunc;
-	typedef TFunction<void(FSolverParticlesRange&, const FSolverReal)> ExternalForceRuleFunc;
-	typedef TFunction<void(const FSolverParticlesRange&, const FSolverReal)> ConstraintParallelInitFunc;
-	typedef TFunction<void(FSolverParticlesRange&, const FSolverReal)> ConstraintRuleFunc;
-	typedef TFunction<void(FSolverParticlesRange&, const FSolverReal, const TArray<FSolverCollisionParticlesRange>&)> CollisionConstraintRuleFunc;
+	/** Soft Body Rules. */
+	typedef TFunction<void(const FSolverParticlesRange&, const FSolverReal Dt, const ESolverMode)> ParallelInitFunc;
+	typedef TFunction<void(FSolverParticlesRange&, const FSolverReal Dt, const ESolverMode)> ConstraintRuleFunc;
+	typedef TFunction<void(FSolverParticlesRange&, const FSolverReal Dt)> PBDConstraintRuleFunc;
+	typedef TFunction<void(FSolverParticlesRange&, const FSolverReal Dt, const TArray<FSolverCollisionParticlesRange>&)> PBDCollisionConstraintRuleFunc;
+	typedef TFunction<void(const FSolverParticlesRange&, const FSolverReal Dt, FEvolutionLinearSystem&)> UpdateLinearSystemFunc;
+	typedef TFunction<void(const FSolverParticlesRange&, const FSolverReal Dt, const TArray<FSolverCollisionParticlesRange>&, FEvolutionLinearSystem&)> UpdateLinearSystemCollisionsFunc;
+
+	/* Add Ranges to allocate space for your rules. */
 	// Warning: Rules are allocated into shared buffers. The ArrayViews may go stale if you allocate ANY new rules.
 	// Allocate your batch of rules, then get your array views.	
+
+	// Presubstep init methods (always run at beginning of substep)
 	void AllocatePreSubstepParallelInitRange(int32 SoftBodyId, int32 NumRules)
 	{
-		return AllocateRulesRange(SoftBodyId, NumRules, ConstParticleRules, SoftBodies.PreSubstepParallelInits);
+		return AllocateRulesRange(SoftBodyId, NumRules, ParallelInitRules, SoftBodies.PreSubstepParallelInits);
 	}
-	void AllocateExternalForceRulesRange(int32 SoftBodyId, int32 NumRules)
+	// PBD Rules that apply external forces (only run if doing PBD)
+	void AllocatePBDExternalForceRulesRange(int32 SoftBodyId, int32 NumRules)
 	{
-		return AllocateRulesRange(SoftBodyId, NumRules, ParticleRules, SoftBodies.ExternalForceRules);
+		return AllocateRulesRange(SoftBodyId, NumRules, PBDConstraintRules, SoftBodies.PBDExternalForceRules);
 	}
-	void AllocateConstraintParallelInitsRange(int32 SoftBodyId, int32 NumRules)
+	// Post initial guess init methods (always run after kinematic and initial guess update, before any solving.) 
+	void AllocatePostInitialGuessParallelInitRange(int32 SoftBodyId, int32 NumRules)
 	{
-		return AllocateRulesRange(SoftBodyId, NumRules, ConstParticleRules, SoftBodies.ConstraintParallelInits);
+		return AllocateRulesRange(SoftBodyId, NumRules, ParallelInitRules, SoftBodies.PostInitialGuessParallelInits);
 	}
+	// Rules that run once per substep after all initial guess and initialization is done.
 	void AllocatePreSubstepConstraintRulesRange(int32 SoftBodyId, int32 NumRules)
 	{
-		return AllocateRulesRange(SoftBodyId, NumRules, ParticleRules, SoftBodies.PreSubstepConstraintRules);
+		return AllocateRulesRange(SoftBodyId, NumRules, ConstraintRules, SoftBodies.PreSubstepConstraintRules);
+	}
+	// Normal per-iteration PBD rules (only run if doing PBD)
+	void AllocatePerIterationPBDConstraintRulesRange(int32 SoftBodyId, int32 NumRules)
+	{
+		return AllocateRulesRange(SoftBodyId, NumRules, PBDConstraintRules, SoftBodies.PerIterationPBDConstraintRules);
 	}
-	void AllocatePerIterationConstraintRulesRange(int32 SoftBodyId, int32 NumRules)
+	// Collision per-iteration PBD rules (only run if doing PBD)
+	void AllocatePerIterationCollisionPBDConstraintRulesRange(int32 SoftBodyId, int32 NumRules)
 	{
-		return AllocateRulesRange(SoftBodyId, NumRules, ParticleRules, SoftBodies.PerIterationConstraintRules);
+		return AllocateRulesRange(SoftBodyId, NumRules, PBDCollisionConstraintRules, SoftBodies.PerIterationCollisionPBDConstraintRules);
 	}
-	void AllocatePerIterationCollisionConstraintRulesRange(int32 SoftBodyId, int32 NumRules)
+	// Normal per-iteration PBD rules that run after collisions (only run if doing PBD)
+	void AllocatePerIterationPostCollisionsPBDConstraintRulesRange(int32 SoftBodyId, int32 NumRules)
 	{
-		return AllocateRulesRange(SoftBodyId, NumRules, CollisionRules, SoftBodies.PerIterationCollisionConstraintRules);
+		return AllocateRulesRange(SoftBodyId, NumRules, PBDConstraintRules, SoftBodies.PerIterationPostCollisionsPBDConstraintRules);
 	}
-	void AllocatePerIterationPostCollisionsConstraintRulesRange(int32 SoftBodyId, int32 NumRules)
+	// Linear system rules (only run if doing ForceBased)
+	void AllocateUpdateLinearSystemRulesRange(int32 SoftBodyId, int32 NumRules)
 	{
-		return AllocateRulesRange(SoftBodyId, NumRules, ParticleRules, SoftBodies.PerIterationPostCollisionsConstraintRules);
+		return AllocateRulesRange(SoftBodyId, NumRules, UpdateLinearSystemRules, SoftBodies.UpdateLinearSystemRules);
 	}
+	// Linear system collision rules (only run if doing ForceBased)
+	void AllocateUpdateLinearSystemCollisionsRulesRange(int32 SoftBodyId, int32 NumRules)
+	{
+		return AllocateRulesRange(SoftBodyId, NumRules, UpdateLinearSystemCollisionsRules, SoftBodies.UpdateLinearSystemCollisionsRules);
+	}
+	// Post substep rules (always run at end of substep)
 	void AllocatePostSubstepConstraintRulesRange(int32 SoftBodyId, int32 NumRules)
 	{
-		return AllocateRulesRange(SoftBodyId, NumRules, ParticleRules, SoftBodies.PostSubstepConstraintRules);
+		return AllocateRulesRange(SoftBodyId, NumRules, ConstraintRules, SoftBodies.PostSubstepConstraintRules);
 	}
 
-	TArrayView<PreSubstepParallelInitFunc> GetPreSubstepParallelInitRange(int32 SoftBodyId)
+	TArrayView<ParallelInitFunc> GetPreSubstepParallelInitRange(int32 SoftBodyId)
 	{
 		return GetRulesRange(SoftBodyId, SoftBodies.PreSubstepParallelInits);
 	}
-	TArrayView<ExternalForceRuleFunc> GetExternalForceRulesRange(int32 SoftBodyId)
+	TArrayView<PBDConstraintRuleFunc> GetPBDExternalForceRulesRange(int32 SoftBodyId)
 	{
-		return GetRulesRange(SoftBodyId, SoftBodies.ExternalForceRules);
+		return GetRulesRange(SoftBodyId, SoftBodies.PBDExternalForceRules);
 	}
-	TArrayView<ConstraintParallelInitFunc> GetConstraintParallelInitsRange(int32 SoftBodyId)
+	TArrayView<ParallelInitFunc> GetPostInitialGuessParallelInitRange(int32 SoftBodyId)
 	{
-		return GetRulesRange(SoftBodyId, SoftBodies.ConstraintParallelInits);
+		return GetRulesRange(SoftBodyId, SoftBodies.PostInitialGuessParallelInits);
 	}
 	TArrayView<ConstraintRuleFunc> GetPreSubstepConstraintRulesRange(int32 SoftBodyId)
 	{
 		return GetRulesRange(SoftBodyId, SoftBodies.PreSubstepConstraintRules);
 	}
-	TArrayView<ConstraintRuleFunc> GetPerIterationConstraintRulesRange(int32 SoftBodyId)
+	TArrayView<PBDConstraintRuleFunc> GetPerIterationPBDConstraintRulesRange(int32 SoftBodyId)
+	{
+		return GetRulesRange(SoftBodyId, SoftBodies.PerIterationPBDConstraintRules);
+	}
+	TArrayView<PBDCollisionConstraintRuleFunc> GetPerIterationCollisionPBDConstraintRulesRange(int32 SoftBodyId)
 	{
-		return GetRulesRange(SoftBodyId, SoftBodies.PerIterationConstraintRules);
+		return GetRulesRange(SoftBodyId, SoftBodies.PerIterationCollisionPBDConstraintRules);
 	}
-	TArrayView<CollisionConstraintRuleFunc> GetPerIterationCollisionConstraintRulesRange(int32 SoftBodyId)
+	TArrayView<PBDConstraintRuleFunc> GetPerIterationPostCollisionsPBDConstraintRulesRange(int32 SoftBodyId)
 	{
-		return GetRulesRange(SoftBodyId, SoftBodies.PerIterationCollisionConstraintRules);
+		return GetRulesRange(SoftBodyId, SoftBodies.PerIterationPostCollisionsPBDConstraintRules);
 	}
-	TArrayView<ConstraintRuleFunc> GetPerIterationPostCollisionsConstraintRulesRange(int32 SoftBodyId)
+	TArrayView<UpdateLinearSystemFunc> GetUpdateLinearSystemRulesRange(int32 SoftBodyId)
 	{
-		return GetRulesRange(SoftBodyId, SoftBodies.PerIterationPostCollisionsConstraintRules);
+		return GetRulesRange(SoftBodyId, SoftBodies.UpdateLinearSystemRules);
+	}
+	TArrayView<UpdateLinearSystemCollisionsFunc> GetUpdateLinearSystemCollisionsRulesRange(int32 SoftBodyId)
+	{
+		return GetRulesRange(SoftBodyId, SoftBodies.UpdateLinearSystemCollisionsRules);
 	}
 	TArrayView<ConstraintRuleFunc> GetPostSubstepConstraintRulesRange(int32 SoftBodyId)
 	{
@@ -183,7 +225,7 @@ private:
 
 		bool IsValid() const
 		{
-			return Array && Offset >= 0 && Offset + RangeSize <= Array->Num();
+			return RangeSize == 0 || (Array && Offset >= 0 && Offset + RangeSize <= Array->Num());
 		}
 
 		TConstArrayView<ElementType> GetConstArrayView() const
@@ -217,13 +259,17 @@ private:
 			TArrayCollection::AddArray(&GlobalDampings);
 			TArrayCollection::AddArray(&LocalDampings);
 			TArrayCollection::AddArray(&UsePerParticleDamping);
+			TArrayCollection::AddArray(&LinearSystems);
+
 			TArrayCollection::AddArray(&PreSubstepParallelInits);
-			TArrayCollection::AddArray(&ExternalForceRules);
-			TArrayCollection::AddArray(&ConstraintParallelInits);
+			TArrayCollection::AddArray(&PBDExternalForceRules);
+			TArrayCollection::AddArray(&PostInitialGuessParallelInits);
 			TArrayCollection::AddArray(&PreSubstepConstraintRules);
-			TArrayCollection::AddArray(&PerIterationConstraintRules);
-			TArrayCollection::AddArray(&PerIterationCollisionConstraintRules);
-			TArrayCollection::AddArray(&PerIterationPostCollisionsConstraintRules);
+			TArrayCollection::AddArray(&PerIterationPBDConstraintRules);
+			TArrayCollection::AddArray(&PerIterationCollisionPBDConstraintRules);
+			TArrayCollection::AddArray(&PerIterationPostCollisionsPBDConstraintRules);
+			TArrayCollection::AddArray(&UpdateLinearSystemRules);
+			TArrayCollection::AddArray(&UpdateLinearSystemCollisionsRules);
 			TArrayCollection::AddArray(&PostSubstepConstraintRules);
 		}
 
@@ -245,13 +291,17 @@ private:
 		TArrayCollectionArray<FSolverReal> GlobalDampings;
 		TArrayCollectionArray<FSolverReal> LocalDampings;
 		TArrayCollectionArray<bool> UsePerParticleDamping;
-		TArrayCollectionArray<TArrayRange<PreSubstepParallelInitFunc>> PreSubstepParallelInits;
-		TArrayCollectionArray<TArrayRange<ExternalForceRuleFunc>> ExternalForceRules;
-		TArrayCollectionArray<TArrayRange<ConstraintParallelInitFunc>> ConstraintParallelInits;
+		TArrayCollectionArray<FEvolutionLinearSystem> LinearSystems;
+
+		TArrayCollectionArray<TArrayRange<ParallelInitFunc>> PreSubstepParallelInits;
+		TArrayCollectionArray<TArrayRange<PBDConstraintRuleFunc>> PBDExternalForceRules;
+		TArrayCollectionArray<TArrayRange<ParallelInitFunc>> PostInitialGuessParallelInits;
 		TArrayCollectionArray<TArrayRange<ConstraintRuleFunc>> PreSubstepConstraintRules;
-		TArrayCollectionArray<TArrayRange<ConstraintRuleFunc>> PerIterationConstraintRules;
-		TArrayCollectionArray<TArrayRange<CollisionConstraintRuleFunc>> PerIterationCollisionConstraintRules;
-		TArrayCollectionArray<TArrayRange<ConstraintRuleFunc>> PerIterationPostCollisionsConstraintRules;
+		TArrayCollectionArray<TArrayRange<PBDConstraintRuleFunc>> PerIterationPBDConstraintRules;
+		TArrayCollectionArray<TArrayRange<PBDCollisionConstraintRuleFunc>> PerIterationCollisionPBDConstraintRules;
+		TArrayCollectionArray<TArrayRange<PBDConstraintRuleFunc>> PerIterationPostCollisionsPBDConstraintRules;
+		TArrayCollectionArray<TArrayRange<UpdateLinearSystemFunc>> UpdateLinearSystemRules;
+		TArrayCollectionArray<TArrayRange<UpdateLinearSystemCollisionsFunc>> UpdateLinearSystemCollisionsRules;
 		TArrayCollectionArray<TArrayRange<ConstraintRuleFunc>> PostSubstepConstraintRules;
 	};
 
@@ -316,6 +366,38 @@ private:
 		TArrayCollectionArray<TSet<int32>> ActiveCollisionParticleRanges;
 	};
 
+	// Wrapper around FEvolutionLinearSystemSolverParameters that knows how to read a property collection
+	struct FLinearSystemParameters : public FEvolutionLinearSystemSolverParameters
+	{
+		typedef FEvolutionLinearSystemSolverParameters Base;
+
+		FLinearSystemParameters()
+			: Base()
+		{}
+
+		FLinearSystemParameters(const FCollectionPropertyConstFacade& PropertyCollection, bool bInXPBDInitialGuess)
+			: Base(GetDoQuasistatics(PropertyCollection, false)
+				, bInXPBDInitialGuess
+				, GetMaxNumCGIterations(PropertyCollection, DefaultMaxNumCGIterations)
+				, GetCGResidualTolerance(PropertyCollection, DefaultCGTolerance)
+				, GetCheckCGResidual(PropertyCollection, bDefaultCheckCGResidual))
+		{}
+
+		void SetProperties(const FCollectionPropertyConstFacade& PropertyCollection, bool bInXPBDInitialGuess)
+		{
+			bXPBDInitialGuess = bInXPBDInitialGuess;
+			bDoQuasistatics = GetDoQuasistatics(PropertyCollection, false);
+			MaxNumCGIterations = GetMaxNumCGIterations(PropertyCollection, DefaultMaxNumCGIterations);
+			CGResidualTolerance = GetCGResidualTolerance(PropertyCollection, DefaultCGTolerance);
+			bCheckCGResidual = GetCheckCGResidual(PropertyCollection, bDefaultCheckCGResidual);
+		}
+
+		UE_CHAOS_DECLARE_INDEXLESS_PROPERTYCOLLECTION_NAME(DoQuasistatics, bool);
+		UE_CHAOS_DECLARE_INDEXLESS_PROPERTYCOLLECTION_NAME(MaxNumCGIterations, int32);
+		UE_CHAOS_DECLARE_INDEXLESS_PROPERTYCOLLECTION_NAME(CGResidualTolerance, float);
+		UE_CHAOS_DECLARE_INDEXLESS_PROPERTYCOLLECTION_NAME(CheckCGResidual, bool);
+	};
+
 	template<typename RuleFunc>
 	void AllocateRulesRange(int32 SoftBodyId, int32 NumRules, TArray<RuleFunc>& RuleArray, TArrayCollectionArray<TArrayRange<RuleFunc>>& RangeArray)
 	{
@@ -332,12 +414,15 @@ private:
 	void AdvanceOneTimeStepInternal(const FSolverReal Dt, const int32 TimeDependentNumIterations, uint32 GroupId);
 
 	// Solver data
-	int32 NumIterations;
+	FSolverReal Time;
+	bool bEnableForceBasedSolver = false;
 	int32 MaxNumIterations; // Used for time-dependent iteration counts
-	bool bDisableTimeDependentNumIterations;
+	int32 NumIterations; // PBD iterations
+	int32 NumNewtonIterations; // Implicit force-based solve
+	bool bDisableTimeDependentNumIterations = false;
 	bool bDoQuasistatics;
 	FSolverReal SolverFrequency; 
-	FSolverReal Time;
+	FLinearSystemParameters LinearSystemParameters; // Per-solver parameters that need to be passed to the linear system solver
 
 	// Per-Particle data
 	FSolverParticles Particles;
@@ -363,10 +448,20 @@ private:
 	CollisionKinematicUpdateFunc CollisionKinematicUpdate;
 
 	// Rules that run on ParticleRanges (SoftBodies will have views into these)
-	TArray<TFunction<void(const FSolverParticlesRange&, const FSolverReal)>> ConstParticleRules;
-	TArray<TFunction<void(FSolverParticlesRange&, const FSolverReal)>> ParticleRules;
+	TArray<ParallelInitFunc> ParallelInitRules;
+	TArray<ConstraintRuleFunc> ConstraintRules;
+	TArray<PBDConstraintRuleFunc> PBDConstraintRules;
+	TArray<PBDCollisionConstraintRuleFunc> PBDCollisionConstraintRules;
+	TArray<UpdateLinearSystemFunc> UpdateLinearSystemRules;
+	TArray<UpdateLinearSystemCollisionsFunc> UpdateLinearSystemCollisionsRules;
 
 	// Collision Rules
-	TArray<TFunction<void(FSolverParticlesRange&, const FSolverReal, const TArray<FSolverCollisionParticlesRange>&)>> CollisionRules;
+
+	UE_CHAOS_DECLARE_INDEXLESS_PROPERTYCOLLECTION_NAME(MaxNumIterations, int32);
+	UE_CHAOS_DECLARE_INDEXLESS_PROPERTYCOLLECTION_NAME(NumIterations, int32);
+	UE_CHAOS_DECLARE_INDEXLESS_PROPERTYCOLLECTION_NAME(DoQuasistatics, bool);
+	UE_CHAOS_DECLARE_INDEXLESS_PROPERTYCOLLECTION_NAME(SolverFrequency, float);
+	UE_CHAOS_DECLARE_INDEXLESS_PROPERTYCOLLECTION_NAME(EnableForceBasedSolver, bool);
+	UE_CHAOS_DECLARE_INDEXLESS_PROPERTYCOLLECTION_NAME(NumNewtonIterations, int32);
 };
 }
\ No newline at end of file
diff --git a/Engine/Source/Runtime/Experimental/Chaos/Public/Chaos/SoftsEvolutionLinearSystem.h b/Engine/Source/Runtime/Experimental/Chaos/Public/Chaos/SoftsEvolutionLinearSystem.h
new file mode 100644
--- /dev/null
+++ b/Engine/Source/Runtime/Experimental/Chaos/Public/Chaos/SoftsEvolutionLinearSystem.h
@@ -0,0 +1,120 @@
+// Copyright Epic Games, Inc. All Rights Reserved.
+#pragma once
+#include "Chaos/PBDSoftsEvolutionFwd.h"
+#include "Chaos/CollectionPropertyFacade.h"
+#include "Chaos/Math/BlockSparseLinearSystem.h"
+
+namespace Chaos::Softs
+{
+struct FEvolutionLinearSystemSolverParameters
+{
+	static constexpr int32 DefaultMaxNumCGIterations = 50;
+	static constexpr FSolverReal DefaultCGTolerance = (FSolverReal)1e-4;
+	static constexpr bool bDefaultCheckCGResidual = false;
+
+	FEvolutionLinearSystemSolverParameters() = default;
+	FEvolutionLinearSystemSolverParameters(
+		bool bInDoQuasistatics,
+		bool bInXPBDInitialGuess,
+		int32 InMaxNumCGIterations,
+		FSolverReal InCGResidualTolerance,
+		bool bInCheckResidual)
+		: bDoQuasistatics(bInDoQuasistatics)
+		, bXPBDInitialGuess(bInXPBDInitialGuess)
+		, MaxNumCGIterations(InMaxNumCGIterations)
+		, CGResidualTolerance(InCGResidualTolerance)
+		, bCheckCGResidual(bInCheckResidual)
+	{}
+	FEvolutionLinearSystemSolverParameters(const FEvolutionLinearSystemSolverParameters&) = default;
+	FEvolutionLinearSystemSolverParameters(FEvolutionLinearSystemSolverParameters&&) = default;
+	~FEvolutionLinearSystemSolverParameters() = default;
+	FEvolutionLinearSystemSolverParameters& operator=(const FEvolutionLinearSystemSolverParameters&) = default;
+	FEvolutionLinearSystemSolverParameters& operator=(FEvolutionLinearSystemSolverParameters&&) = default;
+
+	bool bDoQuasistatics = false;
+	bool bXPBDInitialGuess = false;
+	int32 MaxNumCGIterations = DefaultMaxNumCGIterations;
+	FSolverReal CGResidualTolerance = DefaultCGTolerance;
+	bool bCheckCGResidual = bDefaultCheckCGResidual;
+};
+
+// Assemble and solve the linear system generated by the force based solver in FEvolution.
+// Currently just supports using CG to solve the linear system generated by using newton's method with backward euler.
+// The matrix format and CG is controlled by BlockSparseLinearSystem. 
+class FEvolutionLinearSystem
+{
+public:
+	CHAOS_API FEvolutionLinearSystem(const FSolverParticlesRange& Particles);
+	FEvolutionLinearSystem() = default;
+	FEvolutionLinearSystem(const FEvolutionLinearSystem&) = delete;
+	FEvolutionLinearSystem(FEvolutionLinearSystem&& Other)
+		: Parameters(Other.Parameters)
+		, NumCompactIndices(Other.NumCompactIndices)
+		, CompactifiedIndices(MoveTemp(Other.CompactifiedIndices))
+		, RHS(MoveTemp(Other.RHS))
+		, Matrix(MoveTemp(Other.Matrix))
+		, bDfDxTimesVTerm(Other.bDfDxTimesVTerm)
+	{}
+	FEvolutionLinearSystem& operator=(const FEvolutionLinearSystem&) = delete;
+	FEvolutionLinearSystem& operator=(FEvolutionLinearSystem&& Other)
+	{
+		Parameters = Other.Parameters;
+		NumCompactIndices = Other.NumCompactIndices;
+		CompactifiedIndices = MoveTemp(Other.CompactifiedIndices);
+		RHS = MoveTemp(Other.RHS);
+		Matrix = MoveTemp(Other.Matrix);
+		bDfDxTimesVTerm = Other.bDfDxTimesVTerm;
+		return *this;
+	}
+
+	~FEvolutionLinearSystem() = default;
+
+	CHAOS_API void Init(const FSolverParticlesRange& Particles, const FSolverReal Dt, bool bInFirstNewtonIteration, const FEvolutionLinearSystemSolverParameters& Params);
+
+	void ReserveForParallelAdd(const int32 NumDiagonalEntries, const int32 NumOffDiagonalEntries)
+	{
+		Matrix.ReserveForParallelAdd(NumDiagonalEntries, NumOffDiagonalEntries);
+	}
+
+	CHAOS_API void AddForce(const FSolverParticlesRange& Particles, const FSolverVec3& Force, int32 ParticleIndex, const FSolverReal Dt);
+
+	// Df1Dx2 is Derivative of F for ParticleIndex1 with respect to X of Particle2.
+	// Components of Df1Dx2 are Df1Dx2[j][i] = DF1_i / DX2_j
+	// This is consistent with UnrealMath which assumes row-major matrices and treats positions as row vectors (left multiply)
+	// 
+	// We're assuming all forces are symmetric. This will add Df1Dfx2 and Df2Dfx1 = Df1Dfx2^T (and same for Df2Dv1 = Df1Dv2^T) when ParticleIndex1 != ParticleIndex2
+	CHAOS_API void AddSymmetricForceDerivative(const FSolverParticlesRange& Particles, const FSolverMatrix33* const Df1Dx2, const FSolverMatrix33* const Df1Dv2, int32 ParticleIndex1, int32 ParticleIndex2, const FSolverReal Dt);
+
+	bool Solve(FSolverParticlesRange& Particles, const FSolverReal Dt);
+
+	int32 GetLastSolveIterations() const 
+	{
+		return LastSolveIterations;
+	}
+
+	FSolverReal GetLastSolveError() const
+	{
+		return LastSolveError;
+	}
+
+	bool RequiresSPDForceDerivatives() const
+	{
+		// we're currently just doing CG to solve the primal equations, so we need to be symmetric positive definite
+		return true;
+	}
+private:
+	void CalculateCompactIndices(const FSolverParticlesRange& Particles);
+
+	const FEvolutionLinearSystemSolverParameters* Parameters = nullptr;
+	int32 NumCompactIndices;
+	TArray<int32> CompactifiedIndices; // ParticleRangeIndex -> Dense array with kinematic points removed.
+	TArray<FSolverVec3> RHS;
+	TBlockSparseSymmetricLinearSystem<FSolverReal,3> Matrix;
+	bool bDfDxTimesVTerm = false;
+
+	// Debug reporting
+	int32 LastSolveIterations = 0;
+	FSolverReal LastSolveError = (FSolverReal)0.;
+};
+
+}  // End namespace Chaos::Softs
diff --git a/Engine/Source/Runtime/Experimental/Chaos/Public/Chaos/SoftsExternalForces.h b/Engine/Source/Runtime/Experimental/Chaos/Public/Chaos/SoftsExternalForces.h
--- a/Engine/Source/Runtime/Experimental/Chaos/Public/Chaos/SoftsExternalForces.h
+++ b/Engine/Source/Runtime/Experimental/Chaos/Public/Chaos/SoftsExternalForces.h
@@ -3,6 +3,7 @@
 
 #include "Chaos/PBDSoftsEvolutionFwd.h"
 #include "Chaos/PBDSoftsSolverParticles.h"
+#include "Chaos/SoftsEvolutionLinearSystem.h"
 #include "Chaos/SoftsSolverParticlesRange.h"
 #include "Chaos/CollectionPropertyFacade.h"
 
@@ -62,6 +63,45 @@ public:
 		}
 	}
 
+	void UpdateLinearSystem(const FSolverParticlesRange& Particles, const FSolverReal Dt, FEvolutionLinearSystem& LinearSystem) const
+	{
+		const FSolverReal* const InvM = Particles.GetInvM().GetData();
+		const FSolverReal* const M = Particles.GetM().GetData();
+		const FSolverVec3* const X = Particles.XArray().GetData();
+		const FSolverVec3* const V = Particles.GetV().GetData();
+		const FSolverVec3* const N = Particles.GetConstArrayView(Normals).GetData();
+
+		const bool bHasFictitiousForces = !FictitiousAngularDisplacement.IsNearlyZero();
+		const FSolverVec3 W = FictitiousAngularDisplacement / Dt;
+
+		for (int32 Index = 0; Index < Particles.GetRangeSize(); ++Index)
+		{
+			if (InvM[Index] != (FSolverReal)0.)
+			{
+				FSolverVec3 Force = Gravity * M[Index];
+
+				if (bHasFictitiousForces)
+				{
+					// Centrifugal force (*InvM to get acceleration)
+					Force -= FSolverVec3::CrossProduct(W, FSolverVec3::CrossProduct(W, X[Index] - ReferenceSpaceLocation)) * M[Index];
+				}
+				if (bUsePointBasedWindModel)
+				{
+					const FSolverVec3 VelocityDelta = PointBasedWind - V[Index];
+					FSolverVec3 Direction = VelocityDelta;
+					if (Direction.Normalize())
+					{
+						// Scale by angle
+						const FReal DirectionDot = FVec3::DotProduct(Direction, N[Index]);
+						const FReal ScaleFactor = FMath::Min(1.f, FMath::Abs(DirectionDot) * LegacyWindAdaptation);
+						Force += VelocityDelta * ScaleFactor * M[Index];
+					}
+				}
+				LinearSystem.AddForce(Particles, Force, Index, Dt);
+			}
+		}
+	}
+
 	bool UsePointBasedWindModel() const { return bUsePointBasedWindModel; }
 	const FSolverVec3& GetGravity() const { return Gravity; }
 
diff --git a/Engine/Source/Runtime/Experimental/Chaos/Public/Chaos/SoftsSpring.h b/Engine/Source/Runtime/Experimental/Chaos/Public/Chaos/SoftsSpring.h
new file mode 100644
--- /dev/null
+++ b/Engine/Source/Runtime/Experimental/Chaos/Public/Chaos/SoftsSpring.h
@@ -0,0 +1,116 @@
+// Copyright Epic Games, Inc. All Rights Reserved.
+#pragma once
+#include "Chaos/SoftsEvolutionLinearSystem.h"
+
+namespace Chaos::Softs
+{
+
+namespace Spring
+{
+
+template<typename SolverParticlesOrRange>
+FSolverVec3 GetXPBDSpringDelta(const SolverParticlesOrRange& Particles, const FSolverReal Dt,
+	const TVec2<int32>& Constraint, const FSolverReal RestLength, FSolverReal& Lambda,
+	const FSolverReal StiffnessValue, const FSolverReal MinStiffness, const FSolverReal DampingRatioValue)
+{
+	const int32 Index1 = Constraint[0];
+	const int32 Index2 = Constraint[1];
+
+	if (StiffnessValue < MinStiffness || (Particles.InvM(Index2) == (FSolverReal)0. && Particles.InvM(Index1) == (FSolverReal)0.))
+	{
+		return FSolverVec3((FSolverReal)0.);
+	}
+
+	const FSolverReal CombinedInvMass = Particles.InvM(Index2) + Particles.InvM(Index1);
+	const FSolverReal Damping = DampingRatioValue * 2.f * FMath::Sqrt(StiffnessValue / CombinedInvMass) * (RestLength > UE_SMALL_NUMBER ? (FSolverReal)1. / RestLength : (FSolverReal)1.);
+
+	const FSolverVec3& P1 = Particles.P(Index1);
+	const FSolverVec3& P2 = Particles.P(Index2);
+	FSolverVec3 Direction = P1 - P2;
+	const FSolverReal Distance = Direction.SafeNormalize();
+	const FSolverReal Offset = Distance - RestLength;
+
+	const FSolverVec3& X1 = Particles.X(Index1);
+	const FSolverVec3& X2 = Particles.X(Index2);
+
+	const FSolverVec3 RelativeVelocityTimesDt = P1 - X1 - P2 + X2;
+
+	const FSolverReal AlphaInv = StiffnessValue * Dt * Dt;
+	const FSolverReal BetaDt = Damping * Dt;
+
+	const FSolverReal DLambda = (AlphaInv * Offset - Lambda + BetaDt * FSolverVec3::DotProduct(Direction, RelativeVelocityTimesDt)) / ((AlphaInv + BetaDt) * CombinedInvMass + (FSolverReal)1.);
+	const FSolverVec3 Delta = DLambda * Direction;
+	Lambda += DLambda;
+
+	return Delta;
+}
+
+inline void UpdateSpringLinearSystem(const FSolverParticlesRange& Particles, const FSolverReal Dt,
+	const TVec2<int32>& Constraint, const FSolverReal RestLength,
+	const FSolverReal StiffnessValue, const FSolverReal MinStiffness, const FSolverReal DampingRatioValue,
+	FEvolutionLinearSystem& LinearSystem)
+{
+	const int32 Index1 = Constraint[0];
+	const int32 Index2 = Constraint[1];
+
+	if (StiffnessValue < MinStiffness || (Particles.InvM(Index2) == (FSolverReal)0. && Particles.InvM(Index1) == (FSolverReal)0.))
+	{
+		return;
+	}
+
+	const FSolverVec3& P1 = Particles.P(Index1);
+	const FSolverVec3& P2 = Particles.P(Index2);
+	FSolverVec3 Direction = P1 - P2;
+	const FSolverReal Distance = Direction.SafeNormalize();
+	if (Distance < UE_SMALL_NUMBER)
+	{
+		// We can't calculate a direction if distance is zero. Just skip
+		return;
+	}
+	const FSolverReal CombinedInvMass = Particles.InvM(Index2) + Particles.InvM(Index1);
+
+	const FSolverReal Damping = DampingRatioValue * 2.f * FMath::Sqrt(StiffnessValue / CombinedInvMass) * (RestLength > UE_SMALL_NUMBER ? (FSolverReal)1. / RestLength : (FSolverReal)1.);
+	const FSolverVec3 RelVel = Particles.V(Index1) - Particles.V(Index2);
+	const FSolverReal CDot = FSolverVec3::DotProduct(Direction, RelVel);
+
+	const FSolverReal Offset = Distance - RestLength; // C
+	// const FSolverVec3 GradC1 = Direction;
+	// const FSolverVec3 GradC2 = -GradC1;
+	const FSolverReal Scalar = -(StiffnessValue * Offset + Damping * CDot);
+	const FSolverVec3 Force1 = Scalar * Direction;
+	// const FSolverVec3 Force2 = -Force1
+
+	const FSolverMatrix33 DirDirT = FSolverMatrix33::OuterProduct(Direction, Direction); // = GradC1 GradC1^T
+	const FSolverMatrix33 Hess11 = (FSolverReal)1. / Distance * (FSolverMatrix33::Identity - DirDirT); // = Hess22
+	// Hess12 = Hess21 = -Hess11
+
+	// Df1Dx1 = -Stiffness * (DirDirT + Offset * Hess11) - Damping * CDot * Hess11
+	const FSolverReal ScalarModified = LinearSystem.RequiresSPDForceDerivatives() ? FMath::Min(Scalar, 0) : Scalar;
+	const FSolverMatrix33 Df1Dx1 = -StiffnessValue * DirDirT + ScalarModified * Hess11;
+	const FSolverMatrix33 Df1Dx2 = -Df1Dx1;
+	// Df2Dx2 = Df1Dx1
+	const FSolverMatrix33 Df1Dv1 = -Damping * DirDirT;
+	const FSolverMatrix33 Df1Dv2 = -Df1Dv1;
+	// Df2Dv2 = Df1Dv1
+
+	if (Particles.InvM(Index1) != (FSolverReal)0.)
+	{
+		LinearSystem.AddForce(Particles, Force1, Index1, Dt);
+		LinearSystem.AddSymmetricForceDerivative(Particles, &Df1Dx1, &Df1Dv1, Index1, Index1, Dt);
+		LinearSystem.AddSymmetricForceDerivative(Particles, &Df1Dx2, &Df1Dv2, Index1, Index2, Dt);
+		if (Particles.InvM(Index2) != (FSolverReal)0.)
+		{
+			LinearSystem.AddForce(Particles, -Force1, Index2, Dt);
+			LinearSystem.AddSymmetricForceDerivative(Particles, &Df1Dx1, &Df1Dv1, Index2, Index2, Dt);
+		}
+	}
+	else
+	{
+		check(Particles.InvM(Index2) != (FSolverReal)0.);
+		LinearSystem.AddForce(Particles, -Force1, Index2, Dt);
+		LinearSystem.AddSymmetricForceDerivative(Particles, &Df1Dx1, &Df1Dv1, Index2, Index2, Dt);
+		LinearSystem.AddSymmetricForceDerivative(Particles, &Df1Dx2, &Df1Dv2, Index2, Index1, Dt);
+	}
+}
+}
+}
diff --git a/Engine/Source/Runtime/Experimental/Chaos/Public/Chaos/XPBDSpringConstraints.h b/Engine/Source/Runtime/Experimental/Chaos/Public/Chaos/XPBDSpringConstraints.h
--- a/Engine/Source/Runtime/Experimental/Chaos/Public/Chaos/XPBDSpringConstraints.h
+++ b/Engine/Source/Runtime/Experimental/Chaos/Public/Chaos/XPBDSpringConstraints.h
@@ -5,23 +5,24 @@
 #include "ChaosStats.h"
 #include "PBDWeightMap.h"
 #include "Chaos/CollectionPropertyFacade.h"
+#include "Chaos/SoftsSpring.h"
 
 namespace Chaos::Softs
 {
 
 // Stiffness is in kg cm /s^2
 UE_DEPRECATED(5.2, "Use FXPBDSpringConstraints::MinStiffness instead.")
-static const FSolverReal XPBDSpringMinStiffness = (FSolverReal)1e-4; // Stiffness below this will be considered 0 since all of our calculations are actually based on 1 / stiffness.
+static const FSolverReal XPBDSpringMinStiffness = (FSolverReal)UE_SMALL_NUMBER; // Stiffness below this will be considered 0 
 UE_DEPRECATED(5.2, "Use FXPBDSpringConstraints::MaxStiffness instead.")
-static const FSolverReal XPBDSpringMaxStiffness = (FSolverReal)1e7;
+static const FSolverReal XPBDSpringMaxStiffness = (FSolverReal)1e9;
 
 class FXPBDSpringConstraints : public FPBDSpringConstraintsBase
 {
 	typedef FPBDSpringConstraintsBase Base;
 
 public:
-	static constexpr FSolverReal MinStiffness = (FSolverReal)1e-4; // Stiffness below this will be considered 0 since all of our calculations are actually based on 1 / stiffness.
-	static constexpr FSolverReal MaxStiffness = (FSolverReal)1e7;
+	static constexpr FSolverReal MinStiffness = (FSolverReal)UE_SMALL_NUMBER; // Stiffness below this will be considered 0 
+	static constexpr FSolverReal MaxStiffness = (FSolverReal)1e9;
 	static constexpr FSolverReal MinDampingRatio = (FSolverReal)0.;
 	static constexpr FSolverReal MaxDampingRatio = (FSolverReal)1000.;
 
@@ -104,13 +105,18 @@ public:
 			ParticleOffset,
 			ParticleCount)
 	{
-		Lambdas.Init((FSolverReal)0., Constraints.Num());
+		Lambdas.Reset();
+		Lambdas.SetNumZeroed(Constraints.Num());
 		InitColor(Particles);
 	}
 
 	virtual ~FXPBDSpringConstraints() override {}
 
-	void Init() const { for (FSolverReal& Lambda : Lambdas) { Lambda = (FSolverReal)0.; } }
+	void Init() const 
+	{
+		Lambdas.Reset();
+		Lambdas.SetNumZeroed(Constraints.Num());
+	}
 
 	// Update stiffness values
 	void SetProperties(const FSolverVec2& InStiffness, const FSolverVec2& InDampingRatio = FSolverVec2::ZeroVector)
@@ -129,55 +135,20 @@ public:
 	template<typename SolverParticlesOrRange>
 	CHAOS_API void Apply(SolverParticlesOrRange& Particles, const FSolverReal Dt) const;
 
+	CHAOS_API void UpdateLinearSystem(const FSolverParticlesRange& Particles, const FSolverReal Dt, FEvolutionLinearSystem& LinearSystem) const;
+
 	const TArray<int32>& GetConstraintsPerColorStartIndex() const { return ConstraintsPerColorStartIndex; }
 
 private:
 	template<typename SolverParticlesOrRange>
 	CHAOS_API void InitColor(const SolverParticlesOrRange& InParticles);
 	template<typename SolverParticlesOrRange>
-	void ApplyHelper(SolverParticlesOrRange& Particles, const FSolverReal Dt, const int32 ConstraintIndex, const FSolverReal StiffnessValue, const FSolverReal DampingRatioValue) const;
+	void ApplyHelper(SolverParticlesOrRange& Particles, const FSolverReal Dt, const int32 ConstraintIndex, const FSolverReal StiffnessValue, const FSolverReal DampingRatioValue) const; 
 
 	template<typename SolverParticlesOrRange>
 	FSolverVec3 GetDelta(const SolverParticlesOrRange& Particles, const FSolverReal Dt, const int32 ConstraintIndex, const FSolverReal StiffnessValue, const FSolverReal DampingRatioValue) const
 	{
-		const TVec2<int32>& Constraint = Constraints[ConstraintIndex];
-
-		const int32 i1 = Constraint[0];
-		const int32 i2 = Constraint[1];
-
-		if (StiffnessValue < MinStiffness || (Particles.InvM(i2) == (FSolverReal)0. && Particles.InvM(i1) == (FSolverReal)0.))
-		{
-			return FSolverVec3((FSolverReal)0.);
-		}
-
-		const FSolverReal CombinedInvMass = Particles.InvM(i2) + Particles.InvM(i1);
-
-		// This scale factor makes things more resolution independent.
-		const FSolverReal FinalStiffnessValue = Dists[ConstraintIndex] < UE_SMALL_NUMBER ? StiffnessValue : StiffnessValue / Dists[ConstraintIndex];
-
-		const FSolverReal Damping = DampingRatioValue * 2.f * FMath::Sqrt(FinalStiffnessValue / CombinedInvMass);
-
-		const FSolverVec3& P1 = Particles.P(i1);
-		const FSolverVec3& P2 = Particles.P(i2);
-		FSolverVec3 Direction = P1 - P2;
-		const FSolverReal Distance = Direction.SafeNormalize();
-		const FSolverReal Offset = Distance - Dists[ConstraintIndex];
-
-		const FSolverVec3& X1 = Particles.X(i1);
-		const FSolverVec3& X2 = Particles.X(i2);
-
-		const FSolverVec3 RelativeVelocityTimesDt = P1 - X1 - P2 + X2;
-
-
-		FSolverReal& Lambda = Lambdas[ConstraintIndex];
-		const FSolverReal Alpha = (FSolverReal)1.f / (FinalStiffnessValue * Dt * Dt);
-		const FSolverReal Gamma = Alpha * Damping * Dt;
-
-		const FSolverReal DLambda = (Offset - Alpha * Lambda + Gamma * FSolverVec3::DotProduct(Direction, RelativeVelocityTimesDt)) / (((FSolverReal)1.f + Gamma) * CombinedInvMass + Alpha);
-		const FSolverVec3 Delta = DLambda * Direction;
-		Lambda += DLambda;
-
-		return Delta;
+		return Spring::GetXPBDSpringDelta(Particles, Dt, Constraints[ConstraintIndex], Dists[ConstraintIndex], Lambdas[ConstraintIndex], StiffnessValue, MinStiffness, DampingRatioValue);
 	}
 
 protected:
diff --git a/Engine/Source/Runtime/Experimental/ChaosCore/Public/Chaos/Matrix.h b/Engine/Source/Runtime/Experimental/ChaosCore/Public/Chaos/Matrix.h
--- a/Engine/Source/Runtime/Experimental/ChaosCore/Public/Chaos/Matrix.h
+++ b/Engine/Source/Runtime/Experimental/ChaosCore/Public/Chaos/Matrix.h
@@ -795,6 +795,20 @@ namespace Chaos
 				&& (FMath::Abs(Other.M[2][2] - M[2][2]) <= Tolerance);
 		}
 
+		// M[i][j] = x[i] * y[j]
+		static PMatrix<FRealDouble, 3, 3> OuterProduct(const TVector<FRealDouble, 3>& X, const TVector<FRealDouble, 3>& Y)
+		{
+			return PMatrix<FRealDouble, 3, 3>(
+				X[0] * Y[0],
+				X[1] * Y[0],
+				X[2] * Y[0],
+				X[0] * Y[1],
+				X[1] * Y[1],
+				X[2] * Y[1],
+				X[0] * Y[2],
+				X[1] * Y[2],
+				X[2] * Y[2]);
+		}
 
 		static const PMatrix<FRealDouble, 3, 3> Zero;
 		static const PMatrix<FRealDouble, 3, 3> Identity;
@@ -1222,6 +1236,21 @@ namespace Chaos
 				&& (FMath::Abs(Other.M[2][2] - M[2][2]) <= Tolerance);
 		}
 
+		// M[i][j] = x[i] * y[j]
+		static PMatrix<FRealSingle, 3, 3> OuterProduct(const TVector<FRealSingle, 3>& X, const TVector<FRealSingle, 3>& Y)
+		{
+			return PMatrix<FRealSingle, 3, 3>(
+				X[0] * Y[0],
+				X[1] * Y[0],
+				X[2] * Y[0],
+				X[0] * Y[1],
+				X[1] * Y[1],
+				X[2] * Y[1],
+				X[0] * Y[2],
+				X[1] * Y[2],
+				X[2] * Y[2]);
+		}
+
 
 		static const PMatrix<FRealSingle, 3, 3> Zero;
 		static const PMatrix<FRealSingle, 3, 3> Identity;
