Chaos Cloth: Add "NumSelfCollisionSubsteps". This control allows you to skip doing self collision detection for some subframes, and is mostly useful with high res cloth where lots of subframes are necessary.

Also fixed some bugs with NumIterations in the new SoftsEvolution (it was always clamping it to 1 and not reporting the correct number in debug visualization).

diff --git a/Engine/Plugins/Experimental/ChaosCloth/Source/Chaos/Private/ChaosCloth/ChaosClothConstraints.cpp b/Engine/Plugins/Experimental/ChaosCloth/Source/Chaos/Private/ChaosCloth/ChaosClothConstraints.cpp
--- a/Engine/Plugins/Experimental/ChaosCloth/Source/Chaos/Private/ChaosCloth/ChaosClothConstraints.cpp
+++ b/Engine/Plugins/Experimental/ChaosCloth/Source/Chaos/Private/ChaosCloth/ChaosClothConstraints.cpp
@@ -1606,6 +1606,10 @@ void FClothConstraints::CreateForceBasedRules()
 		RuleCreator.AddPostInitialGuessParallelInitRule(
 			[this](const Softs::FSolverParticlesRange& Particles, const Softs::FSolverReal Dt, const Softs::ESolverMode SolverMode)
 		{
+			if (bSkipSelfCollisionInit)
+			{
+				return;
+			}
 			// Thickness * 2 to account for collision radius for both particles
 			SelfCollisionInit->Init(Particles, SelfCollisionConstraints->GetThickness() * (Softs::FSolverReal)2.f);
 			SelfCollisionConstraints->Init(Particles, SelfCollisionInit->GetSpatialHash(), SelfCollisionInit->GetVertexGIAColors(), SelfCollisionInit->GetTriangleGIAColors());
@@ -1620,6 +1624,10 @@ void FClothConstraints::CreateForceBasedRules()
 		RuleCreator.AddPostInitialGuessParallelInitRule(
 			[this](const Softs::FSolverParticlesRange& Particles, const Softs::FSolverReal Dt, const Softs::ESolverMode SolverMode)
 		{
+			if (bSkipSelfCollisionInit)
+			{
+				return;
+			}
 			SelfCollisionSphereConstraints->Init(Particles);
 		});
 
@@ -1632,12 +1640,20 @@ void FClothConstraints::CreateForceBasedRules()
 		RuleCreator.AddPreSubstepConstraintRule(
 			[this](Softs::FSolverParticlesRange& Particles, const Softs::FSolverReal Dt, const Softs::ESolverMode SolverMode)
 		{
+			if (bSkipSelfCollisionInit)
+			{
+				return;
+			}
 			SelfIntersectionConstraints->Apply(Particles, SelfCollisionInit->GetContourMinimizationIntersections(), Dt);
 		});
 
 		RuleCreator.AddPostSubstepConstraintRule(
 			[this](Softs::FSolverParticlesRange& Particles, const Softs::FSolverReal Dt, const Softs::ESolverMode SolverMode)
 		{
+			if (bSkipSelfCollisionInit)
+			{
+				return;
+			}
 			const int32 NumContourIterations = SelfCollisionInit->GetNumContourMinimizationPostSteps();
 			for (int32 Iter = 0; Iter < NumContourIterations; ++Iter)
 			{
@@ -1996,6 +2012,10 @@ void FClothConstraints::CreatePBDRules()
 		ConstraintInits[ConstraintInitIndex++] =
 			[this](Softs::FSolverParticles& Particles, const Softs::FSolverReal /*Dt*/)
 			{
+				if (bSkipSelfCollisionInit)
+				{
+					return;
+				}
 				// Thickness * 2 to account for collision radius for both particles
 				SelfCollisionInit->Init(Particles, SelfCollisionConstraints->GetThickness() * (Softs::FSolverReal)2.f);
 				SelfCollisionConstraints->Init(Particles, SelfCollisionInit->GetSpatialHash(), SelfCollisionInit->GetVertexGIAColors(), SelfCollisionInit->GetTriangleGIAColors());
@@ -2014,6 +2034,10 @@ void FClothConstraints::CreatePBDRules()
 		ConstraintInits[ConstraintInitIndex++] =
 			[this](Softs::FSolverParticles& Particles, const Softs::FSolverReal /*Dt*/)
 			{
+				if (bSkipSelfCollisionInit)
+				{
+					return;
+				}
 				SelfCollisionSphereConstraints->Init(Particles);
 			};
 
@@ -2033,12 +2057,20 @@ void FClothConstraints::CreatePBDRules()
 		ConstraintInits[ConstraintInitIndex++] =
 			[this](Softs::FSolverParticles& Particles, const Softs::FSolverReal Dt)
 			{
+				if (bSkipSelfCollisionInit)
+				{
+					return;
+				}
 				SelfIntersectionConstraints->Apply(Particles, SelfCollisionInit->GetContourMinimizationIntersections(), Dt);
 			};
 
 		PostprocessingConstraintRules[PostprocessingConstraintRuleIndex++] =
 			[this](Softs::FSolverParticles& Particles, const Softs::FSolverReal Dt)
 		{
+			if (bSkipSelfCollisionInit)
+			{
+				return;
+			}
 			const int32 NumContourIterations = SelfCollisionInit->GetNumContourMinimizationPostSteps();
 			for (int32 Iter = 0; Iter < NumContourIterations; ++Iter)
 			{
diff --git a/Engine/Plugins/Experimental/ChaosCloth/Source/Chaos/Private/ChaosCloth/ChaosClothingSimulationSolver.cpp b/Engine/Plugins/Experimental/ChaosCloth/Source/Chaos/Private/ChaosCloth/ChaosClothingSimulationSolver.cpp
--- a/Engine/Plugins/Experimental/ChaosCloth/Source/Chaos/Private/ChaosCloth/ChaosClothingSimulationSolver.cpp
+++ b/Engine/Plugins/Experimental/ChaosCloth/Source/Chaos/Private/ChaosCloth/ChaosClothingSimulationSolver.cpp
@@ -83,6 +83,8 @@ namespace ClothingSimulationSolverDefault
 	static const int32 MaxNumIterations = 10;
 	static const int32 NumSubsteps = 1;
 	static const int32 MinNumSubsteps = 1;
+	static const bool bEnableNumSelfCollisionSubsteps = false;
+	static const int32 NumSelfCollisionSubsteps = 1;
 	static const FRealSingle SolverFrequency = 60.f;
 	static const FRealSingle SelfCollisionThickness = 2.f;
 	static const FRealSingle CollisionThickness = 1.2f;
@@ -1683,14 +1685,14 @@ void FClothingSimulationSolver::ApplyPreSimulationTransforms()
 	}
 }
 
-void FClothingSimulationSolver::PreSubstep(const Softs::FSolverReal InterpolationAlpha)
+void FClothingSimulationSolver::PreSubstep( const Softs::FSolverReal InterpolationAlpha, bool bDetectSelfCollisions)
 {
 	TRACE_CPUPROFILER_EVENT_SCOPE(FClothingSimulationSolver_PreSubstep);
 
 	if (Evolution)
 	{
 		const TArray<uint32> ActiveGroups = Evolution->GetActiveGroups().Array();
-		PhysicsParallelFor(ActiveGroups.Num(), [this, &ActiveGroups, InterpolationAlpha](int32 ActiveGroupIndex)
+		PhysicsParallelFor(ActiveGroups.Num(), [this, &ActiveGroups, InterpolationAlpha, bDetectSelfCollisions](int32 ActiveGroupIndex)
 		{
 			const uint32 GroupId = ActiveGroups[ActiveGroupIndex];
 			const TSet<int32>& ActiveSoftBodies = Evolution->GetGroupActiveSoftBodies(GroupId);
@@ -1700,6 +1702,7 @@ void FClothingSimulationSolver::PreSubstep(const Softs::FSolverReal Interpolatio
 				SCOPE_CYCLE_COUNTER(STAT_ChaosClothParticlePreSubstepKinematicInterpolation);
 				Softs::FSolverParticlesRange& Particles = Evolution->GetSoftBodyParticles(SoftBodyId);
 
+				GetClothConstraints(SoftBodyId).SetSkipSelfCollisionInit(!bDetectSelfCollisions);
 #if INTEL_ISPC
 				if (bRealTypeCompatibleWithISPC && bChaos_PreSubstepInterpolation_ISPC_Enabled)  // TODO: Make the ISPC works with both Single and Double depending on the FSolverReal type
 				{
@@ -1747,11 +1750,13 @@ void FClothingSimulationSolver::PreSubstep(const Softs::FSolverReal Interpolatio
 	{
 		const TPBDActiveView<Softs::FSolverParticles>& ParticlesActiveView = PBDEvolution->ParticlesActiveView();
 		ParticlesActiveView.RangeFor(
-			[this, InterpolationAlpha](Softs::FSolverParticles& Particles, int32 Offset, int32 Range)
+			[this, InterpolationAlpha, bDetectSelfCollisions](Softs::FSolverParticles& Particles, int32 Offset, int32 Range)
 		{
 			TRACE_CPUPROFILER_EVENT_SCOPE(FClothingSimulationSolver_ParticlePreSubstepKinematicInterpolation);
 			SCOPE_CYCLE_COUNTER(STAT_ChaosClothParticlePreSubstepKinematicInterpolation);
 
+			GetClothConstraints(Offset).SetSkipSelfCollisionInit(!bDetectSelfCollisions);
+
 			const int32 RangeSize = Range - Offset;
 	
 #if INTEL_ISPC
@@ -1935,6 +1940,12 @@ void FClothingSimulationSolver::Update(Softs::FSolverReal InDeltaTime)
 				Properties.GetValue<int32>(TEXT("NumSubsteps"), ClothingSimulationSolverDefault::NumSubsteps),
 				ClothingSimulationSolverDefault::MinNumSubsteps);
 
+			const bool bConfigEnableNumSelfCollisionSubsteps = Properties.GetValue<bool>(TEXT("EnableNumSelfCollisionSubsteps"), ClothingSimulationSolverDefault::bEnableNumSelfCollisionSubsteps);
+
+			const int32 ConfigNumSelfCollisionSubsteps = bConfigEnableNumSelfCollisionSubsteps ? FMath::Clamp(Properties.GetValue<int32>(TEXT("NumSelfCollisionSubsteps"), ClothingSimulationSolverDefault::NumSelfCollisionSubsteps), ClothingSimulationSolverDefault::MinNumSubsteps, ConfigNumSubsteps) : ConfigNumSubsteps;
+
+			const int32 NumSubstepsPerSelfCollisionSubstep = (ConfigNumSubsteps - 1) / ConfigNumSelfCollisionSubsteps + 1;
+
 			if (Evolution)
 			{
 				Evolution->SetSolverProperties(Properties);
@@ -1966,7 +1977,7 @@ void FClothingSimulationSolver::Update(Softs::FSolverReal InDeltaTime)
 	
 			for (int32 i = 0; i < ConfigNumSubsteps; ++i)
 			{
-				PreSubstep(FMath::Clamp((Softs::FSolverReal)(i + 1) / (Softs::FSolverReal)ConfigNumSubsteps, (Softs::FSolverReal)0., (Softs::FSolverReal)1.));
+				PreSubstep(FMath::Clamp((Softs::FSolverReal)(i + 1) / (Softs::FSolverReal)ConfigNumSubsteps, (Softs::FSolverReal)0., (Softs::FSolverReal)1.), i% NumSubstepsPerSelfCollisionSubstep == 0);
 				if (Evolution)
 				{
 					Evolution->AdvanceOneTimeStep(SubstepDeltaTime, (Softs::FSolverReal)ConfigNumSubsteps);
@@ -2048,7 +2059,7 @@ void FClothingSimulationSolver::UpdateFromCache(const TArray<FVector>& CachedPos
 
 int32 FClothingSimulationSolver::GetNumUsedIterations() const
 {
-	return Evolution ? Evolution->GetIterations() : PBDEvolution->GetIterations();
+	return Evolution ? Evolution->GetNumUsedIterations() : PBDEvolution->GetIterations();
 }
 
 int32 FClothingSimulationSolver::GetNumLinearSolverIterations(int32 ParticleRangeId) const
diff --git a/Engine/Plugins/Experimental/ChaosCloth/Source/Chaos/Public/ChaosCloth/ChaosClothConstraints.h b/Engine/Plugins/Experimental/ChaosCloth/Source/Chaos/Public/ChaosCloth/ChaosClothConstraints.h
--- a/Engine/Plugins/Experimental/ChaosCloth/Source/Chaos/Public/ChaosCloth/ChaosClothConstraints.h
+++ b/Engine/Plugins/Experimental/ChaosCloth/Source/Chaos/Public/ChaosCloth/ChaosClothConstraints.h
@@ -57,6 +57,7 @@ PRAGMA_ENABLE_DEPRECATION_WARNINGS
 			int32 InParticleOffset,
 			int32 InNumParticles);
 
+		void SetSkipSelfCollisionInit(bool bValue) { bSkipSelfCollisionInit = bValue; }
 		// ---- End of Solver interface ----
 
 		// ---- Cloth interface ----
@@ -245,6 +246,8 @@ PRAGMA_ENABLE_DEPRECATION_WARNINGS
 		int32 NumPostCollisionConstraintRules;
 		int32 NumPostprocessingConstraintRules;
 
+		bool bSkipSelfCollisionInit = false;
+
 		//~ Begin Force-based solver only fields
 		FPerSolverFieldSystem* PerSolverField;
 		const TArray<Softs::FSolverVec3>* Normals;
diff --git a/Engine/Plugins/Experimental/ChaosCloth/Source/Chaos/Public/ChaosCloth/ChaosClothingSimulationSolver.h b/Engine/Plugins/Experimental/ChaosCloth/Source/Chaos/Public/ChaosCloth/ChaosClothingSimulationSolver.h
--- a/Engine/Plugins/Experimental/ChaosCloth/Source/Chaos/Public/ChaosCloth/ChaosClothingSimulationSolver.h
+++ b/Engine/Plugins/Experimental/ChaosCloth/Source/Chaos/Public/ChaosCloth/ChaosClothingSimulationSolver.h
@@ -370,7 +370,7 @@ namespace Chaos
 		/** End PBD - only methods */
 
 		CHAOSCLOTH_API void ApplyPreSimulationTransforms();
-		CHAOSCLOTH_API void PreSubstep(const Softs::FSolverReal InterpolationAlpha);
+		CHAOSCLOTH_API void PreSubstep(const Softs::FSolverReal InterpolationAlpha, bool bDetectSelfCollisions);
 
 		// Update the solver field forces/velocities at the particles location
 		CHAOSCLOTH_API void UpdateSolverField();
diff --git a/Engine/Plugins/Experimental/ChaosClothAssetEditor/Source/ChaosClothAssetDataflowNodes/Private/ChaosClothAsset/SimulationSolverConfigNode.cpp b/Engine/Plugins/Experimental/ChaosClothAssetEditor/Source/ChaosClothAssetDataflowNodes/Private/ChaosClothAsset/SimulationSolverConfigNode.cpp
--- a/Engine/Plugins/Experimental/ChaosClothAssetEditor/Source/ChaosClothAssetDataflowNodes/Private/ChaosClothAsset/SimulationSolverConfigNode.cpp
+++ b/Engine/Plugins/Experimental/ChaosClothAssetEditor/Source/ChaosClothAssetDataflowNodes/Private/ChaosClothAsset/SimulationSolverConfigNode.cpp
@@ -17,6 +17,8 @@ void FChaosClothAssetSimulationSolverConfigNode::AddProperties(Dataflow::FContex
 	UE_CHAOS_CLOTHASSET_SIMULATIONCONFIG_SETPROPERTY(NumIterations);
 	UE_CHAOS_CLOTHASSET_SIMULATIONCONFIG_SETPROPERTY(MaxNumIterations);
 	UE_CHAOS_CLOTHASSET_SIMULATIONCONFIG_SETPROPERTY(NumSubsteps);
+	UE_CHAOS_CLOTHASSET_SIMULATIONCONFIG_SETPROPERTYBOOL(EnableNumSelfCollisionSubsteps);
+	UE_CHAOS_CLOTHASSET_SIMULATIONCONFIG_SETPROPERTY(NumSelfCollisionSubsteps);
 	UE_CHAOS_CLOTHASSET_SIMULATIONCONFIG_SETPROPERTYBOOL(EnableForceBasedSolver);
 	UE_CHAOS_CLOTHASSET_SIMULATIONCONFIG_SETPROPERTY(NumNewtonIterations);
 	UE_CHAOS_CLOTHASSET_SIMULATIONCONFIG_SETPROPERTY(MaxNumCGIterations);
diff --git a/Engine/Plugins/Experimental/ChaosClothAssetEditor/Source/ChaosClothAssetDataflowNodes/Public/ChaosClothAsset/SimulationSolverConfigNode.h b/Engine/Plugins/Experimental/ChaosClothAssetEditor/Source/ChaosClothAssetDataflowNodes/Public/ChaosClothAsset/SimulationSolverConfigNode.h
--- a/Engine/Plugins/Experimental/ChaosClothAssetEditor/Source/ChaosClothAssetDataflowNodes/Public/ChaosClothAsset/SimulationSolverConfigNode.h
+++ b/Engine/Plugins/Experimental/ChaosClothAssetEditor/Source/ChaosClothAssetDataflowNodes/Public/ChaosClothAsset/SimulationSolverConfigNode.h
@@ -36,6 +36,20 @@ public:
 	UPROPERTY(EditAnywhere, Category = Simulation, meta = (UIMin = "1", UIMax = "10", ClampMin = "1", ClampMax = "100"))
 	int32 NumSubsteps = 1;
 
+
+	/**
+	* Enable setting separate SelfCollisionSubsteps. Otherwise, self collisions will be detected every substep.
+	*/
+	UPROPERTY(EditAnywhere, Category = Simulation, meta = (InlineEditConditionToggle))
+	bool bEnableNumSelfCollisionSubsteps = false;
+
+	/**
+	 * Set a separate number of self collision substeps. Lower this number to increase speed at the expense of lower self collision accuracy.
+	 * Actual value always clamped between [1, NumSubsteps].
+	 */
+	UPROPERTY(EditAnywhere, Category = Simulation, meta = (UIMin = "1", UIMax = "10", ClampMin = "1", ClampMax = "100", EditCondition = "bEnableNumSelfCollisionSubsteps"))
+	int32 NumSelfCollisionSubsteps = 1;
+
 	/**
 	 * Enable the higher accuracy force-based solver (experimental).
 	 */
diff --git a/Engine/Source/Runtime/Experimental/Chaos/Private/Chaos/SoftsEvolution.cpp b/Engine/Source/Runtime/Experimental/Chaos/Private/Chaos/SoftsEvolution.cpp
--- a/Engine/Source/Runtime/Experimental/Chaos/Private/Chaos/SoftsEvolution.cpp
+++ b/Engine/Source/Runtime/Experimental/Chaos/Private/Chaos/SoftsEvolution.cpp
@@ -326,7 +326,7 @@ TArray<FSolverCollisionParticlesRange> FEvolution::GetActiveCollisionParticles(u
 void FEvolution::SetSolverProperties(const FCollectionPropertyConstFacade& Properties)
 {
 	bEnableForceBasedSolver = GetEnableForceBasedSolver(Properties, Private::EvolutionSolverDefault::bEnableForceBasedSolver);
-	MaxNumIterations = FMath::Min(GetMaxNumIterations(Properties, Private::EvolutionSolverDefault::MaxNumIterations), Private::EvolutionSolverDefault::MinNumIterations);
+	MaxNumIterations = FMath::Max(GetMaxNumIterations(Properties, Private::EvolutionSolverDefault::MaxNumIterations), Private::EvolutionSolverDefault::MinNumIterations);
 	NumIterations = GetNumIterations(Properties, Private::EvolutionSolverDefault::NumIterations);
 	bDoQuasistatics = GetDoQuasistatics(Properties, Private::EvolutionSolverDefault::bDoQuasistatics);
 	SolverFrequency = GetSolverFrequency(Properties, Private::EvolutionSolverDefault::SolverFrequency);
@@ -345,6 +345,7 @@ void FEvolution::AdvanceOneTimeStep(const FSolverReal Dt, const FSolverReal Time
 	const int32 TimeDependentNumIterations = (bDisableTimeDependentNumIterations || NumIterations == 0) ? NumIterations :
 		FMath::Clamp(FMath::RoundToInt32(SolverFrequency * Dt * TimeDependentIterationMultiplier * (Softs::FSolverReal)NumIterations),
 			Private::EvolutionSolverDefault::MinNumIterations, MaxNumIterations);
+	NumUsedIterations = TimeDependentNumIterations;
 
 	const TArray<uint32> ActiveGroupsArray = ActiveGroups.Array();
 
diff --git a/Engine/Source/Runtime/Experimental/Chaos/Public/Chaos/SoftsEvolution.h b/Engine/Source/Runtime/Experimental/Chaos/Public/Chaos/SoftsEvolution.h
--- a/Engine/Source/Runtime/Experimental/Chaos/Public/Chaos/SoftsEvolution.h
+++ b/Engine/Source/Runtime/Experimental/Chaos/Public/Chaos/SoftsEvolution.h
@@ -203,6 +203,7 @@ public:
 	FSolverReal GetTime() const { return Time; }
 	int32 GetIterations() const { return NumIterations; }
 	int32 GetMaxIterations() const { return MaxNumIterations; }
+	int32 GetNumUsedIterations() const { return NumUsedIterations; }
 	bool GetDisableTimeDependentNumIterations() const { return bDisableTimeDependentNumIterations; }
 	bool GetDoQuasistatics() const { return bDoQuasistatics; }
 	void SetDisableTimeDependentNumIterations(bool bDisable) { bDisableTimeDependentNumIterations = bDisable; }
@@ -418,6 +419,7 @@ private:
 	bool bEnableForceBasedSolver = false;
 	int32 MaxNumIterations; // Used for time-dependent iteration counts
 	int32 NumIterations; // PBD iterations
+	int32 NumUsedIterations = 0; // Last actual time-dependent iteration count
 	int32 NumNewtonIterations; // Implicit force-based solve
 	bool bDisableTimeDependentNumIterations = false;
 	bool bDoQuasistatics;
