Fix cloth caching crash from new SoftsEvolution update to the cloth plugin.

diff --git a/Engine/Plugins/Experimental/ChaosCloth/Source/Chaos/Private/ChaosCloth/ChaosClothingSimulationSolver.cpp b/Engine/Plugins/Experimental/ChaosCloth/Source/Chaos/Private/ChaosCloth/ChaosClothingSimulationSolver.cpp
--- a/Engine/Plugins/Experimental/ChaosCloth/Source/Chaos/Private/ChaosCloth/ChaosClothingSimulationSolver.cpp
+++ b/Engine/Plugins/Experimental/ChaosCloth/Source/Chaos/Private/ChaosCloth/ChaosClothingSimulationSolver.cpp
@@ -2015,7 +2015,7 @@ void FClothingSimulationSolver::Update(Softs::FSolverReal InDeltaTime)
 
 void FClothingSimulationSolver::UpdateFromCache(const FClothingSimulationCacheData& CacheData)
 {
-	Chaos::Softs::FSolverParticles& SolverParticles = Evolution->GetParticles();
+	Chaos::Softs::FSolverParticles& SolverParticles = Evolution ? Evolution->GetParticles() : PBDEvolution->GetParticles();
 	const int32 NumParticles = GetNumParticles();
 	const int32 NumCachedParticles = CacheData.CacheIndices.Num();	
 	const bool bHasVelocity = CacheData.CachedVelocities.Num() > 0;
@@ -2039,7 +2039,7 @@ void FClothingSimulationSolver::UpdateFromCache(const FClothingSimulationCacheDa
 
 void FClothingSimulationSolver::UpdateFromCache(const TArray<FVector>& CachedPositions, const TArray<FVector>& CachedVelocities) 
 {
-	Chaos::Softs::FSolverParticles& SolverParticles = Evolution->GetParticles();
+	Chaos::Softs::FSolverParticles& SolverParticles = Evolution ? Evolution->GetParticles() : PBDEvolution->GetParticles();
 	const int32 NumParticles = GetNumParticles();
 	const bool bHasVelocity = CachedVelocities.Num() > 0;
 	if(CachedPositions.Num() == NumParticles)
