Chaos Cloth: Fix crash when drawing backstops without maps.

diff --git a/Engine/Plugins/Experimental/ChaosCloth/Source/Chaos/Private/ChaosCloth/ChaosClothVisualization.cpp b/Engine/Plugins/Experimental/ChaosCloth/Source/Chaos/Private/ChaosCloth/ChaosClothVisualization.cpp
--- a/Engine/Plugins/Experimental/ChaosCloth/Source/Chaos/Private/ChaosCloth/ChaosClothVisualization.cpp
+++ b/Engine/Plugins/Experimental/ChaosCloth/Source/Chaos/Private/ChaosCloth/ChaosClothVisualization.cpp
@@ -1187,8 +1187,6 @@ static FAutoConsoleVariableRef CVarClothVizWeightMapName(TEXT("p.ChaosClothVisua
 			{
 				const bool bUseLegacyBackstop = BackstopConstraint->UseLegacyBackstop();
 
-				const TConstArrayView<FRealSingle>& BackstopDistances = Cloth->GetWeightMapByProperty(Solver, TEXT("BackstopDistance"));
-				const TConstArrayView<FRealSingle>& BackstopRadiuses = Cloth->GetWeightMapByProperty(Solver, TEXT("BackstopRadius"));
 				const TConstArrayView<Softs::FSolverVec3> AnimationPositions = Cloth->GetAnimationPositions(Solver);
 				const TConstArrayView<Softs::FSolverVec3> AnimationNormals = Cloth->GetAnimationNormals(Solver);
 				const TConstArrayView<Softs::FSolverVec3> ParticlePositions = Cloth->GetParticlePositions(Solver);
@@ -1199,8 +1197,8 @@ static FAutoConsoleVariableRef CVarClothVizWeightMapName(TEXT("p.ChaosClothVisua
 					const FLinearColor ColorLight = FLinearColor::MakeFromHSV8(ColorSeed, 160, 128);
 					const FLinearColor ColorDark = FLinearColor::MakeFromHSV8(ColorSeed, 160, 64);
 
-					const FReal BackstopRadius = BackstopRadiuses[Index] * BackstopConstraint->GetScale();
-					const FReal BackstopDistance = BackstopDistances[Index];
+					const FReal BackstopRadius = BackstopConstraint->GetBackstopRadius(Index) * BackstopConstraint->GetScale();
+					const FReal BackstopDistance = BackstopConstraint->GetBackstopDistance(Index) * BackstopConstraint->GetScale();
 
 					const FVector AnimationNormal(AnimationNormals[Index]);
 
@@ -1246,8 +1244,6 @@ static FAutoConsoleVariableRef CVarClothVizWeightMapName(TEXT("p.ChaosClothVisua
 			if (const Softs::FPBDSphericalBackstopConstraint* const BackstopConstraint = ClothConstraints.GetBackstopConstraints().Get())
 			{
 				const bool bUseLegacyBackstop = BackstopConstraint->UseLegacyBackstop();
-				const TConstArrayView<FRealSingle>& BackstopDistances = Cloth->GetWeightMapByProperty(Solver, TEXT("BackstopDistance"));
-				const TConstArrayView<FRealSingle>& BackstopRadiuses = Cloth->GetWeightMapByProperty(Solver, TEXT("BackstopRadius"));
 				const TConstArrayView<Softs::FSolverVec3> AnimationPositions = Cloth->GetAnimationPositions(Solver);
 				const TConstArrayView<Softs::FSolverVec3> AnimationNormals = Cloth->GetAnimationNormals(Solver);
 
@@ -1257,8 +1253,8 @@ static FAutoConsoleVariableRef CVarClothVizWeightMapName(TEXT("p.ChaosClothVisua
 					const FLinearColor ColorLight = FLinearColor::MakeFromHSV8(ColorSeed, 160, 128);
 					const FLinearColor ColorDark = FLinearColor::MakeFromHSV8(ColorSeed, 160, 64);
 
-					const FReal BackstopRadius = BackstopRadiuses[Index] * BackstopConstraint->GetScale();
-					const FReal BackstopDistance = BackstopDistances[Index];
+					const FReal BackstopRadius = BackstopConstraint->GetBackstopRadius(Index) * BackstopConstraint->GetScale();
+					const FReal BackstopDistance = BackstopConstraint->GetBackstopDistance(Index) * BackstopConstraint->GetScale();
 
 					const FVector AnimationNormal(AnimationNormals[Index]);
 
@@ -1745,7 +1741,7 @@ static FAutoConsoleVariableRef CVarClothVizWeightMapName(TEXT("p.ChaosClothVisua
 			const Softs::FSolverVec3& P1 = Positions[Constraints[ConstraintIndex][0]];
 			const Softs::FSolverVec3& P2 = Positions[Constraints[ConstraintIndex][1]];
 
-			const bool bIsBuckled = IsBuckled[ConstraintIndex];
+			const bool bIsBuckled = IsBuckled.IsValidIndex(ConstraintIndex) ? IsBuckled[ConstraintIndex] : false; // IsBuckled is empty if the simulation is paused.
 
 			const FVec3 Pos0 = FVec3(P1) + LocalSpaceLocation;
 			const FVec3 Pos1 = FVec3(P2) + LocalSpaceLocation;
diff --git a/Engine/Source/Runtime/Experimental/Chaos/Public/Chaos/PBDSphericalConstraint.h b/Engine/Source/Runtime/Experimental/Chaos/Public/Chaos/PBDSphericalConstraint.h
--- a/Engine/Source/Runtime/Experimental/Chaos/Public/Chaos/PBDSphericalConstraint.h
+++ b/Engine/Source/Runtime/Experimental/Chaos/Public/Chaos/PBDSphericalConstraint.h
@@ -407,6 +407,16 @@ PRAGMA_ENABLE_DEPRECATION_WARNINGS
 		return bUseLegacyBackstop;
 	}
 
+	FSolverReal GetBackstopRadius(int32 ConstraintIndex) const
+	{
+		return SphereRadii.Num() == ParticleCount ? BackstopRadiusBase + BackstopRadiusRange * SphereRadii[ConstraintIndex] : BackstopRadiusBase;
+	}
+
+	FSolverReal GetBackstopDistance(int32 ConstraintIndex) const
+	{
+		return SphereOffsetDistances.Num() == ParticleCount ? BackstopDistanceBase + BackstopDistanceRange * SphereOffsetDistances[ConstraintIndex] : BackstopDistanceBase;
+	}
+
 private:
 
 	template<bool bHasBackstopDistance, bool bHasBackstopRadius, typename SolverParticlesOrRange>
