Chaos Cloth - Added missing TEXT qualifier to selection string names.

diff --git a/Engine/Plugins/Experimental/ChaosCloth/Source/Chaos/Private/ChaosCloth/ChaosClothingSimulationSkeletalMesh.cpp b/Engine/Plugins/Experimental/ChaosCloth/Source/Chaos/Private/ChaosCloth/ChaosClothingSimulationSkeletalMesh.cpp
--- a/Engine/Plugins/Experimental/ChaosCloth/Source/Chaos/Private/ChaosCloth/ChaosClothingSimulationSkeletalMesh.cpp
+++ b/Engine/Plugins/Experimental/ChaosCloth/Source/Chaos/Private/ChaosCloth/ChaosClothingSimulationSkeletalMesh.cpp
@@ -166,7 +166,7 @@ PRAGMA_ENABLE_DEPRECATION_WARNINGS
 		{
 			const FClothLODDataCommon& ClothLODData = Asset->LodData[LODIndex];
 			const FClothPhysicalMeshData& ClothPhysicalMeshData = ClothLODData.PhysicalMeshData;
-			static const FString SelfCollisionSetName("_SelfCollisionSpheres");
+			static const FString SelfCollisionSetName(TEXT("_SelfCollisionSpheres"));
 			VertexSets.Add(SelfCollisionSetName, &ClothPhysicalMeshData.SelfCollisionVertexSet);
 		}
 
diff --git a/Engine/Source/Runtime/Experimental/Chaos/Private/Chaos/PBDSelfCollisionSphereConstraints.cpp b/Engine/Source/Runtime/Experimental/Chaos/Private/Chaos/PBDSelfCollisionSphereConstraints.cpp
--- a/Engine/Source/Runtime/Experimental/Chaos/Private/Chaos/PBDSelfCollisionSphereConstraints.cpp
+++ b/Engine/Source/Runtime/Experimental/Chaos/Private/Chaos/PBDSelfCollisionSphereConstraints.cpp
@@ -7,7 +7,7 @@
 
 namespace Chaos::Softs {
 
-	const FString FPBDSelfCollisionSphereConstraints::SelectionSetName("_SelfCollisionSpheres");
+	const FString FPBDSelfCollisionSphereConstraints::SelectionSetName(TEXT("_SelfCollisionSpheres"));
 
 	namespace Private
 	{
