Fix CIS warning.

diff --git a/Engine/Plugins/Experimental/ChaosCloth/Source/Chaos/Public/ChaosCloth/ChaosClothVisualization.h b/Engine/Plugins/Experimental/ChaosCloth/Source/Chaos/Public/ChaosCloth/ChaosClothVisualization.h
--- a/Engine/Plugins/Experimental/ChaosCloth/Source/Chaos/Public/ChaosCloth/ChaosClothVisualization.h
+++ b/Engine/Plugins/Experimental/ChaosCloth/Source/Chaos/Public/ChaosCloth/ChaosClothVisualization.h
@@ -75,7 +75,7 @@ namespace Chaos
 		void DrawLocalSpace(FPrimitiveDrawInterface* /*PDI*/ = nullptr) const {}
 		void DrawSelfCollision(FPrimitiveDrawInterface* /*PDI*/ = nullptr) const {}
 		void DrawSelfIntersection(FPrimitiveDrawInterface* /*PDI*/ = nullptr) const {}
-		CHAOSCLOTH_API void DrawKinematicColliderWired(FPrimitiveDrawInterface* PDI = nullptr) const {}
+		CHAOSCLOTH_API void DrawKinematicColliderWired(FPrimitiveDrawInterface* /*PDI*/ = nullptr) const {}
 		void DrawBounds(FPrimitiveDrawInterface* /*PDI*/ = nullptr) const {}
 		void DrawGravity(FPrimitiveDrawInterface* /*PDI*/ = nullptr) const {}
 		void DrawFictitiousAngularForces(FPrimitiveDrawInterface* /*PDI*/ = nullptr) const {}
@@ -102,7 +102,7 @@ namespace Chaos
 		void DrawParticleIndices(FCanvas* /*Canvas*/, const FSceneView* /*SceneView*/) const {}
 		void DrawElementIndices(FCanvas* /*Canvas*/, const FSceneView* /*SceneView*/) const {}
 		void DrawMaxDistanceValues(FCanvas* /*Canvas*/, const FSceneView* /*SceneView*/) const {}
-		CHAOSCLOTH_API void DrawKinematicColliderShaded(FPrimitiveDrawInterface* PDI) const {}
+		CHAOSCLOTH_API void DrawKinematicColliderShaded(FPrimitiveDrawInterface* /*PDI*/) const {}
 #endif  // #if WITH_EDITOR && CHAOS_DEBUG_DRAW
 
 #if CHAOS_DEBUG_DRAW
