Chaos Cloth: Debug draw self collision layers.

This is just debug draw for a feature, but it's actually really important for quickly validating that the layer data is working as expected and hasn't gotten corrupted by downstream dataflow nodes. I already fixed the jacket asset using this (one of the selection sets was unnamed so wasn't actually working).

diff --git a/Engine/Plugins/Experimental/ChaosCloth/Source/Chaos/Private/ChaosCloth/ChaosClothVisualization.cpp b/Engine/Plugins/Experimental/ChaosCloth/Source/Chaos/Private/ChaosCloth/ChaosClothVisualization.cpp
--- a/Engine/Plugins/Experimental/ChaosCloth/Source/Chaos/Private/ChaosCloth/ChaosClothVisualization.cpp
+++ b/Engine/Plugins/Experimental/ChaosCloth/Source/Chaos/Private/ChaosCloth/ChaosClothVisualization.cpp
@@ -88,6 +88,18 @@ static FAutoConsoleVariableRef CVarClothVizAnisoSpringDrawMode(TEXT("p.ChaosClot
 static FString WeightMapName = "";
 static FAutoConsoleVariableRef CVarClothVizWeightMapName(TEXT("p.ChaosClothVisualization.WeightMapName"), WeightMapName, TEXT("Weight map name to be visualized"));
 
+// copied from ClothEditorMode
+FLinearColor PseudoRandomColor(int32 NumColorRotations)
+{
+	constexpr uint8 Spread = 157;  // Prime number that gives a good spread of colors without getting too similar as a rand might do.
+	uint8 Seed = Spread;
+	NumColorRotations = FMath::Abs(NumColorRotations);
+	for (int32 Rotation = 0; Rotation < NumColorRotations; ++Rotation)
+	{
+		Seed += Spread;
+	}
+	return FLinearColor::MakeFromHSV8(Seed, 180, 140);
+}
 }// namespace Private
 
 	FClothVisualization::FClothVisualization(const ::Chaos::FClothingSimulationSolver* InSolver)
@@ -382,6 +394,63 @@ static FAutoConsoleVariableRef CVarClothVizWeightMapName(TEXT("p.ChaosClothVisua
 	{
 		DrawWeightMapWithName(PDI, TEXT("_InpaintWeightMask"));
 	}
+
+	void FClothVisualization::DrawSelfCollisionLayers(FPrimitiveDrawInterface* PDI) const
+	{
+
+		if (!Solver || !ClothMaterialColor)
+		{
+			return;
+		}
+
+		FDynamicMeshBuilder MeshBuilder(PDI->View->GetFeatureLevel());
+		int32 VertexIndex = 0;
+
+		for (const FClothingSimulationCloth* const Cloth : Solver->GetCloths())
+		{
+			const int32 ParticleRangeId = Cloth->GetParticleRangeId(Solver);
+			if (ParticleRangeId == INDEX_NONE)
+			{
+				continue;
+			}
+			// Elements are local indexed for force based solver
+			const int32 Offset = Solver->IsForceBasedSolver() ? 0 : ParticleRangeId;
+			const TConstArrayView<TVec3<int32>> Elements = Cloth->GetTriangleMesh(Solver).GetElements();
+			const TConstArrayView<Softs::FSolverVec3> Positions = Cloth->GetParticlePositions(Solver);
+			const TConstArrayView<int32>& WeightMap = Cloth->GetFaceIntMapByProperty(Solver, TEXT("SelfCollisionLayers"));
+
+			for (int32 ElementIndex = 0; ElementIndex < Elements.Num(); ++ElementIndex, VertexIndex += 3)
+			{
+
+				const TVec3<int32>& Element = Elements[ElementIndex];
+
+				const FVector3f Pos0(Positions[Element.X - Offset]);
+				const FVector3f Pos1(Positions[Element.Y - Offset]);
+				const FVector3f Pos2(Positions[Element.Z - Offset]);
+
+				const FVector3f Normal = FVector3f::CrossProduct(Pos2 - Pos0, Pos1 - Pos0).GetSafeNormal();
+				const FVector3f Tangent = ((Pos1 + Pos2) * 0.5f - Pos0).GetSafeNormal();
+
+				FLinearColor VertexColor1 = FLinearColor::Gray;
+				FLinearColor VertexColor2 = FLinearColor::Gray;
+				FLinearColor VertexColor3 = FLinearColor::Gray;
+
+				if (!WeightMap.IsEmpty() && WeightMap.Num() == Elements.Num() && WeightMap[ElementIndex]!=INDEX_NONE) // if map with that name exists and not empty
+				{
+					VertexColor1 = VertexColor2 = VertexColor3 = Chaos::Private::PseudoRandomColor(WeightMap[ElementIndex]);
+				}
+
+				MeshBuilder.AddVertex(FDynamicMeshVertex(Pos0, Tangent, Normal, FVector2f(0.f, 0.f), VertexColor1.ToFColor(true)));
+				MeshBuilder.AddVertex(FDynamicMeshVertex(Pos1, Tangent, Normal, FVector2f(0.f, 1.f), VertexColor2.ToFColor(true)));
+				MeshBuilder.AddVertex(FDynamicMeshVertex(Pos2, Tangent, Normal, FVector2f(1.f, 1.f), VertexColor3.ToFColor(true)));
+				MeshBuilder.AddTriangle(VertexIndex, VertexIndex + 1, VertexIndex + 2);
+			}
+		}
+
+		FMatrix LocalSimSpaceToWorld(FMatrix::Identity);
+		LocalSimSpaceToWorld.SetOrigin(Solver->GetLocalSpaceLocation());
+		MeshBuilder.Draw(PDI, LocalSimSpaceToWorld, ClothMaterialColor->GetRenderProxy(), SDPG_World, false, false);
+	}
 #endif  // #if WITH_EDITOR
 
 	static void DrawPoint(FPrimitiveDrawInterface* PDI, const FVector& Pos, const FLinearColor& Color, const UMaterial* ClothMaterialVertex, const float Thickness = 1.f)  // Use color or material
diff --git a/Engine/Plugins/Experimental/ChaosCloth/Source/Chaos/Private/ChaosCloth/ChaosClothingSimulationCloth.cpp b/Engine/Plugins/Experimental/ChaosCloth/Source/Chaos/Private/ChaosCloth/ChaosClothingSimulationCloth.cpp
--- a/Engine/Plugins/Experimental/ChaosCloth/Source/Chaos/Private/ChaosCloth/ChaosClothingSimulationCloth.cpp
+++ b/Engine/Plugins/Experimental/ChaosCloth/Source/Chaos/Private/ChaosCloth/ChaosClothingSimulationCloth.cpp
@@ -893,6 +893,19 @@ TConstArrayView<FRealSingle> FClothingSimulationCloth::GetWeightMapByProperty(co
 	return GetWeightMapByName(Solver, PropertyString);
 }
 
+TConstArrayView<int32> FClothingSimulationCloth::GetFaceIntMapByName(const FClothingSimulationSolver* Solver, const FString& Name) const
+{
+	const int32 LODIndex = LODIndices.FindChecked(Solver);
+	return LODData.IsValidIndex(LODIndex) ? LODData[LODIndex]->FaceIntMaps.FindRef(Name) : TConstArrayView<int32>();
+}
+
+TConstArrayView<int32> FClothingSimulationCloth::GetFaceIntMapByProperty(const FClothingSimulationSolver* Solver, const FString& Property) const
+{
+	check(Config);
+	const FString PropertyString = Config->GetProperties(GetLODIndex(Solver)).GetStringValue(Property);
+	return GetFaceIntMapByName(Solver, PropertyString);
+}
+
 const TArray<TConstArrayView<TTuple<int32, int32, float>>>& FClothingSimulationCloth::GetTethers(const FClothingSimulationSolver* Solver) const
 {
 	const int32 LODIndex = LODIndices.FindChecked(Solver);
diff --git a/Engine/Plugins/Experimental/ChaosCloth/Source/Chaos/Public/ChaosCloth/ChaosClothVisualization.h b/Engine/Plugins/Experimental/ChaosCloth/Source/Chaos/Public/ChaosCloth/ChaosClothVisualization.h
--- a/Engine/Plugins/Experimental/ChaosCloth/Source/Chaos/Public/ChaosCloth/ChaosClothVisualization.h
+++ b/Engine/Plugins/Experimental/ChaosCloth/Source/Chaos/Public/ChaosCloth/ChaosClothVisualization.h
@@ -89,6 +89,7 @@ namespace Chaos
 		CHAOSCLOTH_API void DrawElementIndices(FCanvas* Canvas, const FSceneView* SceneView) const;
 		CHAOSCLOTH_API void DrawMaxDistanceValues(FCanvas* Canvas, const FSceneView* SceneView) const;
 		CHAOSCLOTH_API void DrawWeightMap(FPrimitiveDrawInterface* PDI = nullptr) const;
+		CHAOSCLOTH_API void DrawSelfCollisionLayers(FPrimitiveDrawInterface* PDI = nullptr) const;
 		CHAOSCLOTH_API void DrawInpaintWeightsMatched(FPrimitiveDrawInterface* PDI = nullptr) const;
 		CHAOSCLOTH_API void DrawKinematicColliderShaded(FPrimitiveDrawInterface* PDI) const;
 	protected:
diff --git a/Engine/Plugins/Experimental/ChaosCloth/Source/Chaos/Public/ChaosCloth/ChaosClothingSimulationCloth.h b/Engine/Plugins/Experimental/ChaosCloth/Source/Chaos/Public/ChaosCloth/ChaosClothingSimulationCloth.h
--- a/Engine/Plugins/Experimental/ChaosCloth/Source/Chaos/Public/ChaosCloth/ChaosClothingSimulationCloth.h
+++ b/Engine/Plugins/Experimental/ChaosCloth/Source/Chaos/Public/ChaosCloth/ChaosClothingSimulationCloth.h
@@ -183,6 +183,10 @@ PRAGMA_ENABLE_DEPRECATION_WARNINGS
 		CHAOSCLOTH_API TConstArrayView<FRealSingle> GetWeightMapByName(const FClothingSimulationSolver* Solver, const FString& Name) const;
 		// Return the weight map of the specified property name if it exists and is available on the current LOD, or an empty array view otherwise.
 		CHAOSCLOTH_API TConstArrayView<FRealSingle> GetWeightMapByProperty(const FClothingSimulationSolver* Solver, const FString& Property) const;
+		// Return the face int map of the specified name if available on the current LOD, or an empty array view otherwise.
+		CHAOSCLOTH_API TConstArrayView<int32> GetFaceIntMapByName(const FClothingSimulationSolver* Solver, const FString& Name) const;
+		// Return the face int map of the specified property name if it exists and is available on the current LOD, or an empty array view otherwise.
+		CHAOSCLOTH_API TConstArrayView<int32> GetFaceIntMapByProperty(const FClothingSimulationSolver* Solver, const FString& Property) const;
 		UE_DEPRECATED(5.3, "Returns an empty array from 5.3. Update your code with GetWeightMapByName and GetWeightMapByProperty to return the current LOD weight map instead.")
 		CHAOSCLOTH_API const TArray<TConstArrayView<FRealSingle>>& GetWeightMaps(const FClothingSimulationSolver* Solver) const;
 		// Return the current LOD tethers.
diff --git a/Engine/Plugins/Experimental/ChaosClothAssetEditor/Source/ChaosClothAssetEditor/Private/ChaosClothAsset/ClothEditorSimulationVisualization.cpp b/Engine/Plugins/Experimental/ChaosClothAssetEditor/Source/ChaosClothAssetEditor/Private/ChaosClothAsset/ClothEditorSimulationVisualization.cpp
--- a/Engine/Plugins/Experimental/ChaosClothAssetEditor/Source/ChaosClothAssetEditor/Private/ChaosClothAsset/ClothEditorSimulationVisualization.cpp
+++ b/Engine/Plugins/Experimental/ChaosClothAssetEditor/Source/ChaosClothAssetEditor/Private/ChaosClothAsset/ClothEditorSimulationVisualization.cpp
@@ -169,6 +169,7 @@ const FVisualizationOption FVisualizationOption::OptionData[] =
 	FVisualizationOption(&::Chaos::FClothVisualization::DrawFictitiousAngularForces, TEXT("p.ChaosClothAssetEditor.DebugDrawFictitiousAngularForces"), LOCTEXT("ChaosVisName_FictitiousAngularForces", "Fictitious Angular Forces"), LOCTEXT("ChaosVisName_Gravity_FictitiousAngularForces", "Draws fictitious angular forces (force based solver only)")),
 	FVisualizationOption(&::Chaos::FClothVisualization::DrawSelfCollision, TEXT("p.ChaosClothAssetEditor.DebugDrawSelfCollision"), LOCTEXT("ChaosVisName_SelfCollision"       , "Self Collision"), LOCTEXT("ChaosVisName_SelfCollision_Tooltip", "Draws the self collision thickness/debugging information")),
 	FVisualizationOption(&::Chaos::FClothVisualization::DrawSelfIntersection, TEXT("p.ChaosClothAssetEditor.DebugDrawSelfIntersection"), LOCTEXT("ChaosVisName_SelfIntersection", "Self Intersection"), LOCTEXT("ChaosVisName_SelfIntersection_Tooltip", "Draws the self intersection contour/region information")),
+	FVisualizationOption(&::Chaos::FClothVisualization::DrawSelfCollisionLayers, TEXT("p.ChaosClothAssetEditor.DebugDrawSelfCollisionLayers"), LOCTEXT("ChaosVisName_SelfCollisionLayers", "Self Collision Layers"), LOCTEXT("ChaosVisName_SelfCollisionLayers_Tooltip", "Draws the self collision layers"), /*bDisablesSimulation =*/false, /*bHidesClothSections=*/true),
 	FVisualizationOption(&::Chaos::FClothVisualization::DrawKinematicColliderShaded, TEXT("p.ChaosClothAssetEditor.DebugDrawKinematicColliderShaded"), LOCTEXT("ChaosVisName_DrawKinematicColliderShaded", "Draw Kinematic Colliders (Shaded)"), LOCTEXT("ChaosVisName_DrawKinematicColliderShaded_Tooltip", "Draw kinematic cloth colliders with flat shading.")),
 	FVisualizationOption(&::Chaos::FClothVisualization::DrawKinematicColliderWired, TEXT("p.ChaosClothAssetEditor.DebugDrawKinematicColliderWired"), LOCTEXT("ChaosVisName_DrawKinematicColliderWired", "Draw Kinematic Colliders (Wireframe)"), LOCTEXT("ChaosVisName_DrawKinematicColliderWired_Tooltip", "Draw kinematic cloth colliders in wireframe.")),
 	FVisualizationOption(&GetSimulationStatisticsString, TEXT("p.ChaosClothAssetEditor.DebugDrawSimulationStatistics"), LOCTEXT("ChaosVisName_SimulationStatistics"    , "Simulation Statistics"), LOCTEXT("ChaosVisName_SimulationStatistics_Tooltip", "Displays simulation statistics")),
