Chaos Cloth: Add multi-res constraint which targets a fine LOD to goal to a coarse LOD.

This is still experimental. In particular, only LOD0 can be use the multi-res constraint, and LOD switching is not fully supported.

Currently the binding data between coarse and fine LOD is calculated at runtime. We should move this to being stored in the cloth asset in a future CL.

The new constraint works like the AnimDrive constraint, except instead of driving to the skinned mesh, it's driving to a coarser LOD.
Each substep, the coarse LOD will simulate and update the Multi-Res constraint targets. Then the fine LOD will simulate using these targets.

diff --git a/Engine/Plugins/Experimental/ChaosCloth/Source/Chaos/Private/ChaosCloth/ChaosClothConstraints.cpp b/Engine/Plugins/Experimental/ChaosCloth/Source/Chaos/Private/ChaosCloth/ChaosClothConstraints.cpp
--- a/Engine/Plugins/Experimental/ChaosCloth/Source/Chaos/Private/ChaosCloth/ChaosClothConstraints.cpp
+++ b/Engine/Plugins/Experimental/ChaosCloth/Source/Chaos/Private/ChaosCloth/ChaosClothConstraints.cpp
@@ -13,6 +13,7 @@
 #include "Chaos/XPBDAxialSpringConstraints.h"
 #include "Chaos/PBDVolumeConstraint.h"
 #include "Chaos/XPBDLongRangeConstraints.h"
+#include "Chaos/SoftsMultiResConstraints.h"
 #include "Chaos/PBDSphericalConstraint.h"
 #include "Chaos/PBDAnimDriveConstraint.h"
 #include "Chaos/PBDShapeConstraints.h"
@@ -29,6 +30,7 @@
 #include "Chaos/Deformable/GaussSeidelMasterConstraint.h"
 #include "Chaos/Deformable/GaussSeidelCorotatedCodimensionalConstraints.h"
 #include "PhysicsProxy/PerSolverFieldSystem.h"
+#include "Utils/ClothingMeshUtils.h"
 #include "HAL/IConsoleManager.h"
 
 namespace Chaos {
@@ -43,8 +45,6 @@ extern TAutoConsoleVariable<float> CVarGravityMultiplier;
 bool bEnableGS = false;
 
 #if !UE_BUILD_SHIPPING
-
-
     static FAutoConsoleVariableRef CVarClothbEnableGS(TEXT("p.Chaos.Cloth.EnableGaussSeidel"), bEnableGS, TEXT("Use Gauss Seidel constraints instead of XPBD [def: false]"));
 
 	bool bDisplayResidual = false;
@@ -800,7 +800,10 @@ void FClothConstraints::AddRules(
 	const TMap<FString, const TSet<int32>*>& FaceSets,
 	const TMap<FString, TConstArrayView<int32>>& FaceIntMaps,
 	const TArray<TConstArrayView<TTuple<int32, int32, FRealSingle>>>& Tethers,
-	Softs::FSolverReal MeshScale, bool bEnabled)
+	Softs::FSolverReal MeshScale, bool bEnabled,
+	const FTriangleMesh* MultiResCoarseLODMesh,
+	const int32 MultiResCoarseLODParticleRangeId,
+	const TSharedPtr<Softs::FMultiResConstraints>& FineLODMultiResConstraint)
 {
 	// Self collisions
 	CreateSelfCollisionConstraints(ConfigProperties, WeightMaps, VertexSets, FaceSets, FaceIntMaps, TriangleMesh);
@@ -842,6 +845,9 @@ void FClothConstraints::AddRules(
 
 		// Body collisions
 		CreateCollisionConstraint(ConfigProperties, MeshScale);
+
+		//Multires Springs
+		CreateMultiresConstraint(ConfigProperties, WeightMaps, TriangleMesh, MultiResCoarseLODMesh, MultiResCoarseLODParticleRangeId);
 	}
 
 	// Commit rules to solver
@@ -849,7 +855,7 @@ void FClothConstraints::AddRules(
 	{
 		if (Evolution)
 		{
-			CreateForceBasedRules();
+			CreateForceBasedRules(FineLODMultiResConstraint);
 		}
 		else
 		{
@@ -1419,8 +1425,96 @@ void FClothConstraints::CreateCollisionConstraint(
 	}
 }
 
-void FClothConstraints::CreateForceBasedRules()
+void FClothConstraints::CreateMultiresConstraint(
+	const Softs::FCollectionPropertyConstFacade& ConfigProperties,
+	const TMap<FString, TConstArrayView<FRealSingle>>& WeightMaps,
+	const FTriangleMesh& TriangleMesh,
+	const FTriangleMesh* MultiResCoarseLODMesh,
+	const int32 MultiResCoarseLODParticleRangeId)
 {
+	if (Evolution && MultiResCoarseLODMesh && MultiResCoarseLODParticleRangeId != INDEX_NONE)
+	{
+		if (Softs::FMultiResConstraints::IsEnabled(ConfigProperties))
+		{
+			TArray<TVec4<Softs::FSolverReal>> CoarseToFinePositionBaryCoordsAndDist;
+			TArray<TVec3<int32>> CoarseToFineSourceMeshVertIndices;
+			const Softs::FSolverParticlesRange& FineParticles = Evolution->GetSoftBodyParticles(ParticleRangeId);
+			{
+				TRACE_CPUPROFILER_EVENT_SCOPE(BuildMultiResTransitionData);
+				// TODO: Cache this data in the asset.
+				const FPointWeightMap* const MaxDistances = nullptr;  // No need to update the vertex contribution on the transition maps
+				constexpr bool bUseSmoothTransitions = false;         // Smooth transitions are only used at rendering for now and not during LOD transitions
+				constexpr bool bUseMultipleInfluences = false;        // Multiple influences must not be used for LOD transitions
+				constexpr float SkinningKernelRadius = 0.f;           // KernelRadius is only required when using multiple influences
+				static_assert(std::is_same<float, Softs::FSolverReal>::value);
+				static_assert(sizeof(Softs::FSolverVec3) == sizeof(FVector3f));
+
+				TConstArrayView<FVector3f> FinePositions((const FVector3f*)FineParticles.XArray().GetData(), FineParticles.Size());
+				TArray<uint32> FineIndices;
+				FineIndices.Reserve(3 * TriangleMesh.GetNumElements());
+				for (const TVec3<int32>& Element : TriangleMesh.GetElements())
+				{
+					FineIndices.Add(Element[0]);
+					FineIndices.Add(Element[1]);
+					FineIndices.Add(Element[2]);
+				}
+
+				const ClothingMeshUtils::ClothMeshDesc FineLodDesc(FinePositions, TConstArrayView<uint32>(FineIndices));
+
+				const Softs::FSolverParticlesRange& CoarseParticles = Evolution->GetSoftBodyParticles(MultiResCoarseLODParticleRangeId);
+				TConstArrayView<FVector3f> CoarsePositions((const FVector3f*)CoarseParticles.XArray().GetData(), CoarseParticles.Size());
+				TConstArrayView<Softs::FSolverReal> CoarseInvM = CoarseParticles.GetInvM();
+				TArray<uint32> CoarseIndices;
+				CoarseIndices.Reserve(3 * MultiResCoarseLODMesh->GetNumElements());
+				for (const TVec3<int32>& Element : MultiResCoarseLODMesh->GetElements())
+				{
+					if (CoarseInvM[Element[0]] != 0.f || CoarseInvM[Element[1]] != 0.f || CoarseInvM[Element[2]] != 0.f)
+					{
+						CoarseIndices.Add(Element[0]);
+						CoarseIndices.Add(Element[1]);
+						CoarseIndices.Add(Element[2]);
+					}
+				}
+
+				if (CoarseIndices.IsEmpty())
+				{
+					return;
+				}
+
+				const ClothingMeshUtils::ClothMeshDesc CoarseLodDesc(CoarsePositions, TConstArrayView<uint32>(CoarseIndices));
+				TArray<FMeshToMeshVertData> TransitionData;
+				ClothingMeshUtils::GenerateMeshToMeshVertData(TransitionData, FineLodDesc, CoarseLodDesc, MaxDistances, bUseSmoothTransitions, bUseMultipleInfluences, SkinningKernelRadius);
+
+				CoarseToFinePositionBaryCoordsAndDist.Reserve(TransitionData.Num());
+				CoarseToFineSourceMeshVertIndices.Reserve(TransitionData.Num());
+				for (const FMeshToMeshVertData& Data : TransitionData)
+				{
+					CoarseToFinePositionBaryCoordsAndDist.Emplace(Data.PositionBaryCoordsAndDist.X, Data.PositionBaryCoordsAndDist.Y, Data.PositionBaryCoordsAndDist.Z, Data.PositionBaryCoordsAndDist.W);
+					CoarseToFineSourceMeshVertIndices.Emplace(Data.SourceMeshVertIndices[0], Data.SourceMeshVertIndices[1], Data.SourceMeshVertIndices[2]);
+				}
+			}
+
+			MultiResConstraints = MakeShared<Softs::FMultiResConstraints>(
+				FineParticles,
+				MultiResCoarseLODParticleRangeId,
+				*MultiResCoarseLODMesh,
+				MoveTemp(CoarseToFinePositionBaryCoordsAndDist),
+				MoveTemp(CoarseToFineSourceMeshVertIndices),
+				WeightMaps,
+				ConfigProperties);
+			++NumConstraintRules;
+			++NumConstraintInits;
+		}
+	}
+}
+
+void FClothConstraints::CreateForceBasedRules(const TSharedPtr<Softs::FMultiResConstraints>& FineLODMultiResConstraint)
+{
+	if (FineLODMultiResConstraint)
+	{
+		++NumPostprocessingConstraintRules;
+	}
+
 	FRuleCreator RuleCreator(this);
 
 	if (ExternalForces)
@@ -1473,6 +1567,23 @@ void FClothConstraints::CreateForceBasedRules()
 		// TODO: Linear System
 	}
 
+	if (MultiResConstraints)
+	{
+		constexpr bool bInitParticles = false;
+		RuleCreator.AddXPBDParallelInitRule<bInitParticles>(MultiResConstraints.Get());
+		constexpr bool bPostCollisions = false;
+		RuleCreator.AddPerIterationPBDConstraintRule_Apply<bPostCollisions>(MultiResConstraints.Get());
+	}
+
+	if (FineLODMultiResConstraint)
+	{
+		RuleCreator.AddPostSubstepConstraintRule(
+			[FineLODMultiResConstraint, this](Softs::FSolverParticlesRange& Particles, const Softs::FSolverReal Dt, const Softs::ESolverMode SolverMode)
+		{
+			FineLODMultiResConstraint->UpdateFineTargets(Particles);
+		});
+	}
+
 	if (XStretchBiasConstraints)
 	{
 		constexpr bool bInitParticles = false;
@@ -2174,6 +2285,10 @@ void FClothConstraints::Update(
 	{
 		SelfCollisionSphereConstraints->SetProperties(ConfigProperties, VertexSets);
 	}
+	if (MultiResConstraints)
+	{
+		MultiResConstraints->SetProperties(ConfigProperties, WeightMaps);
+	}
 
 	bool bUsePointBasedWindModel = false;
 	if (ExternalForces)
diff --git a/Engine/Plugins/Experimental/ChaosCloth/Source/Chaos/Private/ChaosCloth/ChaosClothVisualization.cpp b/Engine/Plugins/Experimental/ChaosCloth/Source/Chaos/Private/ChaosCloth/ChaosClothVisualization.cpp
--- a/Engine/Plugins/Experimental/ChaosCloth/Source/Chaos/Private/ChaosCloth/ChaosClothVisualization.cpp
+++ b/Engine/Plugins/Experimental/ChaosCloth/Source/Chaos/Private/ChaosCloth/ChaosClothVisualization.cpp
@@ -24,6 +24,7 @@
 #include "Chaos/PBDSpringConstraints.h"
 #include "Chaos/PBDTriangleMeshCollisions.h"
 #include "Chaos/SoftsExternalForces.h"
+#include "Chaos/SoftsMultiResConstraints.h"
 #include "Chaos/XPBDBendingConstraints.h"
 #include "Chaos/XPBDSpringConstraints.h"
 #include "Chaos/XPBDAnisotropicBendingConstraints.h"
@@ -1012,6 +1013,72 @@ FLinearColor PseudoRandomColor(int32 NumColorRotations)
 		}
 	}
 
+	void FClothVisualization::DrawMultiResConstraint(FPrimitiveDrawInterface* PDI) const
+	{
+		if (!Solver)
+		{
+			return;
+		}
+
+		const FVec3& LocalSpaceLocation = Solver->GetLocalSpaceLocation();
+		for (const FClothingSimulationCloth* const Cloth : Solver->GetCloths())
+		{
+			const int32 ParticleRangeId = Cloth->GetParticleRangeId(Solver);
+			if (ParticleRangeId == INDEX_NONE)
+			{
+				continue;
+			}
+			const FClothConstraints& ClothConstraints = Solver->GetClothConstraints(ParticleRangeId);
+			if (const Softs::FMultiResConstraints* const MultiResConstraints = ClothConstraints.GetMultiResConstraints().Get())
+			{
+
+				const int32 CoarseParticleRangeId = MultiResConstraints->GetCoarseSoftBodyId();
+				const FTriangleMesh& CoarseMesh = MultiResConstraints->GetCoarseMesh();
+				TConstArrayView<Softs::FSolverVec3> CoarsePositions = Solver->GetParticleXsView(CoarseParticleRangeId);
+				TConstArrayView<Softs::FSolverReal> CoarseInvMasses = Solver->GetParticleInvMassesView(CoarseParticleRangeId);
+
+				// Draw wired coarse mesh
+				static const FLinearColor DynamicColor = FColor::White;
+				static const FLinearColor KinematicColor = FColor::Purple;
+				for (const TVec3<int32>& Element : CoarseMesh.GetElements())
+				{
+					const FVector Pos0 = LocalSpaceLocation + FVector(CoarsePositions[Element.X]);
+					const FVector Pos1 = LocalSpaceLocation + FVector(CoarsePositions[Element.Y]);
+					const FVector Pos2 = LocalSpaceLocation + FVector(CoarsePositions[Element.Z]);
+					const bool bIsKinematic0 = (CoarseInvMasses[Element.X] == (Softs::FSolverReal)0.);
+					const bool bIsKinematic1 = (CoarseInvMasses[Element.Y] == (Softs::FSolverReal)0.);
+					const bool bIsKinematic2 = (CoarseInvMasses[Element.Z] == (Softs::FSolverReal)0.);
+					if (bIsKinematic0 && bIsKinematic1 && bIsKinematic2)
+					{
+						continue;
+					}
+
+					DrawLine(PDI, Pos0, Pos1, bIsKinematic0 && bIsKinematic1 ? KinematicColor : DynamicColor);
+					DrawLine(PDI, Pos1, Pos2, bIsKinematic1 && bIsKinematic2 ? KinematicColor : DynamicColor);
+					DrawLine(PDI, Pos2, Pos0, bIsKinematic2 && bIsKinematic0 ? KinematicColor : DynamicColor);
+				}
+
+				// Draw springs to targets
+				static const FLinearColor Red(0.3f, 0.f, 0.f);
+				static const FLinearColor Brown(0.1f, 0.05f, 0.f);
+				const TConstArrayView<Softs::FSolverVec3> Positions = Solver->GetParticleXsView(ParticleRangeId);
+				const TConstArrayView<Softs::FSolverReal> InvMasses = Solver->GetParticleInvMassesView(ParticleRangeId);
+				const TArray<Softs::FSolverVec3>& TargetPositions = MultiResConstraints->GetFineTargetPositions();
+				for (int32 Index = 0; Index < TargetPositions.Num(); ++Index)
+				{
+					if (InvMasses[Index] != (FSolverReal)0.f && MultiResConstraints->IsConstraintActive(Index))
+					{
+						const FVector P1 = LocalSpaceLocation + FVector(Positions[Index]);
+						const FVector P2 = LocalSpaceLocation + FVector(TargetPositions[Index]);
+
+						DrawPoint(PDI, P2, Red, nullptr, 2.f);
+						DrawLine(PDI, P1, P2, Brown);
+					}
+				}
+			}
+		}
+	}
+
 	void FClothVisualization::DrawAnimNormals(FPrimitiveDrawInterface* PDI) const
 	{
 		if (!Solver)
diff --git a/Engine/Plugins/Experimental/ChaosCloth/Source/Chaos/Private/ChaosCloth/ChaosClothingSimulation.cpp b/Engine/Plugins/Experimental/ChaosCloth/Source/Chaos/Private/ChaosCloth/ChaosClothingSimulation.cpp
--- a/Engine/Plugins/Experimental/ChaosCloth/Source/Chaos/Private/ChaosCloth/ChaosClothingSimulation.cpp
+++ b/Engine/Plugins/Experimental/ChaosCloth/Source/Chaos/Private/ChaosCloth/ChaosClothingSimulation.cpp
@@ -41,6 +41,7 @@
 #include "Chaos/VelocityField.h"
 #include "Chaos/XPBDSpringConstraints.h"
 #include "Chaos/XPBDBendingConstraints.h"
+#include "Chaos/SoftsMultiResConstraints.h"
 #endif
 
 #if INTEL_ISPC && !UE_BUILD_SHIPPING
@@ -130,6 +131,7 @@ namespace ClothingSimulationConsole
 					bChaos_Bending_ISPC_Enabled =
 					bChaos_XPBDSpring_ISPC_Enabled =
 					bChaos_XPBDBending_ISPC_Enabled =
+					bChaos_MultiRes_ISPC_Enabled =
 					bEnableISPC;
 				return;
 			}
diff --git a/Engine/Plugins/Experimental/ChaosCloth/Source/Chaos/Private/ChaosCloth/ChaosClothingSimulationCloth.cpp b/Engine/Plugins/Experimental/ChaosCloth/Source/Chaos/Private/ChaosCloth/ChaosClothingSimulationCloth.cpp
--- a/Engine/Plugins/Experimental/ChaosCloth/Source/Chaos/Private/ChaosCloth/ChaosClothingSimulationCloth.cpp
+++ b/Engine/Plugins/Experimental/ChaosCloth/Source/Chaos/Private/ChaosCloth/ChaosClothingSimulationCloth.cpp
@@ -37,6 +37,7 @@ namespace ClothingSimulationClothDefault
 	constexpr float GravityZOverride = Softs::FExternalForces::DefaultGravityZOverride; // -980.665f;
 	constexpr float VelocityScale = 0.75f;
 	constexpr float FictitiousAngularScale = Softs::FExternalForces::DefaultFictitiousAngularScale; // 1.f;
+	constexpr int32 MultiResCoarseLODIndex = INDEX_NONE;
 }
 
 namespace ClothingSimulationClothConsoleVariables
@@ -72,6 +73,7 @@ struct FClothingSimulationCloth::FLODData
 		int32 LODIndex;
 		int32 ParticleRangeId;
 		FTriangleMesh OffsetTriangleMesh; // Only used if using PBD solver
+		int32 MultiResCoarseLODIndex = INDEX_NONE;
 	};
 	TMap<FClothingSimulationSolver*, FSolverData> SolverData;
 
@@ -93,6 +95,7 @@ struct FClothingSimulationCloth::FLODData
 
 	static FTriangleMesh BuildTriangleMesh(const TConstArrayView<uint32>& Indices, const int32 NumParticles);
 
+	void AddParticles(FClothingSimulationSolver* Solver, FClothingSimulationCloth* Cloth, int32 LODIndex);
 	void Add(FClothingSimulationSolver* Solver, FClothingSimulationCloth* Cloth, int32 LODIndex);
 	void Remove(FClothingSimulationSolver* Solver);
 
@@ -106,6 +109,9 @@ struct FClothingSimulationCloth::FLODData
 
 	UE_CHAOS_DECLARE_INDEXLESS_PROPERTYCOLLECTION_NAME(MaxDistance, float);
 	UE_CHAOS_DECLARE_INDEXLESS_PROPERTYCOLLECTION_NAME(MassValue, float);
+	UE_CHAOS_DECLARE_INDEXLESS_PROPERTYCOLLECTION_NAME(MultiResCoarseLODIndex, int32);
+	UE_CHAOS_DECLARE_INDEXLESS_PROPERTYCOLLECTION_NAME(IsCoarseMultiResLOD, bool);
+
 };
 
 FClothingSimulationCloth::FLODData::FLODData(
@@ -152,7 +158,7 @@ FTriangleMesh FClothingSimulationCloth::FLODData::BuildTriangleMesh(const TConst
 	return OutTriangleMesh;
 }
 
-void FClothingSimulationCloth::FLODData::Add(FClothingSimulationSolver* Solver, FClothingSimulationCloth* Cloth, int32 InLODIndex)
+void FClothingSimulationCloth::FLODData::AddParticles(FClothingSimulationSolver* Solver, FClothingSimulationCloth* Cloth, int32 InLODIndex)
 {
 	check(Solver);
 	check(Cloth);
@@ -162,26 +168,13 @@ void FClothingSimulationCloth::FLODData::Add(FClothingSimulationSolver* Solver,
 	check(!SolverData.Find(Solver));
 	FSolverData& SolverDatum = SolverData.Add(Solver);
 	SolverDatum.LODIndex = InLODIndex;
-	int32& ParticleRangeId = SolverDatum.ParticleRangeId;
 	// Add particles
-	ParticleRangeId = Solver->AddParticles(NumParticles, Cloth->GroupId);
-
+	SolverDatum.ParticleRangeId = Solver->AddParticles(NumParticles, Cloth->GroupId);
 	if (!NumParticles)
 	{
 		return;
 	}
 
-PRAGMA_DISABLE_DEPRECATION_WARNINGS  // TODO: CHAOS_IS_CLOTHINGSIMULATIONMESH_ABSTRACT
-	// Update source mesh for this LOD, this is required prior to reset the start pose
-	Cloth->Mesh->Update(Solver, INDEX_NONE, InLODIndex, 0, ParticleRangeId);
-	
-	// Retrieve the component's scale
-	const Softs::FSolverReal MeshScale = Cloth->Mesh->GetScale();
-PRAGMA_ENABLE_DEPRECATION_WARNINGS
-
-	// Reset the particles start pose before setting up mass and constraints
-	ResetStartPose(Solver);
-
 	if (!Solver->IsForceBasedSolver())
 	{
 		const TArray<TVec3<int32>>& Elements = NoOffsetTriangleMesh.GetElements();
@@ -190,17 +183,25 @@ PRAGMA_ENABLE_DEPRECATION_WARNINGS
 		for (const TVec3<int32>& Element : Elements)
 		{
 			OffsetElements.Add({
-				Element[0] + ParticleRangeId,
-				Element[1] + ParticleRangeId,
-				Element[2] + ParticleRangeId });
+				Element[0] + SolverDatum.ParticleRangeId,
+				Element[1] + SolverDatum.ParticleRangeId,
+				Element[2] + SolverDatum.ParticleRangeId });
 		}
 
-		SolverDatum.OffsetTriangleMesh.Init(MoveTemp(OffsetElements), ParticleRangeId, ParticleRangeId + NumParticles - 1);
+		SolverDatum.OffsetTriangleMesh.Init(MoveTemp(OffsetElements), SolverDatum.ParticleRangeId, SolverDatum.ParticleRangeId + NumParticles - 1);
 		SolverDatum.OffsetTriangleMesh.GetPointToTriangleMap(); // Builds map for later use by GetPointNormals(), and the velocity fields
 	}
 
 	const FTriangleMesh& TriangleMesh = Solver->IsForceBasedSolver() ? NoOffsetTriangleMesh : SolverDatum.OffsetTriangleMesh;
 
+	PRAGMA_DISABLE_DEPRECATION_WARNINGS  // TODO: CHAOS_IS_CLOTHINGSIMULATIONMESH_ABSTRACT
+	// Update source mesh for this LOD, this is required prior to reset the start pose
+	Cloth->Mesh->Update(Solver, INDEX_NONE, InLODIndex, 0, SolverDatum.ParticleRangeId);
+	PRAGMA_ENABLE_DEPRECATION_WARNINGS
+
+	// Reset the particles start pose before setting up mass and constraints
+	ResetStartPose(Solver);
+
 	// Initialize the normals, in case the sim data is queried before the simulation steps
 	UpdateNormals(Solver);
 
@@ -218,9 +219,9 @@ PRAGMA_ENABLE_DEPRECATION_WARNINGS
 	static const FRealSingle KinematicDistanceThreshold = 0.1f;  // TODO: This is not the same value as set in the painting UI but we might want to expose this value as parameter
 	auto KinematicPredicate =
 		[&MaxDistances](int32 Index)
-		{
-			return MaxDistances.GetValue(Index) < KinematicDistanceThreshold;
-		};
+	{
+		return MaxDistances.GetValue(Index) < KinematicDistanceThreshold;
+	};
 
 	const int32 MassMode = ConfigProperties.GetValue<int32>(TEXT("MassMode"), ClothingSimulationClothDefault::MassMode);
 
@@ -235,22 +236,89 @@ PRAGMA_ENABLE_DEPRECATION_WARNINGS
 	{
 		const FVector2f MassValue = GetWeightedFloatMassValue(ConfigProperties, ClothingSimulationClothDefault::MassValue);
 		const TConstArrayView<float> MassValueMultipliers = WeightMaps.FindRef(GetMassValueString(ConfigProperties, MassValueName.ToString()));
-		Solver->SetParticleMassUniform(ParticleRangeId, MassValue, MassValueMultipliers, MinPerParticleMass, TriangleMesh, KinematicPredicate);
+		Solver->SetParticleMassUniform(SolverDatum.ParticleRangeId, MassValue, MassValueMultipliers, MinPerParticleMass, TriangleMesh, KinematicPredicate);
 	}
-		break;
+	break;
 	case EMassMode::TotalMass:
 	{
 		const FRealSingle MassValue = GetMassValue(ConfigProperties, ClothingSimulationClothDefault::MassValue);
-		Solver->SetParticleMassFromTotalMass(ParticleRangeId, MassValue, MinPerParticleMass, TriangleMesh, KinematicPredicate);
+		Solver->SetParticleMassFromTotalMass(SolverDatum.ParticleRangeId, MassValue, MinPerParticleMass, TriangleMesh, KinematicPredicate);
 	}
-		break;
+	break;
 	case EMassMode::Density:
 	{
 		const FVector2f MassValue = GetWeightedFloatMassValue(ConfigProperties, ClothingSimulationClothDefault::MassValue);
 		const TConstArrayView<float> MassValueMultipliers = WeightMaps.FindRef(GetMassValueString(ConfigProperties, MassValueName.ToString()));
-		Solver->SetParticleMassFromDensity(ParticleRangeId, MassValue, MassValueMultipliers, MinPerParticleMass, TriangleMesh, KinematicPredicate);
+		Solver->SetParticleMassFromDensity(SolverDatum.ParticleRangeId, MassValue, MassValueMultipliers, MinPerParticleMass, TriangleMesh, KinematicPredicate);
 	}
-		break;
+	break;
+	}
+}
+
+void FClothingSimulationCloth::FLODData::Add(FClothingSimulationSolver* Solver, FClothingSimulationCloth* Cloth, int32 InLODIndex)
+{
+	check(Solver);
+	check(Cloth);
+	check(Cloth->Mesh);
+
+	FSolverData& SolverDatum = SolverData.FindChecked(Solver);
+	check(SolverDatum.LODIndex == InLODIndex);
+	if (!NumParticles)
+	{
+		return;
+	}
+
+	check(SolverDatum.ParticleRangeId != INDEX_NONE);
+	const int32 ParticleRangeId = SolverDatum.ParticleRangeId;
+
+	PRAGMA_DISABLE_DEPRECATION_WARNINGS  // TODO: CHAOS_IS_CLOTHINGSIMULATIONMESH_ABSTRACT
+	// Retrieve the component's scale
+	const Softs::FSolverReal MeshScale = Cloth->Mesh->GetScale();
+	PRAGMA_ENABLE_DEPRECATION_WARNINGS
+
+	const FTriangleMesh& TriangleMesh = Solver->IsForceBasedSolver() ? NoOffsetTriangleMesh : SolverDatum.OffsetTriangleMesh;
+
+	// Retrieve config properties
+	check(Cloth->Config);
+	const Softs::FCollectionPropertyFacade& ConfigProperties = Cloth->Config->GetProperties(InLODIndex);
+
+	// Gather multires constraint data.
+	TSharedPtr<Softs::FMultiResConstraints> FineLODMultiResConstraint;
+	const FTriangleMesh* CoarseLODTriangleMesh = nullptr;
+	int32 CoarseLODParticleRangeId = INDEX_NONE;
+	if (Solver->IsForceBasedSolver())
+	{
+		// Multi-Res is only supported with ForceBasedSolver
+		if (InLODIndex == 0)
+		{
+			// Only allow LOD0 to be a fine LOD for now.
+			const int32 MultiResCoarseLODIndex = GetMultiResCoarseLODIndex(ConfigProperties, ClothingSimulationClothDefault::MultiResCoarseLODIndex);
+			if (MultiResCoarseLODIndex != INDEX_NONE && MultiResCoarseLODIndex != InLODIndex)
+			{
+				if (Cloth->Config->IsValidLOD(MultiResCoarseLODIndex) && Cloth->LODData.IsValidIndex(MultiResCoarseLODIndex))
+				{
+					// Check if coarse lod is setup correctly.
+					const Softs::FCollectionPropertyFacade& CoarseConfigProperties = Cloth->Config->GetProperties(MultiResCoarseLODIndex);
+					if (GetIsCoarseMultiResLOD(CoarseConfigProperties, false))
+					{
+						CoarseLODTriangleMesh = &Cloth->LODData[MultiResCoarseLODIndex]->NoOffsetTriangleMesh;
+						CoarseLODParticleRangeId = Cloth->LODData[MultiResCoarseLODIndex]->SolverData.FindChecked(Solver).ParticleRangeId;
+						SolverDatum.MultiResCoarseLODIndex = MultiResCoarseLODIndex;
+					}
+				}
+			}
+		}
+		else if (GetIsCoarseMultiResLOD(ConfigProperties, false))
+		{
+			// check that fine lod is setup correctly.
+			const Softs::FCollectionPropertyFacade& FineConfigProperties = Cloth->Config->GetProperties(0);
+			const int32 MultiResCoarseLODIndex = GetMultiResCoarseLODIndex(FineConfigProperties, ClothingSimulationClothDefault::MultiResCoarseLODIndex);
+			if (MultiResCoarseLODIndex == InLODIndex)
+			{
+				const int32 FineLODParticleRangeId = Cloth->LODData[0]->SolverData.FindChecked(Solver).ParticleRangeId;
+				FineLODMultiResConstraint = Solver->GetClothConstraints(FineLODParticleRangeId).GetMultiResConstraints();
+			}
+		}
 	}
 
 	// Setup solver constraints
@@ -258,7 +326,7 @@ PRAGMA_ENABLE_DEPRECATION_WARNINGS
 
 	// Create constraints
 	const bool bEnabled = false;  // Set constraint disabled by default
-	ClothConstraints.AddRules(ConfigProperties, TriangleMesh, &PatternData, WeightMaps, VertexSets, FaceSets, FaceIntMaps, Tethers, MeshScale, bEnabled);
+	ClothConstraints.AddRules(ConfigProperties, TriangleMesh, &PatternData, WeightMaps, VertexSets, FaceSets, FaceIntMaps, Tethers, MeshScale, bEnabled, CoarseLODTriangleMesh, CoarseLODParticleRangeId, FineLODMultiResConstraint);
 
 	// Update LOD stats
 	const TConstArrayView<Softs::FSolverReal> InvMasses(Solver->GetParticleInvMasses(ParticleRangeId), NumParticles);
@@ -768,7 +836,13 @@ void FClothingSimulationCloth::Add(FClothingSimulationSolver* Solver)
 	int32& LODIndex = LODIndices.Add(Solver);
 	LODIndex = INDEX_NONE;
 
-	// Add LODs
+	// Add all particles first and in reverse order. This is necessary so that any multires coarse lods soft bodies are added first, and all particle offsets are setup
+	// when adding the LOD constraints.
+	for (int32 Index = LODData.Num() - 1; Index >= 0; --Index)
+	{
+		LODData[Index]->AddParticles(Solver, this, Index);
+	}
+	// Now add the LODs themselves. These need to go in normal order since the coarse lod needs the fine lod constraints.
 	for (int32 Index = 0; Index < LODData.Num(); ++Index)
 	{
 		LODData[Index]->Add(Solver, this, Index);
@@ -873,6 +947,7 @@ const FTriangleMesh& FClothingSimulationCloth::GetTriangleMesh(const FClothingSi
 	return LODData.IsValidIndex(LODIndex) ? (Solver->IsForceBasedSolver() ? LODData[LODIndex]->NoOffsetTriangleMesh : LODData[LODIndex]->SolverData.FindChecked(Solver).OffsetTriangleMesh) : EmptyTriangleMesh;
 }
 
+
 // Deprecated for 5.3
 const TArray<TConstArrayView<FRealSingle>>& FClothingSimulationCloth::GetWeightMaps(const FClothingSimulationSolver* Solver) const
 {
@@ -990,13 +1065,33 @@ PRAGMA_ENABLE_DEPRECATION_WARNINGS
 	check(Config);
 	const Softs::FCollectionPropertyFacade& ConfigProperties = Config->GetProperties(LODIndex);
 
+	const int32 CoarseLODIndex = LODIndex != INDEX_NONE ? LODData[LODIndex]->SolverData.FindChecked(Solver).MultiResCoarseLODIndex : INDEX_NONE;
+	const int32 CoarseParticleRangeId = CoarseLODIndex != INDEX_NONE ? GetParticleRangeId(Solver, CoarseLODIndex) : INDEX_NONE;
+	if (CoarseLODIndex != INDEX_NONE)
+	{
+		LODData[CoarseLODIndex]->Enable(Solver, true);
+PRAGMA_DISABLE_DEPRECATION_WARNINGS 
+		// TODO: interpolate/ reset when LOD switching to enable multires
+		Mesh->Update(Solver, CoarseLODIndex, CoarseLODIndex, CoarseParticleRangeId, CoarseParticleRangeId);
+PRAGMA_ENABLE_DEPRECATION_WARNINGS
+	}
+
 	// LOD Switching
 	if (LODIndex != PrevLODIndex)
 	{
 		if (PrevLODIndex != INDEX_NONE)
 		{
-			// Disable previous LOD's particles
-			LODData[PrevLODIndex]->Enable(Solver, false);
+			if (PrevLODIndex != CoarseLODIndex)
+			{
+				// Disable previous LOD's particles
+				LODData[PrevLODIndex]->Enable(Solver, false);
+			}
+			const int32 PrevCoarseLODIndex = LODData[PrevLODIndex]->SolverData.FindChecked(Solver).MultiResCoarseLODIndex;
+			if (PrevCoarseLODIndex != INDEX_NONE && PrevCoarseLODIndex != CoarseLODIndex && PrevCoarseLODIndex != LODIndex)
+			{
+				// Disable previous coarse LOD's particles
+				LODData[PrevCoarseLODIndex]->Enable(Solver, false);
+			}
 		}
 		if (LODIndex != INDEX_NONE)
 		{
@@ -1052,6 +1147,11 @@ PRAGMA_ENABLE_DEPRECATION_WARNINGS
 			{
 				Collider->ResetStartPose(Solver, this);
 			}
+
+			if (CoarseLODIndex != INDEX_NONE)
+			{
+				LODData[CoarseLODIndex]->ResetStartPose(Solver);
+			}
 			UE_LOG(LogChaosCloth, VeryVerbose, TEXT("Cloth in group Id %d Needs reset."), GroupId);
 		}
 		else if (bNeedsTeleport)
@@ -1080,6 +1180,10 @@ PRAGMA_ENABLE_DEPRECATION_WARNINGS
 		if (Solver->IsForceBasedSolver())
 		{
 			Solver->SetProperties(ParticleRangeId, ConfigProperties, LODData[LODIndex]->WeightMaps);
+			if (CoarseLODIndex != INDEX_NONE)
+			{
+				Solver->SetProperties(CoarseParticleRangeId, Config->GetProperties(CoarseLODIndex), LODData[CoarseLODIndex]->WeightMaps);
+			}
 		}
 		else
 		{
@@ -1121,6 +1225,10 @@ PRAGMA_ENABLE_DEPRECATION_WARNINGS
 		if (!Config->IsLegacySingleLOD())
 		{
 			LODData[LODIndex]->Update(Solver, this);
+			if (CoarseLODIndex != INDEX_NONE)
+			{
+				LODData[CoarseLODIndex]->Update(Solver, this);
+			}
 		}
 	}
 
diff --git a/Engine/Plugins/Experimental/ChaosCloth/Source/Chaos/Public/ChaosCloth/ChaosClothConstraints.h b/Engine/Plugins/Experimental/ChaosCloth/Source/Chaos/Public/ChaosCloth/ChaosClothConstraints.h
--- a/Engine/Plugins/Experimental/ChaosCloth/Source/Chaos/Public/ChaosCloth/ChaosClothConstraints.h
+++ b/Engine/Plugins/Experimental/ChaosCloth/Source/Chaos/Public/ChaosCloth/ChaosClothConstraints.h
@@ -71,7 +71,10 @@ PRAGMA_ENABLE_DEPRECATION_WARNINGS
 			const TMap<FString, TConstArrayView<int32>>& FaceIntMaps,
 			const TArray<TConstArrayView<TTuple<int32, int32, FRealSingle>>>& Tethers,
 			Softs::FSolverReal MeshScale,
-			bool bEnabled);
+			bool bEnabled,
+			const FTriangleMesh* MultiResCoarseLODMesh = nullptr,
+			const int32 MultiResCoarseLODParticleRangeId = INDEX_NONE,
+			const TSharedPtr<Softs::FMultiResConstraints>& FineLODMultiResConstraint = TSharedPtr<Softs::FMultiResConstraints>(nullptr));
 
 		UE_DEPRECATED(5.4, "Use AddRules() with WeightMaps, VertexSets, FaceSets, FaceIntMaps, and optional PatternData instead.")
 		void AddRules(
@@ -146,6 +149,7 @@ PRAGMA_ENABLE_DEPRECATION_WARNINGS
 		const TSharedPtr<Softs::FVelocityAndPressureField>& GetVelocityAndPressureField() const { return VelocityAndPressureField; }
 		const TSharedPtr<Softs::FExternalForces>& GetExternalForces() const { return ExternalForces; }
 		const TSharedPtr<Softs::FPBDSoftBodyCollisionConstraint>& GetCollisionConstraint() const { return CollisionConstraint; }
+		const TSharedPtr<Softs::FMultiResConstraints>& GetMultiResConstraints() const { return MultiResConstraints; }
 		// ---- End of debug functions ----
 
 	private:
@@ -198,8 +202,14 @@ PRAGMA_ENABLE_DEPRECATION_WARNINGS
 		void CreateCollisionConstraint(
 			const Softs::FCollectionPropertyConstFacade& ConfigProperties,
 			Softs::FSolverReal MeshScale );
+		void CreateMultiresConstraint(
+			const Softs::FCollectionPropertyConstFacade& ConfigProperties,
+			const TMap<FString, TConstArrayView<FRealSingle>>& WeightMaps,
+			const FTriangleMesh& TriangleMesh,
+			const FTriangleMesh* MultiResCoarseLODMesh,
+			const int32 MultiResCoarseLODParticleRangeId);
 
-		void CreateForceBasedRules();
+		void CreateForceBasedRules(const TSharedPtr<Softs::FMultiResConstraints>& FineLODMultiResConstraint);
 		void CreatePBDRules();
 		void CreateGSRules();
 		void GetGSNumRules();
@@ -230,6 +240,7 @@ PRAGMA_ENABLE_DEPRECATION_WARNINGS
 		TSharedPtr<Softs::FPBDSelfCollisionSphereConstraints> SelfCollisionSphereConstraints;
 		TSharedPtr<Softs::FGaussSeidelMasterConstraint<Softs::FSolverReal, Softs::FSolverParticles>> GSMasterConstraint;
 		TSharedPtr<Softs::FGaussSeidelCorotatedCodimensionalConstraints<Softs::FSolverReal, Softs::FSolverParticles>> GSCorotatedCodimensionalConstraint;
+		TSharedPtr<Softs::FMultiResConstraints> MultiResConstraints;
 		//~ Begin Force-based solver only constraints
 		Softs::FSolverVec3 SolverWindVelocity; // Set from solver and added to wind from the config
 		TSharedPtr<Softs::FVelocityAndPressureField> VelocityAndPressureField;
diff --git a/Engine/Plugins/Experimental/ChaosCloth/Source/Chaos/Public/ChaosCloth/ChaosClothVisualization.h b/Engine/Plugins/Experimental/ChaosCloth/Source/Chaos/Public/ChaosCloth/ChaosClothVisualization.h
--- a/Engine/Plugins/Experimental/ChaosCloth/Source/Chaos/Public/ChaosCloth/ChaosClothVisualization.h
+++ b/Engine/Plugins/Experimental/ChaosCloth/Source/Chaos/Public/ChaosCloth/ChaosClothVisualization.h
@@ -53,6 +53,7 @@ namespace Chaos
 		CHAOSCLOTH_API void DrawBounds(FPrimitiveDrawInterface* PDI = nullptr) const;
 		CHAOSCLOTH_API void DrawGravity(FPrimitiveDrawInterface* PDI = nullptr) const;
 		CHAOSCLOTH_API void DrawFictitiousAngularForces(FPrimitiveDrawInterface* PDI = nullptr) const;
+		CHAOSCLOTH_API void DrawMultiResConstraint(FPrimitiveDrawInterface* PDI = nullptr) const;
 
 #else  // #if CHAOS_DEBUG_DRAW
 		void SetSolver(const ::Chaos::FClothingSimulationSolver* /*InSolver*/) {}
@@ -79,6 +80,7 @@ namespace Chaos
 		void DrawBounds(FPrimitiveDrawInterface* /*PDI*/ = nullptr) const {}
 		void DrawGravity(FPrimitiveDrawInterface* /*PDI*/ = nullptr) const {}
 		void DrawFictitiousAngularForces(FPrimitiveDrawInterface* /*PDI*/ = nullptr) const {}
+		void DrawMultiResConstraint(FPrimitiveDrawInterface* /*PDI*/ = nullptr) const {}
 
 #endif  // #if CHAOS_DEBUG_DRAW
 
diff --git a/Engine/Plugins/Experimental/ChaosCloth/Source/Chaos/Public/ChaosCloth/ChaosClothingSimulationConfig.h b/Engine/Plugins/Experimental/ChaosCloth/Source/Chaos/Public/ChaosCloth/ChaosClothingSimulationConfig.h
--- a/Engine/Plugins/Experimental/ChaosCloth/Source/Chaos/Public/ChaosCloth/ChaosClothingSimulationConfig.h
+++ b/Engine/Plugins/Experimental/ChaosCloth/Source/Chaos/Public/ChaosCloth/ChaosClothingSimulationConfig.h
@@ -82,6 +82,10 @@ namespace Chaos
 
 		bool IsLegacySingleLOD() const { return bIsLegacySingleLOD; }
 		int32 GetNumLODs() const { return PropertyCollections.Num(); }
+		bool IsValidLOD(int32 LODIndex) const
+		{
+			return bIsLegacySingleLOD ? true : Properties.IsValidIndex(LODIndex);
+		}
 
 		template<typename FunctionType>
 		void ForAllProperties(FunctionType Func)
diff --git a/Engine/Plugins/Experimental/ChaosClothAssetEditor/Source/ChaosClothAssetDataflowNodes/Private/ChaosClothAsset/ChaosClothAssetDataflowNodesModule.cpp b/Engine/Plugins/Experimental/ChaosClothAssetEditor/Source/ChaosClothAssetDataflowNodes/Private/ChaosClothAsset/ChaosClothAssetDataflowNodesModule.cpp
--- a/Engine/Plugins/Experimental/ChaosClothAssetEditor/Source/ChaosClothAssetDataflowNodes/Private/ChaosClothAsset/ChaosClothAssetDataflowNodesModule.cpp
+++ b/Engine/Plugins/Experimental/ChaosClothAssetEditor/Source/ChaosClothAssetDataflowNodes/Private/ChaosClothAsset/ChaosClothAssetDataflowNodesModule.cpp
@@ -38,6 +38,7 @@
 #include "ChaosClothAsset/SimulationSelfCollisionConfigNode.h"
 #include "ChaosClothAsset/SimulationSelfCollisionSpheresConfigNode.h"
 #include "ChaosClothAsset/SimulationSolverConfigNode.h"
+#include "ChaosClothAsset/SimulationMultiResConfigNode.h"
 #include "ChaosClothAsset/SimulationStretchConfigNode.h"
 #include "ChaosClothAsset/SimulationStretchOverrideConfigNode.h"
 #include "ChaosClothAsset/SimulationVelocityScaleConfigNode.h"
@@ -105,6 +106,7 @@ namespace UE::Chaos::ClothAsset
 			DATAFLOW_NODE_REGISTER_CREATION_FACTORY(FChaosClothAssetSimulationSelfCollisionConfigNode);
 			DATAFLOW_NODE_REGISTER_CREATION_FACTORY(FChaosClothAssetSimulationSelfCollisionSpheresConfigNode);
 			DATAFLOW_NODE_REGISTER_CREATION_FACTORY(FChaosClothAssetSimulationSolverConfigNode);
+			DATAFLOW_NODE_REGISTER_CREATION_FACTORY(FChaosClothAssetSimulationMultiResConfigNode);
 			DATAFLOW_NODE_REGISTER_CREATION_FACTORY(FChaosClothAssetSimulationVelocityScaleConfigNode);
 			DATAFLOW_NODE_REGISTER_CREATION_FACTORY(FChaosClothAssetSimulationXPBDAnisoBendingConfigNode);
 			DATAFLOW_NODE_REGISTER_CREATION_FACTORY(FChaosClothAssetSimulationXPBDAnisoSpringConfigNode);
diff --git a/Engine/Plugins/Experimental/ChaosClothAssetEditor/Source/ChaosClothAssetDataflowNodes/Private/ChaosClothAsset/SimulationMultiResConfigNode.cpp b/Engine/Plugins/Experimental/ChaosClothAssetEditor/Source/ChaosClothAssetDataflowNodes/Private/ChaosClothAsset/SimulationMultiResConfigNode.cpp
new file mode 100644
--- /dev/null
+++ b/Engine/Plugins/Experimental/ChaosClothAssetEditor/Source/ChaosClothAssetDataflowNodes/Private/ChaosClothAsset/SimulationMultiResConfigNode.cpp
@@ -0,0 +1,30 @@
+// Copyright Epic Games, Inc. All Rights Reserved.
+
+#include "ChaosClothAsset/SimulationMultiResConfigNode.h"
+#include "Chaos/CollectionPropertyFacade.h"
+
+#include UE_INLINE_GENERATED_CPP_BY_NAME(SimulationMultiResConfigNode)
+
+FChaosClothAssetSimulationMultiResConfigNode::FChaosClothAssetSimulationMultiResConfigNode(const Dataflow::FNodeParameters& InParam, FGuid InGuid)
+	: FChaosClothAssetSimulationBaseConfigNode(InParam, InGuid)
+{
+	RegisterCollectionConnections();
+	RegisterInputConnection(&MultiResStiffness.WeightMap);
+	RegisterInputConnection(&MultiResVelocityTargetStiffness.WeightMap);
+}
+
+void FChaosClothAssetSimulationMultiResConfigNode::AddProperties(FPropertyHelper& PropertyHelper) const
+{
+	if (bIsFineLOD)
+	{
+		PropertyHelper.SetProperty(this, &MultiResCoarseLODIndex);
+		PropertyHelper.SetPropertyBool(this, &bMultiResUseXPBD);
+		PropertyHelper.SetPropertyWeighted(this, &MultiResStiffness);
+		PropertyHelper.SetPropertyWeighted(this, &MultiResVelocityTargetStiffness);
+	}
+	if (bIsCoarseMultiResLOD)
+	{
+		PropertyHelper.SetPropertyBool(this, &bIsCoarseMultiResLOD);
+	}
+}
+
diff --git a/Engine/Plugins/Experimental/ChaosClothAssetEditor/Source/ChaosClothAssetDataflowNodes/Public/ChaosClothAsset/SimulationMultiResConfigNode.h b/Engine/Plugins/Experimental/ChaosClothAssetEditor/Source/ChaosClothAssetDataflowNodes/Public/ChaosClothAsset/SimulationMultiResConfigNode.h
new file mode 100644
--- /dev/null
+++ b/Engine/Plugins/Experimental/ChaosClothAssetEditor/Source/ChaosClothAssetDataflowNodes/Public/ChaosClothAsset/SimulationMultiResConfigNode.h
@@ -0,0 +1,45 @@
+// Copyright Epic Games, Inc. All Rights Reserved.
+
+#pragma once 
+
+#include "ChaosClothAsset/SimulationBaseConfigNode.h"
+#include "SimulationMultiResConfigNode.generated.h"
+
+/** Experimental Solver multires configuration node. */
+USTRUCT(Meta = (DataflowCloth))
+struct FChaosClothAssetSimulationMultiResConfigNode : public FChaosClothAssetSimulationBaseConfigNode
+{
+	GENERATED_USTRUCT_BODY()
+	DATAFLOW_NODE_DEFINE_INTERNAL(FChaosClothAssetSimulationMultiResConfigNode, "SimulationMultiResConfig", "Cloth", "Cloth Simulation MultiRes Config")
+
+public:
+
+	/** Enable multi-res simulation for this LOD.*/
+	UPROPERTY(EditAnywhere, Category = "Multi-resolution", meta = (EditCondition = "!bIsCoarseLOD"))
+	bool bIsFineLOD = false;
+
+	/** This LOD is a coarse LOD for a finer LOD's multi-res simulation.*/
+	UPROPERTY(EditAnywhere, Category = "Multi-resolution", meta = (EditCondition = "!bIsFineLOD"))
+	bool bIsCoarseMultiResLOD = false;
+
+	/** Index of the coarse LOD for multi-res simulation. That LOD also needs a Multi Res Config node with "Is Coarse Multi Res LOD" = true.*/
+	UPROPERTY(EditAnywhere, Category = "Multi-resolution", meta = (EditCondition = "bIsFineLOD", EditConditionHides))
+	int32 MultiResCoarseLODIndex = INDEX_NONE;
+
+	/** Use XPBD-style constraints */
+	UPROPERTY(EditAnywhere, Category = "Multi-resolution", meta = (EditCondition = "bIsFineLOD", EditConditionHides))
+	bool bMultiResUseXPBD = false;
+
+	/** MultiRes Spring Stiffness */
+	UPROPERTY(EditAnywhere, Category = "Multi-resolution", meta = (UIMin = "0", UIMax = "100000", ClampMin = "0", ClampMax = "100000", EditCondition = "bIsFineLOD", EditConditionHides))
+	FChaosClothAssetWeightedValue MultiResStiffness = {true, 1.f, 1.f};
+
+	/** MultiRes Velocity Target Spring Stiffness (non-XPBD only) */
+	UPROPERTY(EditAnywhere, Category = "Multi-resolution", meta = (UIMin = "0", UIMax = "1", ClampMin = "0", ClampMax = "1", EditCondition = "bIsFineLOD && !bMultiResUseXPBD", EditConditionHides))
+	FChaosClothAssetWeightedValue MultiResVelocityTargetStiffness = { true, 1.f, 1.f };
+
+	FChaosClothAssetSimulationMultiResConfigNode(const Dataflow::FNodeParameters& InParam, FGuid InGuid = FGuid::NewGuid());
+
+private:
+	virtual void AddProperties(FPropertyHelper& PropertyHelper) const override;
+};
\ No newline at end of file
diff --git a/Engine/Plugins/Experimental/ChaosClothAssetEditor/Source/ChaosClothAssetEditor/Private/ChaosClothAsset/ClothEditorSimulationVisualization.cpp b/Engine/Plugins/Experimental/ChaosClothAssetEditor/Source/ChaosClothAssetEditor/Private/ChaosClothAsset/ClothEditorSimulationVisualization.cpp
--- a/Engine/Plugins/Experimental/ChaosClothAssetEditor/Source/ChaosClothAssetEditor/Private/ChaosClothAsset/ClothEditorSimulationVisualization.cpp
+++ b/Engine/Plugins/Experimental/ChaosClothAssetEditor/Source/ChaosClothAssetEditor/Private/ChaosClothAsset/ClothEditorSimulationVisualization.cpp
@@ -172,6 +172,7 @@ const FVisualizationOption FVisualizationOption::OptionData[] =
 	FVisualizationOption(&::Chaos::FClothVisualization::DrawSelfCollisionLayers, TEXT("p.ChaosClothAssetEditor.DebugDrawSelfCollisionLayers"), LOCTEXT("ChaosVisName_SelfCollisionLayers", "Self Collision Layers"), LOCTEXT("ChaosVisName_SelfCollisionLayers_Tooltip", "Draws the self collision layers"), /*bDisablesSimulation =*/false, /*bHidesClothSections=*/true),
 	FVisualizationOption(&::Chaos::FClothVisualization::DrawKinematicColliderShaded, TEXT("p.ChaosClothAssetEditor.DebugDrawKinematicColliderShaded"), LOCTEXT("ChaosVisName_DrawKinematicColliderShaded", "Draw Kinematic Colliders (Shaded)"), LOCTEXT("ChaosVisName_DrawKinematicColliderShaded_Tooltip", "Draw kinematic cloth colliders with flat shading.")),
 	FVisualizationOption(&::Chaos::FClothVisualization::DrawKinematicColliderWired, TEXT("p.ChaosClothAssetEditor.DebugDrawKinematicColliderWired"), LOCTEXT("ChaosVisName_DrawKinematicColliderWired", "Draw Kinematic Colliders (Wireframe)"), LOCTEXT("ChaosVisName_DrawKinematicColliderWired_Tooltip", "Draw kinematic cloth colliders in wireframe.")),
+	FVisualizationOption(&::Chaos::FClothVisualization::DrawMultiResConstraint, TEXT("p.ChaosClothAssetEditor.DebugDrawMultiResConstraint"), LOCTEXT("ChaosVisName_DrawMultiResConstraint", "Draw Multi Res Constraint"), LOCTEXT("ChaosVisName_DrawMultiResConstraint_Tooltip", "Draw multi res constraint coarse mesh and targets.")),
 	FVisualizationOption(&GetSimulationStatisticsString, TEXT("p.ChaosClothAssetEditor.DebugDrawSimulationStatistics"), LOCTEXT("ChaosVisName_SimulationStatistics"    , "Simulation Statistics"), LOCTEXT("ChaosVisName_SimulationStatistics_Tooltip", "Displays simulation statistics")),
 	FVisualizationOption(&::Chaos::FClothVisualization::DrawWeightMap, TEXT("p.ChaosClothAssetEditor.DrawWeightMap"), LOCTEXT("ChaosVisName_DrawWeightMap", "Weight Map"), LOCTEXT("ChaosVisName_DrawWeightMap_ToolTip", "Draw the weight map for the simulation mesh. You can control the name of the map to be visualized by setting the p.ChaosClothVisualization.WeightMapName console variable."), /*bDisablesSimulation =*/false, /*bHidesClothSections=*/true),
 	FVisualizationOption(&::Chaos::FClothVisualization::DrawInpaintWeightsMatched, TEXT("p.ChaosClothAssetEditor.DrawInpaintWeightsMatched"), LOCTEXT("ChaosVisName_DrawInpaintWeightsMatched"        , "Transfer Skin Weights Node: Matched Vertices"), LOCTEXT("ChaosVisName_DrawInpaintWeightsMatched_ToolTip", "When transferring weights using the InpaintWeights method, will highlight the vertices for which we copied the weights directly from the source mesh. For all other vertices, the weights were computed automatically."), /*bDisablesSimulation =*/false, /*bHidesClothSections=*/true),
diff --git a/Engine/Source/Runtime/Experimental/Chaos/Private/Chaos/SoftsMultiResConstraints.cpp b/Engine/Source/Runtime/Experimental/Chaos/Private/Chaos/SoftsMultiResConstraints.cpp
new file mode 100644
--- /dev/null
+++ b/Engine/Source/Runtime/Experimental/Chaos/Private/Chaos/SoftsMultiResConstraints.cpp
@@ -0,0 +1,280 @@
+// Copyright Epic Games, Inc. All Rights Reserved.
+#include "Chaos/SoftsMultiResConstraints.h"
+#include "Chaos/Framework/Parallel.h"
+#include "Chaos/PBDSoftsSolverParticles.h"
+#include "Chaos/SoftsSpring.h"
+#include "ChaosStats.h"
+#include "ChaosLog.h"
+#include "HAL/IConsoleManager.h"
+
+#include "XPBDInternal.h"
+
+#if INTEL_ISPC
+#include "SoftsMultiResConstraints.ispc.generated.h"
+#endif
+
+#if INTEL_ISPC && !UE_BUILD_SHIPPING
+static_assert(sizeof(ispc::FVector4f) == sizeof(Chaos::Softs::FPAndInvM), "sizeof(ispc::FVector4f) != sizeof(Chaos::Softs::FPAndInvM)");
+static_assert(sizeof(ispc::FIntVector) == sizeof(Chaos::TVec3<int32>), "sizeof(ispc::FIntVector) != sizeof(Chaos::TVec3<int32>");
+static_assert(sizeof(ispc::FVector4f) == sizeof(Chaos::TVec4<Chaos::Softs::FSolverReal>), "sizeof(ispc::FVector4f) == sizeof(Chaos::TVec4<Chaos::Softs::FSolverReal>)");
+static_assert(sizeof(ispc::FIntVector) == sizeof(Chaos::TVec3<int32>), "sizeof(ispc::FIntVector) != sizeof(Chaos::TVec3<int32>");
+
+bool bChaos_MultiRes_ISPC_Enabled = true;
+FAutoConsoleVariableRef CVarChaosMultiResISPCEnabled(TEXT("p.Chaos.MultiRes.ISPC"), bChaos_MultiRes_ISPC_Enabled, TEXT("Whether to use ISPC optimizations in MultiRes constraints"));
+#endif
+
+namespace Chaos::Softs {
+
+static bool bMultiResConstraintApplyTargetNormalOffset = true;
+static FAutoConsoleVariableRef CVarMultiResConstraintApplyTargetNormalOffset(TEXT("p.Chaos.MultiRes.ApplyTargetNormalOffset"), bMultiResConstraintApplyTargetNormalOffset, TEXT("Apply normal offset to targets."));
+
+void FMultiResConstraints::SetProperties(const FCollectionPropertyConstFacade& PropertyCollection,
+	const TMap<FString, TConstArrayView<FRealSingle>>& WeightMaps)
+{
+	if (IsMultiResUseXPBDMutable(PropertyCollection))
+	{
+		bUseXPBD = GetMultiResUseXPBD(PropertyCollection);
+	}
+	if (IsMultiResStiffnessMutable(PropertyCollection))
+	{
+		const FSolverVec2 WeightedValue(GetWeightedFloatMultiResStiffness(PropertyCollection));
+		if (IsMultiResStiffnessStringDirty(PropertyCollection))
+		{
+			const FString& WeightMapName = GetMultiResStiffnessString(PropertyCollection);
+			Stiffness = FPBDStiffness(WeightedValue,
+				WeightMaps.FindRef(WeightMapName),
+				NumParticles,
+				FPBDStiffness::DefaultTableSize,
+				FPBDStiffness::DefaultParameterFitBase,
+				bUseXPBD ? MaxStiffness : (FSolverReal)1.f);
+		}
+		else
+		{
+			Stiffness.SetWeightedValue(WeightedValue, bUseXPBD ? MaxStiffness : (FSolverReal)1.f);
+		}
+	}
+	if (IsMultiResVelocityTargetStiffnessMutable(PropertyCollection))
+	{
+		const FSolverVec2 WeightedValue(GetWeightedFloatMultiResVelocityTargetStiffness(PropertyCollection));
+		if (IsMultiResVelocityTargetStiffnessStringDirty(PropertyCollection))
+		{
+			const FString& WeightMapName = GetMultiResVelocityTargetStiffnessString(PropertyCollection);
+			VelocityTargetStiffness = FPBDStiffness(WeightedValue,
+				WeightMaps.FindRef(WeightMapName),
+				NumParticles,
+				FPBDStiffness::DefaultTableSize,
+				FPBDStiffness::DefaultParameterFitBase,
+				bUseXPBD ? MaxStiffness : (FSolverReal)1.f);
+		}
+		else
+		{
+			VelocityTargetStiffness.SetWeightedValue(WeightedValue, bUseXPBD ? MaxStiffness : (FSolverReal)1.f);
+		}
+	}
+}
+
+void FMultiResConstraints::ApplyHelper(FSolverParticlesRange& Particles, const FSolverReal Dt, const int32 ConstraintIndex, const FSolverReal ExpStiffnessValue, const FSolverReal ExpVelocityTargetStiffnessValue) const
+{
+	if (bUseXPBD)
+	{
+		if (Particles.InvM(ConstraintIndex) != (FSolverReal)0.f)
+		{
+			const FSolverVec3& P1 = Particles.P(ConstraintIndex);
+			const FSolverVec3& P2 = FineTargetPositions[ConstraintIndex];
+			FSolverVec3 Direction = P1 - P2;
+			const FSolverReal Offset = Direction.SafeNormalize();
+
+			const FSolverReal AlphaInv = ExpStiffnessValue * Dt * Dt * (Particles.M(ConstraintIndex) + CoarseBarycentricMass[ConstraintIndex]) * 0.5f;
+			const FSolverReal DLambda = (-AlphaInv * Offset - Lambdas[ConstraintIndex]) / (AlphaInv * Particles.InvM(ConstraintIndex) + (FSolverReal)1.);
+			const FSolverVec3 Delta = DLambda * Direction;
+			Lambdas[ConstraintIndex] += DLambda;
+			Particles.P(ConstraintIndex) += Particles.InvM(ConstraintIndex) * Delta;
+		}
+	}
+	else
+	{
+		if (Particles.InvM(ConstraintIndex) != (FSolverReal)0.f)
+		{
+			Particles.P(ConstraintIndex) -= ExpStiffnessValue * (Particles.P(ConstraintIndex) - FineTargetPositions[ConstraintIndex]);
+			const FSolverVec3 FineDisplacement = Particles.P(ConstraintIndex) - Particles.X(ConstraintIndex);
+			const FSolverVec3 TargetDisplacement = FineTargetVelocities[ConstraintIndex] * Dt;
+			Particles.P(ConstraintIndex) -= ExpVelocityTargetStiffnessValue * (FineDisplacement - TargetDisplacement);
+		}
+	}
+}
+
+void FMultiResConstraints::Apply(FSolverParticlesRange& Particles, const FSolverReal Dt) const
+{
+	TRACE_CPUPROFILER_EVENT_SCOPE(FMultiResConstraints_Apply);
+	if (FineTargetPositions.Num() != Particles.Size())
+	{
+		return;
+	}
+#if INTEL_ISPC
+	if (!bUseXPBD && bChaos_MultiRes_ISPC_Enabled)
+	{
+		const bool bStiffnessHasWeightMap = Stiffness.HasWeightMap();
+		const bool bVelocityStiffnessHasWeightMap = VelocityTargetStiffness.HasWeightMap();
+		if (!bStiffnessHasWeightMap && !bVelocityStiffnessHasWeightMap)
+		{
+			ispc::ApplyMultiResConstraints(
+				(ispc::FVector4f*)Particles.GetPAndInvM().GetData(),
+				(const ispc::FVector3f*)Particles.XArray().GetData(),
+				(const ispc::FVector3f*)FineTargetPositions.GetData(),
+				(const ispc::FVector3f*)FineTargetVelocities.GetData(),
+				Dt,
+				(FSolverReal)Stiffness,
+				(FSolverReal)VelocityTargetStiffness,
+				Particles.Size()
+				);
+		}
+		else
+		{
+			ispc::ApplyMultiResConstraintsWithWeightMaps(
+				(ispc::FVector4f*)Particles.GetPAndInvM().GetData(),
+				(const ispc::FVector3f*)Particles.XArray().GetData(),
+				(const ispc::FVector3f*)FineTargetPositions.GetData(),
+				(const ispc::FVector3f*)FineTargetVelocities.GetData(),
+				Dt,
+				bStiffnessHasWeightMap,
+				Stiffness.GetIndices().GetData(),
+				Stiffness.GetTable().GetData(),
+				bVelocityStiffnessHasWeightMap,
+				VelocityTargetStiffness.GetIndices().GetData(),
+				VelocityTargetStiffness.GetTable().GetData(),
+				Particles.Size()
+			);
+		}
+	}
+	else
+#endif
+	{
+
+		if (Stiffness.HasWeightMap())
+		{
+			if (VelocityTargetStiffness.HasWeightMap())
+			{
+				for (int32 Index = 0; Index < Particles.Size(); ++Index)
+				{
+					const FSolverReal ExpStiffness = Stiffness[Index];
+					const FSolverReal ExpVelocityStiffness = VelocityTargetStiffness[Index];
+					ApplyHelper(Particles, Dt, Index, ExpStiffness, ExpVelocityStiffness);
+				}
+			}
+			else
+			{
+				const FSolverReal ExpVelocityStiffness = (FSolverReal)VelocityTargetStiffness;
+				for (int32 Index = 0; Index < Particles.Size(); ++Index)
+				{
+					const FSolverReal ExpStiffness = Stiffness[Index];
+					ApplyHelper(Particles, Dt, Index, ExpStiffness, ExpVelocityStiffness);
+				}
+			}
+		}
+		else
+		{
+			const FSolverReal ExpStiffness = (FSolverReal)Stiffness;
+			if (VelocityTargetStiffness.HasWeightMap())
+			{
+				for (int32 Index = 0; Index < Particles.Size(); ++Index)
+				{
+					const FSolverReal ExpVelocityStiffness = VelocityTargetStiffness[Index];
+					ApplyHelper(Particles, Dt, Index, ExpStiffness, ExpVelocityStiffness);
+				}
+			}
+			else
+			{
+				const FSolverReal ExpVelocityStiffness = (FSolverReal)VelocityTargetStiffness;
+				if (ExpStiffness == (FSolverReal)0.f && ExpVelocityStiffness == (FSolverReal)0.f)
+				{
+					return;
+				}
+				for (int32 Index = 0; Index < Particles.Size(); ++Index)
+				{
+					ApplyHelper(Particles, Dt, Index, ExpStiffness, ExpVelocityStiffness);
+				}
+			}
+		}
+	}
+}
+
+void FMultiResConstraints::UpdateFineTargets(const FSolverParticlesRange& CoarseParticles)
+{
+	TRACE_CPUPROFILER_EVENT_SCOPE(FMultiResConstraints_UpdateFineTargets);
+
+	CoarseBarycentricMass.SetNum(NumParticles);
+	FineTargetPositions.SetNum(NumParticles);
+	FineTargetVelocities.SetNum(NumParticles);
+	TConstArrayView<FSolverVec3> CoarseX = CoarseParticles.XArray();
+	TConstArrayView<FSolverReal> CoarseMass = CoarseParticles.GetM();
+	TConstArrayView<FSolverVec3> CoarseVelocity = CoarseParticles.GetV();
+	if (bMultiResConstraintApplyTargetNormalOffset)
+	{
+		TArray<FSolverVec3> Normals;
+		{
+			TRACE_CPUPROFILER_EVENT_SCOPE(FXPBDMultiResConstraints_CalculateNormals);
+			Normals.SetNum(CoarseParticles.Size());
+			// Need to calculate coarse normals
+			TArray<FSolverVec3> FaceNormals;
+			CoarseMesh.GetFaceNormals(FaceNormals, CoarseX, /*ReturnEmptyOnError =*/ false);
+			CoarseMesh.GetPointNormals(TArrayView<FSolverVec3>(Normals), TConstArrayView<FSolverVec3>(FaceNormals), /*bUseGlobalArray =*/ false);
+		}
+
+#if INTEL_ISPC
+		if (!bUseXPBD && bChaos_MultiRes_ISPC_Enabled)
+		{
+			ispc::MultiResUpdateFineTargets(
+				(ispc::FVector3f*)FineTargetPositions.GetData(),
+				(ispc::FVector3f*)FineTargetVelocities.GetData(),
+				(const ispc::FVector3f*)CoarseX.GetData(),
+				(const ispc::FVector3f*)Normals.GetData(),
+				(const ispc::FVector3f*)CoarseVelocity.GetData(),
+				(const ispc::FVector4f*)CoarseToFinePositionBaryCoordsAndDist.GetData(),
+				(const ispc::FIntVector*)CoarseToFineSourceMeshVertIndices.GetData(),
+				NumParticles
+			);
+		}
+		else
+#endif
+		{
+			for (int32 Index = 0; Index < NumParticles; ++Index)
+			{
+				const TVec4<FSolverReal>& PositionBaryCoordsAndDist = CoarseToFinePositionBaryCoordsAndDist[Index];
+				const TVec3<int32>& SourceMeshVertIndices = CoarseToFineSourceMeshVertIndices[Index];
+				FineTargetPositions[Index] =
+					PositionBaryCoordsAndDist[0] * (CoarseX[SourceMeshVertIndices[0]] + Normals[SourceMeshVertIndices[0]] * PositionBaryCoordsAndDist[3]) +
+					PositionBaryCoordsAndDist[1] * (CoarseX[SourceMeshVertIndices[1]] + Normals[SourceMeshVertIndices[1]] * PositionBaryCoordsAndDist[3]) +
+					PositionBaryCoordsAndDist[2] * (CoarseX[SourceMeshVertIndices[2]] + Normals[SourceMeshVertIndices[2]] * PositionBaryCoordsAndDist[3]);
+				CoarseBarycentricMass[Index] =
+					PositionBaryCoordsAndDist[0] * CoarseMass[SourceMeshVertIndices[0]] +
+					PositionBaryCoordsAndDist[1] * CoarseMass[SourceMeshVertIndices[1]] +
+					PositionBaryCoordsAndDist[2] * CoarseMass[SourceMeshVertIndices[2]];
+				FineTargetVelocities[Index] =
+					PositionBaryCoordsAndDist[0] * CoarseVelocity[SourceMeshVertIndices[0]] +
+					PositionBaryCoordsAndDist[1] * CoarseVelocity[SourceMeshVertIndices[1]] +
+					PositionBaryCoordsAndDist[2] * CoarseVelocity[SourceMeshVertIndices[2]];
+			}
+		}
+	}
+	else
+	{
+		for (int32 Index = 0; Index < NumParticles; ++Index)
+		{
+			const TVec4<FSolverReal>& PositionBaryCoordsAndDist = CoarseToFinePositionBaryCoordsAndDist[Index];
+			const TVec3<int32>& SourceMeshVertIndices = CoarseToFineSourceMeshVertIndices[Index];
+			FineTargetPositions[Index] =
+				PositionBaryCoordsAndDist[0] * CoarseX[SourceMeshVertIndices[0]] +
+				PositionBaryCoordsAndDist[1] * CoarseX[SourceMeshVertIndices[1]] +
+				PositionBaryCoordsAndDist[2] * CoarseX[SourceMeshVertIndices[2]];
+			CoarseBarycentricMass[Index] =
+				PositionBaryCoordsAndDist[0] * CoarseMass[SourceMeshVertIndices[0]] +
+				PositionBaryCoordsAndDist[1] * CoarseMass[SourceMeshVertIndices[1]] +
+				PositionBaryCoordsAndDist[2] * CoarseMass[SourceMeshVertIndices[2]];
+			FineTargetVelocities[Index] =
+				PositionBaryCoordsAndDist[0] * CoarseVelocity[SourceMeshVertIndices[0]] +
+				PositionBaryCoordsAndDist[1] * CoarseVelocity[SourceMeshVertIndices[1]] +
+				PositionBaryCoordsAndDist[2] * CoarseVelocity[SourceMeshVertIndices[2]];
+		}
+	}
+}
+} // End namespace Chaos::Softs
diff --git a/Engine/Source/Runtime/Experimental/Chaos/Private/Chaos/SoftsMultiResConstraints.ispc b/Engine/Source/Runtime/Experimental/Chaos/Private/Chaos/SoftsMultiResConstraints.ispc
new file mode 100644
--- /dev/null
+++ b/Engine/Source/Runtime/Experimental/Chaos/Private/Chaos/SoftsMultiResConstraints.ispc
@@ -0,0 +1,123 @@
+// Copyright Epic Games, Inc. All Rights Reserved.
+
+#define EXPLICIT_VECTOR4 1
+
+#include "Math/Vector.isph"
+#include "Chaos/PBDSofts.isph"
+
+export void ApplyMultiResConstraints(uniform FVector4f PandInvM[],
+									const uniform FVector3f X[],
+									const uniform FVector3f TargetPositions[],
+									const uniform FVector3f TargetVelocities[],
+									const uniform float Dt,
+									const uniform float Stiffness,
+									const uniform float VelocityStiffness,
+									const uniform int32 NumParticles)
+{
+	foreach(Index = 0 ... NumParticles)
+	{
+		const uniform int32 Index1 = extract(Index,0);
+		const varying FVector4f PandInvM1 = VectorLoad(&PandInvM[Index1]);
+
+		varying FVector3f P1;
+		varying float InvM1;
+		UnzipPandInvM(PandInvM1, P1, InvM1);
+		
+		if(InvM1 > 0.f)
+		{
+			const varying FVector3f X1 = VectorLoad(&X[Index1]);
+			const varying FVector3f TargetPos = VectorLoad(&TargetPositions[Index1]);
+			const varying FVector3f TargetVel = VectorLoad(&TargetVelocities[Index1]);
+
+			P1 = P1 - Stiffness * (P1 - TargetPos);
+			const varying FVector3f FineDisplacement = P1 - X1;
+			P1 = P1 - VelocityStiffness * (FineDisplacement - TargetVel * Dt);
+			VectorStore(&PandInvM[Index1], SetVector4(P1, InvM1));
+		}
+	}
+}
+
+export void ApplyMultiResConstraintsWithWeightMaps(uniform FVector4f PandInvM[],
+									const uniform FVector3f X[],
+									const uniform FVector3f TargetPositions[],
+									const uniform FVector3f TargetVelocities[],
+									const uniform float Dt,
+									const uniform bool StiffnessHasMap,
+									const uniform uint8 StiffnessIndices[],
+									const uniform float StiffnessTable[],
+									const uniform bool VelocityStiffnessHasMap,
+									const uniform uint8 VelocityStiffnessIndices[],
+									const uniform float VelocityStiffnessTable[],
+									const uniform int32 NumParticles)
+{
+	foreach(Index = 0 ... NumParticles)
+	{
+		const uniform int32 Index1 = extract(Index,0);
+		const varying FVector4f PandInvM1 = VectorLoad(&PandInvM[Index1]);
+
+		varying FVector3f P1;
+		varying float InvM1;
+		UnzipPandInvM(PandInvM1, P1, InvM1);
+		
+		if(InvM1 > 0.f)
+		{
+			const varying int8 StiffnessIndex = StiffnessHasMap ? StiffnessIndices[Index] : 0;
+			const varying int8 VelocityStiffnessIndex = VelocityStiffnessHasMap ? VelocityStiffnessIndices[Index] : 0;
+			#pragma ignore warning(perf)
+			const varying float Stiffness = StiffnessTable[StiffnessIndex];
+			#pragma ignore warning(perf)
+			const varying float VelocityStiffness = VelocityStiffnessTable[VelocityStiffnessIndex];
+
+			const varying FVector3f X1 = VectorLoad(&X[Index1]);
+			const varying FVector3f TargetPos = VectorLoad(&TargetPositions[Index1]);
+			const varying FVector3f TargetVel = VectorLoad(&TargetVelocities[Index1]);
+
+			P1 = P1 - Stiffness * (P1 - TargetPos);
+			const varying FVector3f FineDisplacement = P1 - X1;
+			P1 = P1 - VelocityStiffness * (FineDisplacement - TargetVel * Dt);
+			VectorStore(&PandInvM[Index1], SetVector4(P1, InvM1));
+		}
+	}
+}
+
+export void MultiResUpdateFineTargets(uniform FVector3f FineTargetPositions[],
+							uniform FVector3f FineTargetVelocities[],
+							const uniform FVector3f CoarseX[],
+							const uniform FVector3f CoarseNormals[],
+							const uniform FVector3f CoarseV[],
+							const uniform FVector4f CoarseToFinePositionBaryCoordsAndDist[],
+							const uniform FIntVector CoarseToFineSourceMeshVertIndices[],
+							const uniform int32 NumParticles)
+{
+	foreach(Index = 0 ... NumParticles)
+	{
+		const uniform int32 Index0 = extract(Index,0);
+
+		const varying FVector4f PositionBaryCoordsAndDist = VectorLoad(&CoarseToFinePositionBaryCoordsAndDist[Index0]);
+		const varying FIntVector SourceMeshVertIndices = VectorLoad(&CoarseToFineSourceMeshVertIndices[Index0]);
+		const varying int32 Index1 = SourceMeshVertIndices.V[0];
+		const varying int32 Index2 = SourceMeshVertIndices.V[1];
+		const varying int32 Index3 = SourceMeshVertIndices.V[2];
+
+		const varying FVector3f CoarseX1 = VectorGather(&CoarseX[Index1]);
+		const varying FVector3f CoarseX2 = VectorGather(&CoarseX[Index2]);
+		const varying FVector3f CoarseX3 = VectorGather(&CoarseX[Index3]);
+		const varying FVector3f CoarseNormal1 = VectorGather(&CoarseNormals[Index1]);
+		const varying FVector3f CoarseNormal2 = VectorGather(&CoarseNormals[Index2]);
+		const varying FVector3f CoarseNormal3 = VectorGather(&CoarseNormals[Index3]);
+		const varying FVector3f CoarseV1 = VectorGather(&CoarseV[Index1]);
+		const varying FVector3f CoarseV2 = VectorGather(&CoarseV[Index2]);
+		const varying FVector3f CoarseV3 = VectorGather(&CoarseV[Index3]);
+
+		const varying FVector3f FineTargetPos = 
+					PositionBaryCoordsAndDist.V[0] * (CoarseX1 + CoarseNormal1 * PositionBaryCoordsAndDist.V[3]) +
+					PositionBaryCoordsAndDist.V[1] * (CoarseX2 + CoarseNormal2 * PositionBaryCoordsAndDist.V[3]) +
+					PositionBaryCoordsAndDist.V[2] * (CoarseX3 + CoarseNormal3 * PositionBaryCoordsAndDist.V[3]);
+		const varying FVector3f FineTargetVel =
+					PositionBaryCoordsAndDist.V[0] * CoarseV1 +
+					PositionBaryCoordsAndDist.V[1] * CoarseV2 +
+					PositionBaryCoordsAndDist.V[2] * CoarseV3;
+		VectorStore(&FineTargetPositions[Index0], FineTargetPos);
+		VectorStore(&FineTargetVelocities[Index0], FineTargetVel);
+	}
+}
diff --git a/Engine/Source/Runtime/Experimental/Chaos/Private/Chaos/TriangleMesh.cpp b/Engine/Source/Runtime/Experimental/Chaos/Private/Chaos/TriangleMesh.cpp
--- a/Engine/Source/Runtime/Experimental/Chaos/Private/Chaos/TriangleMesh.cpp
+++ b/Engine/Source/Runtime/Experimental/Chaos/Private/Chaos/TriangleMesh.cpp
@@ -439,8 +439,7 @@ void FTriangleMesh::GetPointNormals(TArrayView<TVec3<T>> PointNormals, const TCo
 }
 template CHAOS_API void FTriangleMesh::GetPointNormals<FRealDouble>(TArrayView<TVec3<FRealDouble>>, const TConstArrayView<TVec3<FRealDouble>>&, const bool) const;
 
-template <>
-CHAOS_API void FTriangleMesh::GetPointNormals<FRealSingle>(TArrayView<TVec3<FRealSingle>> PointNormals, const TConstArrayView<TVec3<FRealSingle>>& FaceNormals, const bool bUseGlobalArray) const
+CHAOS_API void FTriangleMesh::GetPointNormals(TArrayView<TVec3<FRealSingle>> PointNormals, const TConstArrayView<TVec3<FRealSingle>>& FaceNormals, const bool bUseGlobalArray) const
 {
 	check(MPointToTriangleMap.Num() != 0);
 
diff --git a/Engine/Source/Runtime/Experimental/Chaos/Public/Chaos/PBDSoftsEvolutionFwd.h b/Engine/Source/Runtime/Experimental/Chaos/Public/Chaos/PBDSoftsEvolutionFwd.h
--- a/Engine/Source/Runtime/Experimental/Chaos/Public/Chaos/PBDSoftsEvolutionFwd.h
+++ b/Engine/Source/Runtime/Experimental/Chaos/Public/Chaos/PBDSoftsEvolutionFwd.h
@@ -77,6 +77,7 @@ namespace Chaos
 		class FPBDTriangleMeshCollisions;
 		class FPBDSelfCollisionSphereConstraints;
 		class FPBDSoftBodyCollisionConstraint;
+		class FMultiResConstraints;
 
 		// Softs solver forces
 		class FVelocityAndPressureField;
diff --git a/Engine/Source/Runtime/Experimental/Chaos/Public/Chaos/SoftsMultiResConstraints.h b/Engine/Source/Runtime/Experimental/Chaos/Public/Chaos/SoftsMultiResConstraints.h
new file mode 100644
--- /dev/null
+++ b/Engine/Source/Runtime/Experimental/Chaos/Public/Chaos/SoftsMultiResConstraints.h
@@ -0,0 +1,135 @@
+// Copyright Epic Games, Inc. All Rights Reserved.
+#pragma once
+
+#include "Chaos/Core.h"
+#include "Chaos/PBDSoftsEvolutionFwd.h"
+#include "Chaos/PBDStiffness.h"
+#include "Chaos/SoftsSolverParticlesRange.h"
+#include "Chaos/TriangleMesh.h"
+#include "Chaos/CollectionPropertyFacade.h"
+
+namespace Chaos::Softs
+{
+class FMultiResConstraints final
+{
+public:
+	static bool IsEnabled(const FCollectionPropertyConstFacade& PropertyCollection)
+	{
+		return IsMultiResStiffnessEnabled(PropertyCollection, false);
+	}
+
+	static constexpr FSolverReal MinStiffness = (FSolverReal)0; // We're not checking against MinStiffness (except when it's constant and == 0)
+	static constexpr FSolverReal MaxStiffness = (FSolverReal)1e9;
+
+	FMultiResConstraints(
+		const FSolverParticlesRange& FineParticles,
+		const int32 InCoarseSoftBodyId,
+		const FTriangleMesh& InCoarseMesh,
+		TArray<TVec4<FSolverReal>>&& InCoarseToFinePositionBaryCoordsAndDist,
+		TArray<TVec3<int32>>&& InCoarseToFineSourceMeshVertIndices,
+		const TMap<FString, TConstArrayView<FRealSingle>>& WeightMaps,
+		const FCollectionPropertyConstFacade& PropertyCollection)
+		: NumParticles(FineParticles.Size())
+		, CoarseSoftBodyId(InCoarseSoftBodyId)
+		, CoarseMesh(InCoarseMesh)
+		, CoarseToFinePositionBaryCoordsAndDist(MoveTemp(InCoarseToFinePositionBaryCoordsAndDist))
+		, CoarseToFineSourceMeshVertIndices(MoveTemp(InCoarseToFineSourceMeshVertIndices))
+		, bUseXPBD(GetMultiResUseXPBD(PropertyCollection, false))
+		, Stiffness(
+			FSolverVec2(GetWeightedFloatMultiResStiffness(PropertyCollection, 1.f)),
+			WeightMaps.FindRef(GetMultiResStiffnessString(PropertyCollection, MultiResStiffnessName.ToString())),
+			NumParticles,
+			FPBDStiffness::DefaultTableSize,
+			FPBDStiffness::DefaultParameterFitBase,
+			bUseXPBD ? MaxStiffness : (FSolverReal)1.f)
+		, VelocityTargetStiffness(
+			FSolverVec2(GetWeightedFloatMultiResVelocityTargetStiffness(PropertyCollection, 1.f)),
+			WeightMaps.FindRef(GetMultiResVelocityTargetStiffnessString(PropertyCollection, MultiResVelocityTargetStiffnessName.ToString())),
+			NumParticles,
+			FPBDStiffness::DefaultTableSize,
+			FPBDStiffness::DefaultParameterFitBase,
+			bUseXPBD ? MaxStiffness : (FSolverReal)1.f)
+		, MultiResUseXPBDIndex(PropertyCollection)
+		, MultiResStiffnessIndex(PropertyCollection)
+		, MultiResVelocityTargetStiffnessIndex(PropertyCollection)
+	{
+		check(NumParticles == CoarseToFinePositionBaryCoordsAndDist.Num());
+		check(NumParticles == CoarseToFineSourceMeshVertIndices.Num());
+	}
+
+	~FMultiResConstraints() {}
+
+	void Init() const 
+	{
+		if (bUseXPBD)
+		{
+			Lambdas.Reset();
+			Lambdas.SetNumZeroed(NumParticles);
+		}
+	}
+
+	void ApplyProperties(const FSolverReal Dt, const int32 NumIterations)
+	{
+		if (bUseXPBD)
+		{
+			Stiffness.ApplyXPBDValues(MaxStiffness);
+			VelocityTargetStiffness.ApplyXPBDValues(MaxStiffness);
+		}
+		else
+		{
+			Stiffness.ApplyPBDValues(Dt, NumIterations);
+			VelocityTargetStiffness.ApplyPBDValues(Dt, NumIterations);
+		}
+	}
+
+	CHAOS_API void UpdateFineTargets(const FSolverParticlesRange& CoarseParticles);
+
+	// Update stiffness values
+	CHAOS_API void SetProperties(const FCollectionPropertyConstFacade& PropertyCollection,
+		const TMap<FString, TConstArrayView<FRealSingle>>& WeightMaps);
+
+	CHAOS_API void Apply(FSolverParticlesRange& FineParticles, const FSolverReal Dt) const;
+
+	const TArray<FSolverVec3>& GetFineTargetPositions() const { return FineTargetPositions; }
+	int32 GetCoarseSoftBodyId() const { return CoarseSoftBodyId; }
+	const FTriangleMesh& GetCoarseMesh() const { return CoarseMesh; }
+	bool IsConstraintActive(const int32 ParticleIndex) const
+	{
+		const FSolverReal StiffnessValue = Stiffness.HasWeightMap() ? Stiffness[ParticleIndex] : (FSolverReal)Stiffness;
+		const FSolverReal VelocityStiffnessValue = VelocityTargetStiffness.HasWeightMap() ? VelocityTargetStiffness[ParticleIndex] : (FSolverReal)VelocityTargetStiffness;
+		return StiffnessValue > 0.f || VelocityStiffnessValue > 0.f;
+	}
+
+private:
+	void ApplyHelper(FSolverParticlesRange& Particles, const FSolverReal Dt, const int32 ConstraintIndex, const FSolverReal ExpStiffnessValue, const FSolverReal ExpVelocityTargetStiffnessValue) const;
+
+	const int32 NumParticles;
+	const int32 CoarseSoftBodyId;
+	const FTriangleMesh& CoarseMesh;
+	const TArray<TVec4<FSolverReal>> CoarseToFinePositionBaryCoordsAndDist;
+	const TArray<TVec3<int32>> CoarseToFineSourceMeshVertIndices;
+
+	bool bUseXPBD;
+	FPBDStiffness Stiffness;
+	FPBDStiffness VelocityTargetStiffness;
+
+	mutable TArray<FSolverReal> Lambdas;
+	TArray<FSolverVec3> FineTargetPositions;
+	TArray<FSolverVec3> FineTargetVelocities;
+	TArray<FSolverReal> CoarseBarycentricMass;
+
+	UE_CHAOS_DECLARE_PROPERTYCOLLECTION_NAME(MultiResUseXPBD, bool);
+	UE_CHAOS_DECLARE_PROPERTYCOLLECTION_NAME(MultiResStiffness, float);
+	UE_CHAOS_DECLARE_PROPERTYCOLLECTION_NAME(MultiResVelocityTargetStiffness, float);
+};
+
+}  // End namespace Chaos::Softs
+
+// Support ISPC enable/disable in non-shipping builds
+#if !INTEL_ISPC
+const bool bChaos_MultiRes_ISPC_Enabled = false;
+#elif UE_BUILD_SHIPPING
+const bool bChaos_MultiRes_ISPC_Enabled = true;
+#else
+extern CHAOS_API bool bChaos_MultiRes_ISPC_Enabled;
+#endif
diff --git a/Engine/Source/Runtime/Experimental/Chaos/Public/Chaos/TriangleMesh.h b/Engine/Source/Runtime/Experimental/Chaos/Public/Chaos/TriangleMesh.h
--- a/Engine/Source/Runtime/Experimental/Chaos/Public/Chaos/TriangleMesh.h
+++ b/Engine/Source/Runtime/Experimental/Chaos/Public/Chaos/TriangleMesh.h
@@ -166,6 +166,8 @@ namespace Chaos
 		template <typename T>
 		void GetPointNormals(TArrayView<TVec3<T>> PointNormals, const TConstArrayView<TVec3<T>>& FaceNormals, const bool bUseGlobalArray) const;
 
+		CHAOS_API void GetPointNormals(TArrayView<TVec3<FRealSingle>> PointNormals, const TConstArrayView<TVec3<FRealSingle>>& FaceNormals, const bool bUseGlobalArray) const;
+
 		static CHAOS_API FTriangleMesh GetConvexHullFromParticles(const TConstArrayView<FVec3>& points);
 		/** Deprecated. Use TArrayView version. */
 		static FORCEINLINE FTriangleMesh GetConvexHullFromParticles(const FParticles& InParticles)
