Chaos Cloth: Allow VelocityScale to go over 1.

diff --git a/Engine/Plugins/Experimental/ChaosCloth/Source/Chaos/Private/ChaosCloth/ChaosClothingSimulationCloth.cpp b/Engine/Plugins/Experimental/ChaosCloth/Source/Chaos/Private/ChaosCloth/ChaosClothingSimulationCloth.cpp
--- a/Engine/Plugins/Experimental/ChaosCloth/Source/Chaos/Private/ChaosCloth/ChaosClothingSimulationCloth.cpp
+++ b/Engine/Plugins/Experimental/ChaosCloth/Source/Chaos/Private/ChaosCloth/ChaosClothingSimulationCloth.cpp
@@ -36,6 +36,7 @@ namespace ClothingSimulationClothDefault
 	constexpr float GravityScale = Softs::FExternalForces::DefaultGravityScale; // 1.f;
 	constexpr float GravityZOverride = Softs::FExternalForces::DefaultGravityZOverride; // -980.665f;
 	constexpr float VelocityScale = 0.75f;
+	constexpr float MaxVelocityScale = 1.f;
 	constexpr float FictitiousAngularScale = Softs::FExternalForces::DefaultFictitiousAngularScale; // 1.f;
 	constexpr int32 MultiResCoarseLODIndex = INDEX_NONE;
 }
@@ -1134,12 +1135,13 @@ PRAGMA_ENABLE_DEPRECATION_WARNINGS
 		// Set the reference input velocity and deal with teleport & reset; external forces depends on these values, so they must be initialized before then
 		FVec3f OutLinearVelocityScale;
 		FRealSingle OutAngularVelocityScale;
-
+		FRealSingle OutMaxVelocityScale;
 		if (bNeedsReset)
 		{
 			// Make sure not to do any pre-sim transform just after a reset
 			OutLinearVelocityScale = FVec3f(1.f);
 			OutAngularVelocityScale = 1.f;
+			OutMaxVelocityScale = 1.f;
 
 			// Reset to start pose
 			LODData[LODIndex]->ResetStartPose(Solver);
@@ -1159,6 +1161,7 @@ PRAGMA_ENABLE_DEPRECATION_WARNINGS
 			// Remove all impulse velocity from the last frame
 			OutLinearVelocityScale = FVec3f(0.f);
 			OutAngularVelocityScale = 0.f;
+			OutMaxVelocityScale = 1.f;
 			UE_LOG(LogChaosCloth, VeryVerbose, TEXT("Cloth in group Id %d Needs teleport."), GroupId);
 		}
 		else
@@ -1166,6 +1169,7 @@ PRAGMA_ENABLE_DEPRECATION_WARNINGS
 			// Use the cloth config parameters
 			OutLinearVelocityScale = ConfigProperties.GetValue<FVector3f>(TEXT("LinearVelocityScale"), FVector3f(ClothingSimulationClothDefault::VelocityScale));
 			OutAngularVelocityScale = ConfigProperties.GetValue<float>(TEXT("AngularVelocityScale"), ClothingSimulationClothDefault::VelocityScale);
+			OutMaxVelocityScale = ConfigProperties.GetValue<float>(TEXT("MaxVelocityScale"), ClothingSimulationClothDefault::MaxVelocityScale);
 		}
 
 		// NOTE: Force-based solver doesn't actually use FictitiousAngularScale here. It gets it from the property collection directly.
@@ -1176,7 +1180,8 @@ PRAGMA_ENABLE_DEPRECATION_WARNINGS
 			ReferenceSpaceTransform,
 			OutLinearVelocityScale,
 			OutAngularVelocityScale,
-			FictitiousAngularScale);
+			FictitiousAngularScale,
+			OutMaxVelocityScale);
 		if (Solver->IsForceBasedSolver())
 		{
 			Solver->SetProperties(ParticleRangeId, ConfigProperties, LODData[LODIndex]->WeightMaps);
diff --git a/Engine/Plugins/Experimental/ChaosCloth/Source/Chaos/Private/ChaosCloth/ChaosClothingSimulationSolver.cpp b/Engine/Plugins/Experimental/ChaosCloth/Source/Chaos/Private/ChaosCloth/ChaosClothingSimulationSolver.cpp
--- a/Engine/Plugins/Experimental/ChaosCloth/Source/Chaos/Private/ChaosCloth/ChaosClothingSimulationSolver.cpp
+++ b/Engine/Plugins/Experimental/ChaosCloth/Source/Chaos/Private/ChaosCloth/ChaosClothingSimulationSolver.cpp
@@ -1077,7 +1077,8 @@ void FClothingSimulationSolver::SetReferenceVelocityScale(
 	const FRigidTransform3& OldReferenceSpaceTransform,  // Transforms are in world space so have to be FReal based for LWC
 	const FRigidTransform3& ReferenceSpaceTransform,
 	const TVec3<FRealSingle>& LinearVelocityScale,
-	FRealSingle AngularVelocityScale, FRealSingle FictitiousAngularScale
+	FRealSingle AngularVelocityScale, FRealSingle FictitiousAngularScale,
+	FRealSingle MaxVelocityScale
 )
 {
 	FRigidTransform3 OldRootBoneLocalTransform = OldReferenceSpaceTransform;
@@ -1088,8 +1089,23 @@ void FClothingSimulationSolver::SetReferenceVelocityScale(
 	// Calculate deltas
 	const FRigidTransform3 DeltaTransform = ReferenceSpaceTransform.GetRelativeTransform(OldReferenceSpaceTransform);
 
+	auto CalculateClampedVelocityScale = [MaxVelocityScale, SolverVelocityScale](FRealSingle InVelocityScale)
+	{
+		FReal CombinedVelocityScale = SolverVelocityScale * InVelocityScale;
+		if (CombinedVelocityScale <= FMath::Min((FReal)1., (FReal)MaxVelocityScale))
+		{
+			// When combined velocity scale is <= 1 (or smaller max velocity scale), just use it.
+			return CombinedVelocityScale;
+		}
+		// Otherwise, clamp SolverVelocityScale (what's calculated when doing delta time smoothing) to 1, and apply MaxVelocityScale clamp on the user supplied value.
+		return FMath::Clamp(SolverVelocityScale, (FReal)0., (FReal)1.) * FMath::Clamp(InVelocityScale, (FReal)0., (FReal)MaxVelocityScale);
+	};
+
 	// Apply linear velocity scale
-	const FVec3 LinearRatio = FVec3(1.) - FVec3(LinearVelocityScale * SolverVelocityScale).BoundToBox(FVec3(0.), FVec3(1.));
+	const FVec3 LinearRatio = FVec3(1.) - FVec3(
+		CalculateClampedVelocityScale(LinearVelocityScale[0]),
+		CalculateClampedVelocityScale(LinearVelocityScale[1]),
+		CalculateClampedVelocityScale(LinearVelocityScale[2]));
 	const FVec3 DeltaPosition = LinearRatio * DeltaTransform.GetTranslation();
 
 	// Apply angular velocity scale
@@ -1101,7 +1117,7 @@ void FClothingSimulationSolver::SetReferenceVelocityScale(
 		DeltaAngle -= (FReal)2. * (FReal)PI;
 	}
 
-	const FReal PartialDeltaAngle = DeltaAngle * FMath::Clamp((FReal)1. - (FReal)AngularVelocityScale * SolverVelocityScale, (FReal)0., (FReal)1.);
+	const FReal PartialDeltaAngle = DeltaAngle * ((FReal)1. - CalculateClampedVelocityScale(AngularVelocityScale));
 	DeltaRotation = UE::Math::TQuat<FReal>(Axis, PartialDeltaAngle);
 
 	// Transform points back into the previous frame of reference before applying the adjusted deltas
diff --git a/Engine/Plugins/Experimental/ChaosCloth/Source/Chaos/Public/ChaosCloth/ChaosClothingSimulationSolver.h b/Engine/Plugins/Experimental/ChaosCloth/Source/Chaos/Public/ChaosCloth/ChaosClothingSimulationSolver.h
--- a/Engine/Plugins/Experimental/ChaosCloth/Source/Chaos/Public/ChaosCloth/ChaosClothingSimulationSolver.h
+++ b/Engine/Plugins/Experimental/ChaosCloth/Source/Chaos/Public/ChaosCloth/ChaosClothingSimulationSolver.h
@@ -144,7 +144,8 @@ namespace Chaos
 			const FRigidTransform3& ReferenceSpaceTransform,
 			const TVec3<FRealSingle>& LinearVelocityScale,
 			FRealSingle AngularVelocityScale,
-			FRealSingle FictitiousAngularScale);
+			FRealSingle FictitiousAngularScale,
+			FRealSingle MaxVelocityScale = 1.f);
 
 		/** PBDSolver version */
 		CHAOSCLOTH_API void SetProperties(
diff --git a/Engine/Plugins/Experimental/ChaosClothAssetEditor/Source/ChaosClothAssetDataflowNodes/Private/ChaosClothAsset/SimulationVelocityScaleConfigNode.cpp b/Engine/Plugins/Experimental/ChaosClothAssetEditor/Source/ChaosClothAssetDataflowNodes/Private/ChaosClothAsset/SimulationVelocityScaleConfigNode.cpp
--- a/Engine/Plugins/Experimental/ChaosClothAssetEditor/Source/ChaosClothAssetDataflowNodes/Private/ChaosClothAsset/SimulationVelocityScaleConfigNode.cpp
+++ b/Engine/Plugins/Experimental/ChaosClothAssetEditor/Source/ChaosClothAssetDataflowNodes/Private/ChaosClothAsset/SimulationVelocityScaleConfigNode.cpp
@@ -13,7 +13,9 @@ FChaosClothAssetSimulationVelocityScaleConfigNode::FChaosClothAssetSimulationVel
 
 void FChaosClothAssetSimulationVelocityScaleConfigNode::AddProperties(FPropertyHelper& PropertyHelper) const
 {
+	
 	PropertyHelper.SetProperty(this, &LinearVelocityScale);
 	PropertyHelper.SetProperty(this, &AngularVelocityScale);
+	PropertyHelper.SetProperty(this, &MaxVelocityScale);
 	PropertyHelper.SetProperty(this, &FictitiousAngularScale);
 }
diff --git a/Engine/Plugins/Experimental/ChaosClothAssetEditor/Source/ChaosClothAssetDataflowNodes/Public/ChaosClothAsset/SimulationVelocityScaleConfigNode.h b/Engine/Plugins/Experimental/ChaosClothAssetEditor/Source/ChaosClothAssetDataflowNodes/Public/ChaosClothAsset/SimulationVelocityScaleConfigNode.h
--- a/Engine/Plugins/Experimental/ChaosClothAssetEditor/Source/ChaosClothAssetDataflowNodes/Public/ChaosClothAsset/SimulationVelocityScaleConfigNode.h
+++ b/Engine/Plugins/Experimental/ChaosClothAssetEditor/Source/ChaosClothAssetDataflowNodes/Public/ChaosClothAsset/SimulationVelocityScaleConfigNode.h
@@ -14,20 +14,29 @@ struct FChaosClothAssetSimulationVelocityScaleConfigNode : public FChaosClothAss
 	DATAFLOW_NODE_DEFINE_INTERNAL(FChaosClothAssetSimulationVelocityScaleConfigNode, "SimulationVelocityScaleConfig", "Cloth", "Cloth Simulation Velocity Scale Config")
 
 public:
+
 	/**
 	 * The amount of linear velocities sent to the local cloth space from the reference bone
 	 * (the closest bone to the root on which the cloth section has been skinned, or the root itself if the cloth isn't skinned).
+	 * This value will be clamped by "Max Velocity Scale". A velocity scale of > 1 will amplify the velocities from the reference bone.
 	 */
-	UPROPERTY(EditAnywhere, Category = "Animation Properties", meta = (UIMin = "0", UIMax = "1", ClampMin = "0", ClampMax = "1"))
+	UPROPERTY(EditAnywhere, Category = "Animation Properties", meta = (UIMin = "0", UIMax = "1", ClampMin = "0", ClampMax = "100"))
 	FVector3f LinearVelocityScale = { 0.75f, 0.75f, 0.75f };
 
 	/**
 	 * The amount of angular velocities sent to the local cloth space from the reference bone
 	 * (the closest bone to the root on which the cloth section has been skinned, or the root itself if the cloth isn't skinned).
+	 * This value will be clamped by "Max Velocity Scale". A velocity scale of > 1 will amplify the velocities from the reference bone.
 	 */
-	UPROPERTY(EditAnywhere, Category = "Animation Properties", meta = (UIMin = "0", UIMax = "1", ClampMin = "0", ClampMax = "1"))
+	UPROPERTY(EditAnywhere, Category = "Animation Properties", meta = (UIMin = "0", UIMax = "1", ClampMin = "0", ClampMax = "100"))
 	float AngularVelocityScale = 0.75f;
 
+	/**
+	 * Clamp on Linear and Angular Velocity Scale. The final velocity scale (e.g., including contributions from blueprints) will be clamped to this value.
+	 */
+	UPROPERTY(EditAnywhere, Category = "Animation Properties", meta = (UIMin = "0", UIMax = "1", ClampMin = "0", ClampMax = "100"))
+	float MaxVelocityScale = 1.f;
+
 	/**
 	 * The portion of the angular velocity that is used to calculate the strength of all fictitious forces (e.g. centrifugal force).
 	 * This parameter is only having an effect on the portion of the reference bone's angular velocity that has been removed from the
