Chaos Cloth:
- Move the Chaos Cloth plugin out of Experimental.
- Deprecate the Chaos Cloth Editor plugin and move the Chaos Cloth Editor module into the Chaos Cloth plugin.

diff --git a/Engine/Plugins/ChaosCloth/Source/ChaosCloth/Private/ChaosCloth/ChaosClothingSimulationCloth.cpp b/Engine/Plugins/ChaosCloth/Source/ChaosCloth/Private/ChaosCloth/ChaosClothingSimulationCloth.cpp
new file mode 100644
--- /dev/null
+++ b/Engine/Plugins/ChaosCloth/Source/ChaosCloth/Private/ChaosCloth/ChaosClothingSimulationCloth.cpp
@@ -0,0 +1,1395 @@
+// Copyright Epic Games, Inc. All Rights Reserved.
+
+#include "ChaosCloth/ChaosClothingSimulationCloth.h"
+#include "ChaosCloth/ChaosClothingSimulationSolver.h"
+#include "ChaosCloth/ChaosClothingSimulationMesh.h"
+#include "ChaosCloth/ChaosClothingSimulationCollider.h"
+#include "ChaosCloth/ChaosClothingSimulationConfig.h"
+#include "ChaosCloth/ChaosClothingPatternData.h"
+#include "ChaosCloth/ChaosWeightMapTarget.h"
+#include "ChaosCloth/ChaosClothPrivate.h"
+#include "Chaos/CollectionPropertyFacade.h"
+#include "Chaos/PBDFlatWeightMap.h"
+#include "Chaos/PBDSoftBodyCollisionConstraint.h"
+#include "Chaos/SoftsExternalForces.h"
+#include "Containers/ArrayView.h"
+#include "GeometryCollection/ManagedArrayCollection.h"
+#include "HAL/IConsoleManager.h"
+#include "ClothingSimulation.h"
+
+namespace Chaos
+{
+
+namespace ClothingSimulationClothDefault
+{
+	constexpr int32 MassMode = (int32)FClothingSimulationCloth::EMassMode::Density;
+	constexpr bool bUseGeodesicTethers = true;
+	constexpr float MassValue = 0.35f;
+	constexpr float MinPerParticleMass = 0.0001f;
+	constexpr float CollisionThickness = Softs::FPBDSoftBodyCollisionConstraint::DefaultCollisionThickness;
+	constexpr float FrictionCoefficient = Softs::FPBDSoftBodyCollisionConstraint::DefaultFrictionCoefficient;
+	constexpr float DampingCoefficient = 0.01f;
+	constexpr float Drag = 0.035f;
+	constexpr float Lift = 0.035f;
+	constexpr float Pressure = 0.f;
+	constexpr float AirDensity = 1.225f;  // Air density in kg/m^3
+	constexpr float GravityScale = Softs::FExternalForces::DefaultGravityScale; // 1.f;
+	constexpr float GravityZOverride = Softs::FExternalForces::DefaultGravityZOverride; // -980.665f;
+	constexpr float VelocityScale = 0.75f;
+	constexpr float MaxVelocityScale = 1.f;
+	constexpr float FictitiousAngularScale = Softs::FExternalForces::DefaultFictitiousAngularScale; // 1.f;
+	constexpr int32 MultiResCoarseLODIndex = INDEX_NONE;
+}
+
+namespace ClothingSimulationClothConsoleVariables
+{
+	TAutoConsoleVariable<bool> CVarLegacyDisablesAccurateWind(
+		TEXT("p.ChaosCloth.LegacyDisablesAccurateWind"),
+		true,
+		TEXT("Whether using the Legacy wind model switches off the accurate wind model, or adds up to it"));
+
+	TAutoConsoleVariable<float> CVarGravityMultiplier(
+		TEXT("p.ChaosCloth.GravityMultiplier"),
+		1.f,
+		TEXT("Scalar multiplier applied at the final stage of the cloth's gravity formulation."));
+}
+
+using namespace Softs;
+struct FClothingSimulationCloth::FLODData
+{
+	// Input mesh
+	const int32 NumParticles;
+	const TMap<FString, TConstArrayView<FRealSingle>> WeightMaps;
+	const TMap<FString, const TSet<int32>*> VertexSets;
+	const TMap<FString, const TSet<int32>*> FaceSets;
+	const TMap<FString, TConstArrayView<int32>> FaceIntMaps;
+	const TArray<TConstArrayView<TTuple<int32, int32, FRealSingle>>> Tethers;
+
+	const FClothingPatternData PatternData;
+	const FTriangleMesh NoOffsetTriangleMesh;
+
+	// Per Solver data
+	struct FSolverData
+	{
+		int32 LODIndex;
+		int32 ParticleRangeId;
+		FTriangleMesh OffsetTriangleMesh; // Only used if using PBD solver
+		int32 MultiResCoarseLODIndex = INDEX_NONE;
+	};
+	TMap<FClothingSimulationSolver*, FSolverData> SolverData;
+
+	// Stats
+	int32 NumKinematicParticles;
+	int32 NumDynamicParticles;
+
+	FLODData(
+		int32 InNumParticles,
+		const TConstArrayView<uint32>& InIndices,
+		const TConstArrayView<FVector2f>& InPatternPositions,
+		const TConstArrayView<uint32>& InPatternIndices,
+		const TConstArrayView<uint32>& InPatternToWeldedIndices,
+		TMap<FString, TConstArrayView<FRealSingle>>&& InWeightMaps,
+		TMap<FString, const TSet<int32>*>&& InVertexSets,
+		TMap<FString, const TSet<int32>*>&& InFaceSets,
+		TMap<FString, TConstArrayView<int32>>&& InFaceIntMaps,
+		const TArray<TConstArrayView<TTuple<int32, int32, FRealSingle>>>& InTethers);
+
+	static FTriangleMesh BuildTriangleMesh(const TConstArrayView<uint32>& Indices, const int32 NumParticles);
+
+	void AddParticles(FClothingSimulationSolver* Solver, FClothingSimulationCloth* Cloth, int32 LODIndex);
+	void Add(FClothingSimulationSolver* Solver, FClothingSimulationCloth* Cloth, int32 LODIndex);
+	void Remove(FClothingSimulationSolver* Solver);
+
+	void Update(FClothingSimulationSolver* Solver, FClothingSimulationCloth* Cloth);
+
+	void Enable(FClothingSimulationSolver* Solver, bool bEnable) const;
+
+	void ResetStartPose(FClothingSimulationSolver* Solver) const;
+
+	void UpdateNormals(FClothingSimulationSolver* Solver) const;
+
+	UE_CHAOS_DECLARE_INDEXLESS_PROPERTYCOLLECTION_NAME(MaxDistance, float);
+	UE_CHAOS_DECLARE_INDEXLESS_PROPERTYCOLLECTION_NAME(MassValue, float);
+	UE_CHAOS_DECLARE_INDEXLESS_PROPERTYCOLLECTION_NAME(MultiResCoarseLODIndex, int32);
+	UE_CHAOS_DECLARE_INDEXLESS_PROPERTYCOLLECTION_NAME(IsCoarseMultiResLOD, bool);
+
+};
+
+FClothingSimulationCloth::FLODData::FLODData(
+	int32 InNumParticles,
+	const TConstArrayView<uint32>& InIndices,
+	const TConstArrayView<FVector2f>& InPatternPositions,
+	const TConstArrayView<uint32>& InPatternIndices,
+	const TConstArrayView<uint32>& InPatternToWeldedIndices,
+	TMap<FString, TConstArrayView<FRealSingle>>&& InWeightMaps,
+	TMap<FString, const TSet<int32>*>&& InVertexSets,
+	TMap<FString, const TSet<int32>*>&& InFaceSets,
+	TMap<FString, TConstArrayView<int32>>&& InFaceIntMaps,
+	const TArray<TConstArrayView<TTuple<int32, int32, FRealSingle>>>& InTethers)
+	: NumParticles(InNumParticles)
+	, WeightMaps(MoveTemp(InWeightMaps))
+	, VertexSets(MoveTemp(InVertexSets))
+	, FaceSets(MoveTemp(InFaceSets))
+	, FaceIntMaps(MoveTemp(InFaceIntMaps))
+	, Tethers(InTethers)
+	, PatternData(NumParticles, InIndices, InPatternPositions, InPatternIndices, InPatternToWeldedIndices)
+	, NoOffsetTriangleMesh(BuildTriangleMesh(InIndices, NumParticles))
+{
+}
+
+FTriangleMesh FClothingSimulationCloth::FLODData::BuildTriangleMesh(const TConstArrayView<uint32>& Indices, const int32 InNumParticles)
+{
+	FTriangleMesh OutTriangleMesh;
+	// Build a sim friendly triangle mesh including the solver particle's Offset
+	const int32 NumElements = Indices.Num() / 3;
+	TArray<TVec3<int32>> Elements;
+	Elements.Reserve(NumElements);
+
+	for (int32 i = 0; i < NumElements; ++i)
+	{
+		const int32 Index = 3 * i;
+		Elements.Add(
+			{ static_cast<int32>(Indices[Index]),
+			 static_cast<int32>(Indices[Index + 1]),
+			 static_cast<int32>(Indices[Index + 2]) });
+	}
+
+	OutTriangleMesh.Init(MoveTemp(Elements), 0, InNumParticles - 1);
+	OutTriangleMesh.GetPointToTriangleMap(); // Builds map for later use by GetPointNormals(), and the velocity fields
+	return OutTriangleMesh;
+}
+
+void FClothingSimulationCloth::FLODData::AddParticles(FClothingSimulationSolver* Solver, FClothingSimulationCloth* Cloth, int32 InLODIndex)
+{
+	check(Solver);
+	check(Cloth);
+	check(Cloth->Mesh);
+
+	// Add a new solver data chunk
+	check(!SolverData.Find(Solver));
+	FSolverData& SolverDatum = SolverData.Add(Solver);
+	SolverDatum.LODIndex = InLODIndex;
+	// Add particles
+	SolverDatum.ParticleRangeId = Solver->AddParticles(NumParticles, Cloth->GroupId);
+	if (!NumParticles)
+	{
+		return;
+	}
+
+	if (!Solver->IsForceBasedSolver())
+	{
+		const TArray<TVec3<int32>>& Elements = NoOffsetTriangleMesh.GetElements();
+		TArray<TVec3<int32>> OffsetElements;
+		OffsetElements.Reserve(Elements.Num());
+		for (const TVec3<int32>& Element : Elements)
+		{
+			OffsetElements.Add({
+				Element[0] + SolverDatum.ParticleRangeId,
+				Element[1] + SolverDatum.ParticleRangeId,
+				Element[2] + SolverDatum.ParticleRangeId });
+		}
+
+		SolverDatum.OffsetTriangleMesh.Init(MoveTemp(OffsetElements), SolverDatum.ParticleRangeId, SolverDatum.ParticleRangeId + NumParticles - 1);
+		SolverDatum.OffsetTriangleMesh.GetPointToTriangleMap(); // Builds map for later use by GetPointNormals(), and the velocity fields
+	}
+
+	const FTriangleMesh& TriangleMesh = Solver->IsForceBasedSolver() ? NoOffsetTriangleMesh : SolverDatum.OffsetTriangleMesh;
+
+	PRAGMA_DISABLE_DEPRECATION_WARNINGS  // TODO: CHAOS_IS_CLOTHINGSIMULATIONMESH_ABSTRACT
+	// Update source mesh for this LOD, this is required prior to reset the start pose
+	Cloth->Mesh->Update(Solver, INDEX_NONE, InLODIndex, 0, SolverDatum.ParticleRangeId);
+	PRAGMA_ENABLE_DEPRECATION_WARNINGS
+
+	// Reset the particles start pose before setting up mass and constraints
+	ResetStartPose(Solver);
+
+	// Initialize the normals, in case the sim data is queried before the simulation steps
+	UpdateNormals(Solver);
+
+	// Retrieve config properties
+	check(Cloth->Config);
+	const Softs::FCollectionPropertyFacade& ConfigProperties = Cloth->Config->GetProperties(InLODIndex);
+
+	// Retrieve MaxDistance information (weight map and Low/High values)
+	const Softs::FPBDFlatWeightMapView MaxDistances(
+		GetWeightedFloatMaxDistance(ConfigProperties, FVector2f(0.f, 1.f)),
+		WeightMaps.FindRef(GetMaxDistanceString(ConfigProperties, MaxDistanceName.ToString())),
+		NumParticles);
+
+	// Set the particle masses
+	static const FRealSingle KinematicDistanceThreshold = 0.1f;  // TODO: This is not the same value as set in the painting UI but we might want to expose this value as parameter
+	auto KinematicPredicate =
+		[&MaxDistances](int32 Index)
+	{
+		return MaxDistances.GetValue(Index) < KinematicDistanceThreshold;
+	};
+
+	const int32 MassMode = ConfigProperties.GetValue<int32>(TEXT("MassMode"), ClothingSimulationClothDefault::MassMode);
+
+	constexpr FRealSingle MinPerParticleMassClampMin = UE_SMALL_NUMBER;
+	const FRealSingle MinPerParticleMass = FMath::Max((FRealSingle)ConfigProperties.GetValue<float>(TEXT("MinPerParticleMass"), ClothingSimulationClothDefault::MinPerParticleMass), MinPerParticleMassClampMin);
+
+	switch (MassMode)
+	{
+	default:
+		check(false);
+	case EMassMode::UniformMass:
+	{
+		const FVector2f MassValue = GetWeightedFloatMassValue(ConfigProperties, ClothingSimulationClothDefault::MassValue);
+		const TConstArrayView<float> MassValueMultipliers = WeightMaps.FindRef(GetMassValueString(ConfigProperties, MassValueName.ToString()));
+		Solver->SetParticleMassUniform(SolverDatum.ParticleRangeId, MassValue, MassValueMultipliers, MinPerParticleMass, TriangleMesh, KinematicPredicate);
+	}
+	break;
+	case EMassMode::TotalMass:
+	{
+		const FRealSingle MassValue = GetMassValue(ConfigProperties, ClothingSimulationClothDefault::MassValue);
+		Solver->SetParticleMassFromTotalMass(SolverDatum.ParticleRangeId, MassValue, MinPerParticleMass, TriangleMesh, KinematicPredicate);
+	}
+	break;
+	case EMassMode::Density:
+	{
+		const FVector2f MassValue = GetWeightedFloatMassValue(ConfigProperties, ClothingSimulationClothDefault::MassValue);
+		const TConstArrayView<float> MassValueMultipliers = WeightMaps.FindRef(GetMassValueString(ConfigProperties, MassValueName.ToString()));
+		Solver->SetParticleMassFromDensity(SolverDatum.ParticleRangeId, MassValue, MassValueMultipliers, MinPerParticleMass, TriangleMesh, KinematicPredicate);
+	}
+	break;
+	}
+}
+
+void FClothingSimulationCloth::FLODData::Add(FClothingSimulationSolver* Solver, FClothingSimulationCloth* Cloth, int32 InLODIndex)
+{
+	check(Solver);
+	check(Cloth);
+	check(Cloth->Mesh);
+
+	FSolverData& SolverDatum = SolverData.FindChecked(Solver);
+	check(SolverDatum.LODIndex == InLODIndex);
+	if (!NumParticles)
+	{
+		return;
+	}
+
+	check(SolverDatum.ParticleRangeId != INDEX_NONE);
+	const int32 ParticleRangeId = SolverDatum.ParticleRangeId;
+
+	PRAGMA_DISABLE_DEPRECATION_WARNINGS  // TODO: CHAOS_IS_CLOTHINGSIMULATIONMESH_ABSTRACT
+	// Retrieve the component's scale
+	const Softs::FSolverReal MeshScale = Cloth->Mesh->GetScale();
+	PRAGMA_ENABLE_DEPRECATION_WARNINGS
+
+	const FTriangleMesh& TriangleMesh = Solver->IsForceBasedSolver() ? NoOffsetTriangleMesh : SolverDatum.OffsetTriangleMesh;
+
+	// Retrieve config properties
+	check(Cloth->Config);
+	const Softs::FCollectionPropertyFacade& ConfigProperties = Cloth->Config->GetProperties(InLODIndex);
+
+	// Gather multires constraint data.
+	TSharedPtr<Softs::FMultiResConstraints> FineLODMultiResConstraint;
+	const FTriangleMesh* CoarseLODTriangleMesh = nullptr;
+	int32 CoarseLODParticleRangeId = INDEX_NONE;
+	if (Solver->IsForceBasedSolver())
+	{
+		// Multi-Res is only supported with ForceBasedSolver
+		if (InLODIndex == 0)
+		{
+			// Only allow LOD0 to be a fine LOD for now.
+			const int32 MultiResCoarseLODIndex = GetMultiResCoarseLODIndex(ConfigProperties, ClothingSimulationClothDefault::MultiResCoarseLODIndex);
+			if (MultiResCoarseLODIndex != INDEX_NONE && MultiResCoarseLODIndex != InLODIndex)
+			{
+				if (Cloth->Config->IsValidLOD(MultiResCoarseLODIndex) && Cloth->LODData.IsValidIndex(MultiResCoarseLODIndex))
+				{
+					// Check if coarse lod is setup correctly.
+					const Softs::FCollectionPropertyFacade& CoarseConfigProperties = Cloth->Config->GetProperties(MultiResCoarseLODIndex);
+					if (GetIsCoarseMultiResLOD(CoarseConfigProperties, false))
+					{
+						CoarseLODTriangleMesh = &Cloth->LODData[MultiResCoarseLODIndex]->NoOffsetTriangleMesh;
+						CoarseLODParticleRangeId = Cloth->LODData[MultiResCoarseLODIndex]->SolverData.FindChecked(Solver).ParticleRangeId;
+						SolverDatum.MultiResCoarseLODIndex = MultiResCoarseLODIndex;
+					}
+				}
+			}
+		}
+		else if (GetIsCoarseMultiResLOD(ConfigProperties, false))
+		{
+			// check that fine lod is setup correctly.
+			const Softs::FCollectionPropertyFacade& FineConfigProperties = Cloth->Config->GetProperties(0);
+			const int32 MultiResCoarseLODIndex = GetMultiResCoarseLODIndex(FineConfigProperties, ClothingSimulationClothDefault::MultiResCoarseLODIndex);
+			if (MultiResCoarseLODIndex == InLODIndex)
+			{
+				const int32 FineLODParticleRangeId = Cloth->LODData[0]->SolverData.FindChecked(Solver).ParticleRangeId;
+				FineLODMultiResConstraint = Solver->GetClothConstraints(FineLODParticleRangeId).GetMultiResConstraints();
+			}
+		}
+	}
+
+	// Setup solver constraints
+	FClothConstraints& ClothConstraints = Solver->GetClothConstraints(ParticleRangeId);
+
+	// Create constraints
+	const bool bEnabled = false;  // Set constraint disabled by default
+	ClothConstraints.AddRules(ConfigProperties, TriangleMesh, &PatternData, WeightMaps, VertexSets, FaceSets, FaceIntMaps, Tethers, MeshScale, bEnabled, CoarseLODTriangleMesh, CoarseLODParticleRangeId, FineLODMultiResConstraint);
+
+	// Update LOD stats
+	const TConstArrayView<Softs::FSolverReal> InvMasses(Solver->GetParticleInvMasses(ParticleRangeId), NumParticles);
+	NumKinematicParticles = 0;
+	NumDynamicParticles = 0;
+	for (int32 Index = 0; Index < NumParticles; ++Index)
+	{
+		if (InvMasses[Index] == (Softs::FSolverReal)0.)
+		{
+			++NumKinematicParticles;
+		}
+		else
+		{
+			++NumDynamicParticles;
+		}
+	}
+}
+
+void FClothingSimulationCloth::FLODData::Remove(FClothingSimulationSolver* Solver)
+{
+	SolverData.Remove(Solver);
+}
+
+void FClothingSimulationCloth::FLODData::Update(FClothingSimulationSolver* Solver, FClothingSimulationCloth* Cloth)
+{
+	check(Solver);
+	check(Cloth); 
+	const FSolverData& SolverDatum = SolverData.FindChecked(Solver);
+
+	const int32 ParticleRangeId = SolverDatum.ParticleRangeId;
+	check(ParticleRangeId != INDEX_NONE);
+
+	// Update the animatable constraint parameters
+	FClothConstraints& ClothConstraints = Solver->GetClothConstraints(ParticleRangeId);
+
+	check(Cloth->Config);
+PRAGMA_DISABLE_DEPRECATION_WARNINGS  // TODO: CHAOS_IS_CLOTHINGSIMULATIONMESH_ABSTRACT
+	const Softs::FSolverReal MeshScale = Cloth->Mesh->GetScale();
+PRAGMA_ENABLE_DEPRECATION_WARNINGS
+	const Softs::FSolverReal MaxDistancesScale = (Softs::FSolverReal)Cloth->MaxDistancesMultiplier;
+	ClothConstraints.Update(Cloth->Config->GetProperties(SolverDatum.LODIndex), WeightMaps, VertexSets, FaceSets, FaceIntMaps, MeshScale, MaxDistancesScale);
+}
+
+void FClothingSimulationCloth::FLODData::Enable(FClothingSimulationSolver* Solver, bool bEnable) const
+{
+	check(Solver);
+	const int32 ParticleRangeId = SolverData.FindChecked(Solver).ParticleRangeId;
+	check(ParticleRangeId != INDEX_NONE);
+	
+	// Enable particles (and related constraints)
+	Solver->EnableParticles(ParticleRangeId, bEnable);
+}
+
+void FClothingSimulationCloth::FLODData::ResetStartPose(FClothingSimulationSolver* Solver) const
+{
+	check(Solver);
+	const int32 ParticleRangeId = SolverData.FindChecked(Solver).ParticleRangeId;
+	check(ParticleRangeId != INDEX_NONE);
+
+	Solver->ResetStartPose(ParticleRangeId, NumParticles);
+}
+
+void FClothingSimulationCloth::FLODData::UpdateNormals(FClothingSimulationSolver* Solver) const
+{
+	check(Solver);
+
+	const FSolverData& SolverDatum = SolverData.FindChecked(Solver);
+	const int32 ParticleRangeId = SolverDatum.ParticleRangeId;
+
+	if (ParticleRangeId != INDEX_NONE)
+	{
+		TConstArrayView<Softs::FSolverVec3> Points = Solver->GetParticleXsView(ParticleRangeId);
+		TArrayView<Softs::FSolverVec3> Normals = Solver->GetNormalsView(ParticleRangeId);
+		TArray<Softs::FSolverVec3> FaceNormals;
+		NoOffsetTriangleMesh.GetFaceNormals(FaceNormals, Points, /*ReturnEmptyOnError =*/ false);
+		NoOffsetTriangleMesh.GetPointNormals(Normals, TConstArrayView<Softs::FSolverVec3>(FaceNormals), /*bUseGlobalArray =*/ false);
+	}
+}
+
+FClothingSimulationCloth::FClothingSimulationCloth(
+	FClothingSimulationConfig* InConfig,
+PRAGMA_DISABLE_DEPRECATION_WARNINGS  // TODO: CHAOS_IS_CLOTHINGSIMULATIONMESH_ABSTRACT
+	FClothingSimulationMesh* InMesh,
+PRAGMA_ENABLE_DEPRECATION_WARNINGS
+	TArray<FClothingSimulationCollider*>&& InColliders,
+	uint32 InGroupId)
+	: GroupId(InGroupId)
+{
+	SetConfig(InConfig);
+	SetMesh(InMesh);
+	SetColliders(MoveTemp(InColliders));
+}
+
+FClothingSimulationCloth::FClothingSimulationCloth(
+PRAGMA_DISABLE_DEPRECATION_WARNINGS  // TODO: CHAOS_IS_CLOTHINGSIMULATIONMESH_ABSTRACT
+	FClothingSimulationMesh* InMesh,
+PRAGMA_ENABLE_DEPRECATION_WARNINGS
+	TArray<FClothingSimulationCollider*>&& InColliders,
+	uint32 InGroupId,
+	EMassMode InMassMode,
+	FRealSingle InMassValue,
+	FRealSingle InMinPerParticleMass,
+	const TVec2<FRealSingle>& InEdgeStiffness,
+	const TVec2<FRealSingle>& InBendingStiffness,
+	FRealSingle InBucklingRatio,
+	const TVec2<FRealSingle>& InBucklingStiffness,
+	bool bInUseBendingElements,
+	const TVec2<FRealSingle>& InAreaStiffness,
+	FRealSingle /*InVolumeStiffness*/,  // Deprecated
+	bool /*bInUseThinShellVolumeConstraints*/,  // Deprecated
+	const TVec2<FRealSingle>& InTetherStiffness,
+	const TVec2<FRealSingle>& InTetherScale,
+PRAGMA_DISABLE_DEPRECATION_WARNINGS
+	ETetherMode InTetherMode,
+PRAGMA_ENABLE_DEPRECATION_WARNINGS
+	FRealSingle InMaxDistancesMultiplier,
+	const TVec2<FRealSingle>& InAnimDriveStiffness,
+	const TVec2<FRealSingle>& InAnimDriveDamping,
+	FRealSingle /*InShapeTargetStiffness*/,  // Deprecated
+	bool bInUseXPBDEdgeConstraints,
+	bool bInUseXPBDBendingConstraints,
+	bool /*bInUseXPBDAreaConstraints*/,  // Deprecated
+	FRealSingle InGravityScale,
+	bool bInUseGravityOverride,
+	const TVec3<FRealSingle>& InGravityOverride,
+	const TVec3<FRealSingle>& InLinearVelocityScale,
+	FRealSingle InAngularVelocityScale,
+	FRealSingle InFictitiousAngularScale,
+	const TVec2<FRealSingle>& InDrag,
+	const TVec2<FRealSingle>& InLift,
+	bool bInUsePointBasedWindModel,
+	const TVec2<FRealSingle>& InPressure,
+	FRealSingle InDampingCoefficient,
+	FRealSingle InLocalDampingCoefficient,
+	FRealSingle InCollisionThickness,
+	FRealSingle InFrictionCoefficient,
+	bool bInUseCCD,
+	bool bInUseSelfCollisions,
+	FRealSingle InSelfCollisionThickness,
+	FRealSingle InSelfCollisionFrictionCoefficient,
+	bool bInUseSelfIntersections,
+	bool bInUseLegacyBackstop,
+	bool bInUseLODIndexOverride, 
+	int32 InLODIndexOverride,
+	const TVec2<FRealSingle>& InEdgeDampingRatio,
+	const TVec2<FRealSingle>& InBendingDampingRatio)
+	: GroupId(InGroupId)
+	, PropertyCollection(MakeShared<FManagedArrayCollection>())
+	, bUseLODIndexOverride(bInUseLODIndexOverride)
+	, LODIndexOverride(InLODIndexOverride)
+{
+	// Turn parameters into a config
+	Softs::FCollectionPropertyMutableFacade Properties(PropertyCollection);
+	Properties.DefineSchema();
+
+	constexpr bool bEnable = true;
+	constexpr bool bAnimatable = true;
+	
+	// Mass
+	{
+		Properties.AddValue(TEXT("MassMode"), (int32)InMassMode);
+		Properties.AddValue(TEXT("MassValue"), InMassValue);
+		Properties.AddValue(TEXT("MinPerParticleMass"), InMinPerParticleMass);
+	}
+	
+	// Edge constraint
+	if (InEdgeStiffness[0] > 0.f || InEdgeStiffness[1] > 0.f)
+	{
+		const int32 EdgeSpringStiffnessIndex = Properties.AddProperty(TEXT("EdgeSpringStiffness"), bEnable, bAnimatable);
+		Properties.SetWeightedValue(EdgeSpringStiffnessIndex, InEdgeStiffness[0], InEdgeStiffness[1]);
+		Properties.SetStringValue(EdgeSpringStiffnessIndex, TEXT("EdgeStiffness"));
+	}
+	
+	// Bending constraint
+	if (InBendingStiffness[0] > 0.f || InBendingStiffness[1] > 0.f ||
+		(bInUseBendingElements && (InBucklingStiffness[0] > 0.f || InBucklingStiffness[1] > 0.f)))
+	{
+		if (bInUseBendingElements)
+		{
+			const int32 BendingElementStiffnessIndex = Properties.AddProperty(TEXT("BendingElementStiffness"), bEnable, bAnimatable);
+			Properties.SetWeightedValue(BendingElementStiffnessIndex, InBendingStiffness[0], InBendingStiffness[1]);
+			Properties.SetStringValue(BendingElementStiffnessIndex, TEXT("BendingStiffness"));
+
+			Properties.AddValue(TEXT("BucklingRatio"), InBucklingRatio);
+
+			if (InBucklingStiffness[0] > 0.f || InBucklingStiffness[1] > 0.f)
+			{
+				const int32 BucklingStiffnessIndex = Properties.AddProperty(TEXT("BucklingStiffness"), bEnable, bAnimatable);
+				Properties.SetWeightedValue(BucklingStiffnessIndex, InBucklingStiffness[0], InBucklingStiffness[1]);
+				Properties.SetStringValue(BucklingStiffnessIndex, TEXT("BucklingStiffness"));
+			}
+		}
+		else  // Not using bending elements
+		{
+			const int32 BendingSpringStiffnessIndex = Properties.AddProperty(TEXT("BendingSpringStiffness"), bEnable, bAnimatable);
+			Properties.SetWeightedValue(BendingSpringStiffnessIndex, InBendingStiffness[0], InBendingStiffness[1]);
+			Properties.SetStringValue(BendingSpringStiffnessIndex, TEXT("BendingStiffness"));
+		}
+	}
+
+	// Area constraint
+	if (InAreaStiffness[0] > 0.f || InAreaStiffness[1] > 0.f)
+	{
+		const int32 AreaSpringStiffnessIndex = Properties.AddProperty(TEXT("AreaSpringStiffness"), bEnable, bAnimatable);
+		Properties.SetWeightedValue(AreaSpringStiffnessIndex, InAreaStiffness[0], InAreaStiffness[1]);
+		Properties.SetStringValue(AreaSpringStiffnessIndex, TEXT("AreaStiffness"));
+	}
+
+	// Long range attachment
+	if (InTetherStiffness[0] > 0.f || InTetherStiffness[1] > 0.f)
+	{
+		PRAGMA_DISABLE_DEPRECATION_WARNINGS
+		const bool bUseGeodesicTethers = InTetherMode == ETetherMode::Geodesic;
+		PRAGMA_ENABLE_DEPRECATION_WARNINGS
+		Properties.AddValue(TEXT("UseGeodesicTethers"), bUseGeodesicTethers);
+
+		const int32 TetherStiffnessIndex = Properties.AddProperty(TEXT("TetherStiffness"), bEnable , bAnimatable);
+		Properties.SetWeightedValue(TetherStiffnessIndex, InTetherStiffness[0], InTetherStiffness[1]);
+		Properties.SetStringValue(TetherStiffnessIndex, TEXT("TetherStiffness"));
+	
+		const int32 TetherScaleIndex = Properties.AddProperty(TEXT("TetherScale"), bEnable, bAnimatable);
+		Properties.SetWeightedValue(TetherScaleIndex, InTetherScale[0], InTetherScale[1]);
+		Properties.SetStringValue(TetherScaleIndex, TEXT("TetherScale"));
+	}
+
+	// AnimDrive
+	if (InAnimDriveStiffness[0] > 0.f || InAnimDriveStiffness[1] > 0.f)
+	{
+		const int32 AnimDriveStiffnessIndex = Properties.AddProperty(TEXT("AnimDriveStiffness"), bEnable, bAnimatable);
+		Properties.SetWeightedValue(AnimDriveStiffnessIndex, InAnimDriveStiffness[0], InAnimDriveStiffness[1]);
+		Properties.SetStringValue(AnimDriveStiffnessIndex, TEXT("AnimDriveStiffness"));
+
+		const int32 AnimDriveDampingIndex = Properties.AddProperty(TEXT("AnimDriveDamping"), bEnable, bAnimatable);
+		Properties.SetWeightedValue(AnimDriveDampingIndex, InAnimDriveDamping[0], InAnimDriveDamping[1]);
+		Properties.SetStringValue(AnimDriveDampingIndex, TEXT("AnimDriveDamping"));
+	}
+
+	// Gravity
+	{
+		Properties.AddValue(TEXT("GravityScale"), InGravityScale, bEnable, bAnimatable);
+		Properties.AddValue(TEXT("UseGravityOverride"), bInUseGravityOverride, bEnable, bAnimatable);
+		Properties.AddValue(TEXT("GravityOverride"), FVector3f(InGravityOverride), bEnable, bAnimatable);
+	}
+
+	// Velocity scale
+	{
+		Properties.AddValue(TEXT("LinearVelocityScale"), FVector3f(InLinearVelocityScale), bEnable, bAnimatable);
+		Properties.AddValue(TEXT("AngularVelocityScale"), InAngularVelocityScale, bEnable, bAnimatable);
+		Properties.AddValue(TEXT("FictitiousAngularScale"), InFictitiousAngularScale, bEnable, bAnimatable);
+	}
+
+	// Aerodynamics
+	Properties.AddValue(TEXT("UsePointBasedWindModel"), bInUsePointBasedWindModel);
+	if (!bInUsePointBasedWindModel && (InDrag[0] > 0.f || InDrag[1] > 0.f || InLift[0] > 0.f || InLift[1] > 0.f))
+	{
+		const int32 DragIndex = Properties.AddProperty(TEXT("Drag"), bEnable, bAnimatable);
+		Properties.SetWeightedValue(DragIndex, InDrag[0], InDrag[1]);
+		Properties.SetStringValue(DragIndex, TEXT("Drag"));
+
+		const int32 LiftIndex = Properties.AddProperty(TEXT("Lift"), bEnable, bAnimatable);
+		Properties.SetWeightedValue(LiftIndex, InLift[0], InLift[1]);
+		Properties.SetStringValue(LiftIndex, TEXT("Lift"));
+
+		Properties.AddValue(TEXT("FluidDensity"), ClothingSimulationClothDefault::AirDensity, bEnable, bAnimatable);
+
+		Properties.AddValue(TEXT("WindVelocity"), FVector3f(0.f), bEnable, bAnimatable);  // Wind velocity must exist to be animatable
+	}
+
+	// Pressure
+	if (InPressure[0] != 0.f || InPressure[1] != 0.f)
+	{
+		const int32 PressureIndex = Properties.AddProperty(TEXT("Pressure"), bEnable, bAnimatable);
+		Properties.SetWeightedValue(PressureIndex, InPressure[0], InPressure[1]);
+		Properties.SetStringValue(PressureIndex, TEXT("Pressure"));
+	}
+
+	// Damping
+	Properties.AddValue(TEXT("DampingCoefficient"), InDampingCoefficient, bEnable, bAnimatable);
+	Properties.AddValue(TEXT("LocalDampingCoefficient"), InLocalDampingCoefficient, bEnable, bAnimatable);
+
+	// Collision
+	Properties.AddValue(TEXT("CollisionThickness"), InCollisionThickness, bEnable, bAnimatable);
+	Properties.AddValue(TEXT("FrictionCoefficient"), InFrictionCoefficient, bEnable, bAnimatable);
+	Properties.AddValue(TEXT("UseCCD"), bInUseCCD, bEnable, bAnimatable);
+	Properties.AddValue(TEXT("UseSelfCollisions"), bInUseSelfCollisions);
+	Properties.AddValue(TEXT("SelfCollisionThickness"), InSelfCollisionThickness);
+	Properties.AddValue(TEXT("SelfCollisionFrictionCoefficient"), InSelfCollisionFrictionCoefficient);
+	Properties.AddValue(TEXT("UseSelfIntersections"), bInUseSelfIntersections);
+	Properties.AddValue(TEXT("SelfCollideAgainstAllKinematicVertices"), true); // Match legacy behavior
+
+	// Max distance
+	{
+		const int32 MaxDistanceIndex = Properties.AddProperty(TEXT("MaxDistance"));
+		Properties.SetWeightedValue(MaxDistanceIndex, 0.f, 1.f);  // Backward compatibility with legacy mask must use a unit range since the multiplier is in the mask
+		Properties.SetStringValue(MaxDistanceIndex, TEXT("MaxDistance"));
+	}
+
+	// Backstop
+	{
+		const int32 BackstopScaleIndex = Properties.AddProperty(TEXT("BackstopDistance"));
+		Properties.SetWeightedValue(BackstopScaleIndex, 0.f, 1.f);  // Backward compatibility with legacy mask must use a unit range since the multiplier is in the mask
+		Properties.SetStringValue(BackstopScaleIndex, TEXT("BackstopDistance"));
+
+		const int32 BackstopRadiusIndex = Properties.AddProperty(TEXT("BackstopRadius"));
+		Properties.SetWeightedValue(BackstopRadiusIndex, 0.f, 1.f);  // Backward compatibility with legacy mask must use a unit range since the multiplier is in the mask
+		Properties.SetStringValue(BackstopRadiusIndex, TEXT("BackstopRadius"));
+
+		Properties.AddValue(TEXT("UseLegacyBackstop"), bInUseLegacyBackstop);
+	}
+
+	PRAGMA_DISABLE_DEPRECATION_WARNINGS
+	Config = new FClothingSimulationConfig(PropertyCollection);
+	PRAGMA_ENABLE_DEPRECATION_WARNINGS
+
+	// Set mesh and colliders
+	SetMesh(InMesh);
+	SetColliders(MoveTemp(InColliders));
+}
+
+FClothingSimulationCloth::~FClothingSimulationCloth()
+{
+	// If the PropertyCollection is owned by this object, so does the current config object
+	if (PropertyCollection.IsValid())
+	{
+		delete Config;
+	}
+}
+
+PRAGMA_DISABLE_DEPRECATION_WARNINGS  // TODO: CHAOS_IS_CLOTHINGSIMULATIONMESH_ABSTRACT
+void FClothingSimulationCloth::SetMesh(FClothingSimulationMesh* InMesh)
+{
+	Mesh = InMesh;
+	
+	// Reset LODs
+	const int32 NumLODs = Mesh ? Mesh->GetNumLODs() : 0;
+	LODData.Reset(NumLODs);
+	for (int32 LODIndex = 0; LODIndex < NumLODs; ++LODIndex)
+	{
+		// Regenerate LOD weight maps lookup map
+		const TArray<FName> WeightMapNames = Mesh->GetWeightMapNames(LODIndex);
+		TMap<FString, TConstArrayView<FRealSingle>> WeightMaps;
+		WeightMaps.Reserve(WeightMapNames.Num());
+
+		const TArray<TConstArrayView<FRealSingle>> WeightMapArray = Mesh->GetWeightMaps(LODIndex);
+		ensure(WeightMapArray.Num() == WeightMapNames.Num());
+
+		for (int32 WeightMapIndex = 0; WeightMapIndex < WeightMapNames.Num(); ++WeightMapIndex)
+		{
+			WeightMaps.Add(WeightMapNames[WeightMapIndex].ToString(),
+				WeightMapArray.IsValidIndex(WeightMapIndex) ?
+					WeightMapArray[WeightMapIndex] :
+					TConstArrayView<FRealSingle>());
+		}
+
+		TMap<FString, const TSet<int32>*> VertexSets = Mesh->GetVertexSets(LODIndex);
+		TMap<FString, const TSet<int32>*> FaceSets = Mesh->GetFaceSets(LODIndex);
+		TMap<FString, TConstArrayView<int32>> FaceIntMaps = Mesh->GetFaceIntMaps(LODIndex);
+
+		const bool bUseGeodesicTethers = Config->GetProperties(LODIndex).GetValue<bool>(TEXT("UseGeodesicTethers"), ClothingSimulationClothDefault::bUseGeodesicTethers);
+
+		// Add LOD data
+		LODData.Add(MakeUnique<FLODData>(
+			Mesh->GetNumPoints(LODIndex),
+			Mesh->GetIndices(LODIndex),
+			Mesh->GetPatternPositions(LODIndex),
+			Mesh->GetPatternIndices(LODIndex),
+			Mesh->GetPatternToWeldedIndices(LODIndex),
+			MoveTemp(WeightMaps),
+			MoveTemp(VertexSets),
+			MoveTemp(FaceSets),
+			MoveTemp(FaceIntMaps),
+			Mesh->GetTethers(LODIndex, bUseGeodesicTethers)));
+	}
+
+	// Iterate all known solvers
+	TArray<FClothingSimulationSolver*> Solvers;
+	LODIndices.GetKeys(Solvers);
+	for (FClothingSimulationSolver* const Solver : Solvers)
+	{
+		// Refresh this cloth to recreate particles
+		Solver->RefreshCloth(this);
+	}
+}
+PRAGMA_ENABLE_DEPRECATION_WARNINGS
+
+void FClothingSimulationCloth::SetConfig(FClothingSimulationConfig* InConfig)
+{
+	// If the PropertyCollection is owned by this object, so does the current config object
+	if (PropertyCollection.IsValid())
+	{
+		delete Config;
+		PropertyCollection.Reset();
+	}
+
+	if (InConfig)
+	{
+		Config = InConfig;
+	}
+	else
+	{
+		// Create a default empty config object for coherence
+		PropertyCollection = MakeShared<FManagedArrayCollection>();
+		PRAGMA_DISABLE_DEPRECATION_WARNINGS
+		Config = new FClothingSimulationConfig(PropertyCollection); 
+		PRAGMA_ENABLE_DEPRECATION_WARNINGS
+	}
+}
+
+void FClothingSimulationCloth::SetColliders(TArray<FClothingSimulationCollider*>&& InColliders)
+{
+	// Empty the collider list, but keep the pointers around for the removal operation below
+	const TArray<FClothingSimulationCollider*> TempColliders = MoveTemp(Colliders);
+
+	// Replace with the new colliders
+	Colliders = InColliders;
+
+	// Iterate all known solvers
+	TArray<FClothingSimulationSolver*> Solvers;
+	LODIndices.GetKeys(Solvers);
+	for (FClothingSimulationSolver* const Solver : Solvers)
+	{
+		// Remove any held collider data related to this cloth simulation
+		for (FClothingSimulationCollider* const Collider : TempColliders)
+		{
+			Collider->Remove(Solver, this);
+		}
+
+		// Refresh this cloth to recreate collision particles
+		Solver->RefreshCloth(this);
+	}
+}
+
+void FClothingSimulationCloth::AddCollider(FClothingSimulationCollider* InCollider)
+{
+	check(InCollider);
+
+	if (Colliders.Find(InCollider) != INDEX_NONE)
+	{
+		return;
+	}
+
+	// Add the collider to the solver update array
+	Colliders.Emplace(InCollider);
+
+	// Iterate all known solvers
+	TArray<FClothingSimulationSolver*> Solvers;
+	LODIndices.GetKeys(Solvers);
+	for (FClothingSimulationSolver* const Solver : Solvers)
+	{
+		// Refresh this cloth to recreate collision particles
+		Solver->RefreshCloth(this);
+	}
+}
+
+void FClothingSimulationCloth::RemoveCollider(FClothingSimulationCollider* InCollider)
+{
+	if (Colliders.Find(InCollider) == INDEX_NONE)
+	{
+		return;
+	}
+
+	// Remove collider from array
+	Colliders.RemoveSwap(InCollider);
+
+	// Iterate all known solvers
+	TArray<FClothingSimulationSolver*> Solvers;
+	LODIndices.GetKeys(Solvers);
+	for (FClothingSimulationSolver* const Solver : Solvers)
+	{
+		// Remove any held collider data related to this cloth simulation
+		InCollider->Remove(Solver, this);
+
+		// Refresh this cloth to recreate collision particles
+		Solver->RefreshCloth(this);
+	}
+}
+
+void FClothingSimulationCloth::RemoveColliders()
+{
+	// Empty the collider list, but keep the pointers around for the removal operation below
+	const TArray<FClothingSimulationCollider*> TempColliders = MoveTemp(Colliders);
+
+	// Iterate all known solvers
+	TArray<FClothingSimulationSolver*> Solvers;
+	LODIndices.GetKeys(Solvers);
+	for (FClothingSimulationSolver* const Solver : Solvers)
+	{
+		// Remove any held collider data related to this cloth simulation
+		for (FClothingSimulationCollider* const Collider : TempColliders)
+		{
+			Collider->Remove(Solver, this);
+		}
+
+		// Refresh this cloth to recreate collision particles
+		Solver->RefreshCloth(this);
+	}
+}
+
+void FClothingSimulationCloth::Add(FClothingSimulationSolver* Solver)
+{
+	check(Solver);
+
+	// Can't add a cloth twice to the same solver
+	check(!LODIndices.Find(Solver));
+
+	// Initialize LODIndex
+	int32& LODIndex = LODIndices.Add(Solver);
+	LODIndex = INDEX_NONE;
+
+	// Add all particles first and in reverse order. This is necessary so that any multires coarse lods soft bodies are added first, and all particle offsets are setup
+	// when adding the LOD constraints.
+	for (int32 Index = LODData.Num() - 1; Index >= 0; --Index)
+	{
+		LODData[Index]->AddParticles(Solver, this, Index);
+	}
+	// Now add the LODs themselves. These need to go in normal order since the coarse lod needs the fine lod constraints.
+	for (int32 Index = 0; Index < LODData.Num(); ++Index)
+	{
+		LODData[Index]->Add(Solver, this, Index);
+	}
+
+	// Add Colliders
+	for (FClothingSimulationCollider* Collider : Colliders)
+	{
+		Collider->Add(Solver, this);
+	}
+}
+
+void FClothingSimulationCloth::Remove(FClothingSimulationSolver* Solver)
+{
+	// Remove Colliders
+	for (FClothingSimulationCollider* Collider : Colliders)
+	{
+		Collider->Remove(Solver, this);
+	}
+
+	// Remove solver from maps
+	LODIndices.Remove(Solver);
+	for (TUniquePtr<FLODData>& LODDatum: LODData)
+	{
+		LODDatum->Remove(Solver);
+	}
+}
+
+int32 FClothingSimulationCloth::GetNumParticles(int32 InLODIndex) const
+{
+	return LODData.IsValidIndex(InLODIndex) ? LODData[InLODIndex]->NumParticles : 0;
+}
+
+int32 FClothingSimulationCloth::GetParticleRangeId(const FClothingSimulationSolver* Solver, int32 InLODIndex) const
+{
+	return LODData.IsValidIndex(InLODIndex) ? LODData[InLODIndex]->SolverData.FindChecked(Solver).ParticleRangeId : 0;
+}
+
+TVec3<FRealSingle> FClothingSimulationCloth::GetGravity(const FClothingSimulationSolver* Solver) const
+{
+	check(Solver);
+	if (Solver->IsForceBasedSolver())
+	{
+		const int32 ParticleRangeId = GetParticleRangeId(Solver);
+		if (const Softs::FExternalForces* const ExternalForces = Solver->GetClothConstraints(ParticleRangeId).GetExternalForces().Get())
+		{
+			return ExternalForces->GetGravity();
+		}
+		else
+		{
+			return TVec3<FRealSingle>(0.f, 0.f, Softs::FExternalForces::DefaultGravityZOverride);
+		}
+	}
+	else
+	{
+		check(Config);
+		const Softs::FCollectionPropertyFacade& ConfigProperties = Config->GetProperties(GetLODIndex(Solver));
+
+		const bool bUseGravityOverride = ConfigProperties.GetValue<bool>(TEXT("UseGravityOverride"));
+		const TVec3<FRealSingle> GravityOverride = (TVec3<FRealSingle>)ConfigProperties.GetValue<FVector3f>(TEXT("GravityOverride"), FVector3f(0.f, 0.f, ClothingSimulationClothDefault::GravityZOverride));
+		const FRealSingle GravityScale = (FRealSingle)ConfigProperties.GetValue<float>(TEXT("GravityScale"), 1.f);
+		const FRealSingle GravityMultiplier = (FRealSingle)ClothingSimulationClothConsoleVariables::CVarGravityMultiplier.GetValueOnAnyThread();
+
+		return (Solver->IsClothGravityOverrideEnabled() && bUseGravityOverride ? GravityOverride : Solver->GetGravity() * GravityScale) * GravityMultiplier;
+	}
+}
+
+FAABB3 FClothingSimulationCloth::CalculateBoundingBox(const FClothingSimulationSolver* Solver) const
+{
+	check(Solver);
+
+	// Calculate local space bounding box
+	Softs::FSolverAABB3 BoundingBox = Softs::FSolverAABB3::EmptyAABB();
+
+	const TConstArrayView<Softs::FSolverVec3> ParticlePositions = GetParticlePositions(Solver);
+	for (const Softs::FSolverVec3& ParticlePosition : ParticlePositions)
+	{
+		BoundingBox.GrowToInclude(ParticlePosition);
+	}
+
+	// Return world space bounding box
+	return FAABB3(BoundingBox).TransformedAABB(FRigidTransform3(Solver->GetLocalSpaceLocation(), FRotation3::Identity));
+}
+
+int32 FClothingSimulationCloth::GetParticleRangeId(const FClothingSimulationSolver* Solver) const
+{
+	const int32 LODIndex = LODIndices.FindChecked(Solver);
+	return LODData.IsValidIndex(LODIndex) ? GetParticleRangeId(Solver, LODIndex) : INDEX_NONE;
+}
+
+int32 FClothingSimulationCloth::GetNumParticles(const FClothingSimulationSolver* Solver) const
+{
+	const int32 LODIndex = LODIndices.FindChecked(Solver);
+	return LODData.IsValidIndex(LODIndex) ? GetNumParticles(LODIndex) : 0;
+}
+
+const FTriangleMesh& FClothingSimulationCloth::GetTriangleMesh(const FClothingSimulationSolver* Solver) const
+{
+	const int32 LODIndex = LODIndices.FindChecked(Solver);
+	static const FTriangleMesh EmptyTriangleMesh;
+
+	return LODData.IsValidIndex(LODIndex) ? (Solver->IsForceBasedSolver() ? LODData[LODIndex]->NoOffsetTriangleMesh : LODData[LODIndex]->SolverData.FindChecked(Solver).OffsetTriangleMesh) : EmptyTriangleMesh;
+}
+
+
+// Deprecated for 5.3
+const TArray<TConstArrayView<FRealSingle>>& FClothingSimulationCloth::GetWeightMaps(const FClothingSimulationSolver* Solver) const
+{
+	static const TArray<TConstArrayView<FRealSingle>> EmptyWeightMaps;
+	return EmptyWeightMaps;  // Can't return a reference anymore, so returning an empty weightmap array is probably the safest for existing code
+}
+
+TConstArrayView<FRealSingle> FClothingSimulationCloth::GetWeightMapByName(const FClothingSimulationSolver* Solver, const FString& Name) const
+{
+	const int32 LODIndex = LODIndices.FindChecked(Solver);
+	return LODData.IsValidIndex(LODIndex) ? LODData[LODIndex]->WeightMaps.FindRef(Name) : TConstArrayView<FRealSingle>();
+}
+
+TConstArrayView<FRealSingle> FClothingSimulationCloth::GetWeightMapByProperty(const FClothingSimulationSolver* Solver, const FString& Property) const
+{
+	check(Config);
+	const FString PropertyString = Config->GetProperties(GetLODIndex(Solver)).GetStringValue(Property);
+	return GetWeightMapByName(Solver, PropertyString);
+}
+
+TConstArrayView<int32> FClothingSimulationCloth::GetFaceIntMapByName(const FClothingSimulationSolver* Solver, const FString& Name) const
+{
+	const int32 LODIndex = LODIndices.FindChecked(Solver);
+	return LODData.IsValidIndex(LODIndex) ? LODData[LODIndex]->FaceIntMaps.FindRef(Name) : TConstArrayView<int32>();
+}
+
+TConstArrayView<int32> FClothingSimulationCloth::GetFaceIntMapByProperty(const FClothingSimulationSolver* Solver, const FString& Property) const
+{
+	check(Config);
+	const FString PropertyString = Config->GetProperties(GetLODIndex(Solver)).GetStringValue(Property);
+	return GetFaceIntMapByName(Solver, PropertyString);
+}
+
+const TArray<TConstArrayView<TTuple<int32, int32, float>>>& FClothingSimulationCloth::GetTethers(const FClothingSimulationSolver* Solver) const
+{
+	const int32 LODIndex = LODIndices.FindChecked(Solver);
+	static const TArray<TConstArrayView<TTuple<int32, int32, float>>> EmptyTethers;
+	return LODData.IsValidIndex(LODIndex) ? LODData[LODIndex]->Tethers : EmptyTethers;
+}
+
+int32 FClothingSimulationCloth::GetReferenceBoneIndex() const
+{
+PRAGMA_DISABLE_DEPRECATION_WARNINGS  // TODO: CHAOS_IS_CLOTHINGSIMULATIONMESH_ABSTRACT
+	return Mesh ? Mesh->GetReferenceBoneIndex() : INDEX_NONE;
+PRAGMA_ENABLE_DEPRECATION_WARNINGS
+}
+
+void FClothingSimulationCloth::PreUpdate(FClothingSimulationSolver* Solver)
+{
+	check(Solver);
+
+	// Exit if the input mesh is missing
+	if (!Mesh)
+	{
+		return;
+	}
+
+	// Update Cloth Colliders
+	{
+		SCOPE_CYCLE_COUNTER(STAT_ClothUpdateCollisions);
+
+		for (FClothingSimulationCollider* Collider : Colliders)
+		{
+			Collider->PreUpdate(Solver, this);
+		}
+	}
+}
+
+void FClothingSimulationCloth::Update(FClothingSimulationSolver* Solver)
+{
+	check(Solver);
+
+	// Exit if the input mesh is missing
+	if (!Mesh)
+	{
+		return;
+	}
+
+	// Retrieve LOD Index, either from the override, or from the mesh input
+	int32& LODIndex = LODIndices.FindChecked(Solver);  // Must be added to solver first
+
+	const int32 PrevLODIndex = LODIndex;
+PRAGMA_DISABLE_DEPRECATION_WARNINGS  // TODO: CHAOS_IS_CLOTHINGSIMULATIONMESH_ABSTRACT
+	LODIndex = bUseLODIndexOverride && LODData.IsValidIndex(LODIndexOverride) ? LODIndexOverride : Mesh->GetLODIndex();
+
+	// Update reference space transform from the mesh's reference bone transform  TODO: Add override in the style of LODIndexOverride
+	const FRigidTransform3 OldReferenceSpaceTransform = ReferenceSpaceTransform;
+	ReferenceSpaceTransform = Mesh->GetReferenceBoneTransform();
+PRAGMA_ENABLE_DEPRECATION_WARNINGS
+	ReferenceSpaceTransform.SetScale3D(FVec3(1.f));
+
+	// Update Cloth Colliders
+	{
+		SCOPE_CYCLE_COUNTER(STAT_ClothUpdateCollisions);
+
+		for (FClothingSimulationCollider* Collider : Colliders)
+		{
+			Collider->Update(Solver, this);
+		}
+	}
+
+	// Update the source mesh skinned positions
+	const int32 PrevParticleRangeId = GetParticleRangeId(Solver, PrevLODIndex);
+	const int32 ParticleRangeId = GetParticleRangeId(Solver, LODIndex);
+	if (PrevParticleRangeId == INDEX_NONE || ParticleRangeId == INDEX_NONE)
+	{
+		return;
+	}
+
+PRAGMA_DISABLE_DEPRECATION_WARNINGS  // TODO: CHAOS_IS_CLOTHINGSIMULATIONMESH_ABSTRACT
+	Mesh->Update(Solver, PrevLODIndex, LODIndex, PrevParticleRangeId, ParticleRangeId);
+PRAGMA_ENABLE_DEPRECATION_WARNINGS
+
+	// Retrieve config
+	check(Config);
+	const Softs::FCollectionPropertyFacade& ConfigProperties = Config->GetProperties(LODIndex);
+
+	const int32 CoarseLODIndex = LODIndex != INDEX_NONE ? LODData[LODIndex]->SolverData.FindChecked(Solver).MultiResCoarseLODIndex : INDEX_NONE;
+	const int32 CoarseParticleRangeId = CoarseLODIndex != INDEX_NONE ? GetParticleRangeId(Solver, CoarseLODIndex) : INDEX_NONE;
+	if (CoarseLODIndex != INDEX_NONE)
+	{
+		LODData[CoarseLODIndex]->Enable(Solver, true);
+PRAGMA_DISABLE_DEPRECATION_WARNINGS 
+		// TODO: interpolate/ reset when LOD switching to enable multires
+		Mesh->Update(Solver, CoarseLODIndex, CoarseLODIndex, CoarseParticleRangeId, CoarseParticleRangeId);
+PRAGMA_ENABLE_DEPRECATION_WARNINGS
+	}
+
+	// LOD Switching
+	if (LODIndex != PrevLODIndex)
+	{
+		if (PrevLODIndex != INDEX_NONE)
+		{
+			if (PrevLODIndex != CoarseLODIndex)
+			{
+				// Disable previous LOD's particles
+				LODData[PrevLODIndex]->Enable(Solver, false);
+			}
+			const int32 PrevCoarseLODIndex = LODData[PrevLODIndex]->SolverData.FindChecked(Solver).MultiResCoarseLODIndex;
+			if (PrevCoarseLODIndex != INDEX_NONE && PrevCoarseLODIndex != CoarseLODIndex && PrevCoarseLODIndex != LODIndex)
+			{
+				// Disable previous coarse LOD's particles
+				LODData[PrevCoarseLODIndex]->Enable(Solver, false);
+			}
+		}
+		if (LODIndex != INDEX_NONE)
+		{
+			// Enable new LOD's particles
+			LODData[LODIndex]->Enable(Solver, true);
+			NumActiveKinematicParticles = LODData[LODIndex]->NumKinematicParticles;
+			NumActiveDynamicParticles = LODData[LODIndex]->NumDynamicParticles;
+
+			// Wrap new LOD based on previous LOD if possible (can only do 1 level LOD at a time, and if previous LOD exists)
+PRAGMA_DISABLE_DEPRECATION_WARNINGS  // TODO: CHAOS_IS_CLOTHINGSIMULATIONMESH_ABSTRACT
+			bNeedsReset = bNeedsReset || !Mesh->WrapDeformLOD(
+PRAGMA_ENABLE_DEPRECATION_WARNINGS
+				PrevLODIndex,
+				LODIndex,
+				Solver->GetNormals(PrevParticleRangeId),
+				Solver->GetParticlePandInvMs(PrevParticleRangeId),
+				Solver->GetParticleVs(PrevParticleRangeId),
+				Solver->GetParticlePandInvMs(ParticleRangeId),
+				Solver->GetParticleXs(ParticleRangeId),
+				Solver->GetParticleVs(ParticleRangeId));
+
+			if(!Solver->IsForceBasedSolver())
+			{
+				// Update the wind velocity field for the new LOD mesh
+				Solver->SetWindAndPressureGeometry(GroupId, GetTriangleMesh(Solver), ConfigProperties, LODData[LODIndex]->WeightMaps);
+			}
+		}
+		else
+		{
+			NumActiveKinematicParticles = 0;
+			NumActiveDynamicParticles = 0;
+		}
+	}
+
+	// Update Cloth group parameters  TODO: Cloth groups should exist as their own node object so that they can be used by several cloth objects
+	if (LODIndex != INDEX_NONE)
+	{
+		// TODO: Move all groupID updates out of the cloth update to allow to use of the same GroupId with different cloths
+
+		// Set the reference input velocity and deal with teleport & reset; external forces depends on these values, so they must be initialized before then
+		FVec3f OutLinearVelocityScale;
+		FRealSingle OutAngularVelocityScale;
+		FRealSingle OutMaxVelocityScale;
+		if (bNeedsReset)
+		{
+			// Make sure not to do any pre-sim transform just after a reset
+			OutLinearVelocityScale = FVec3f(1.f);
+			OutAngularVelocityScale = 1.f;
+			OutMaxVelocityScale = 1.f;
+
+			// Reset to start pose
+			LODData[LODIndex]->ResetStartPose(Solver);
+			for (FClothingSimulationCollider* Collider : Colliders)
+			{
+				Collider->ResetStartPose(Solver, this);
+			}
+
+			if (CoarseLODIndex != INDEX_NONE)
+			{
+				LODData[CoarseLODIndex]->ResetStartPose(Solver);
+			}
+			UE_LOG(LogChaosCloth, VeryVerbose, TEXT("Cloth in group Id %d Needs reset."), GroupId);
+		}
+		else if (bNeedsTeleport)
+		{
+			// Remove all impulse velocity from the last frame
+			OutLinearVelocityScale = FVec3f(0.f);
+			OutAngularVelocityScale = 0.f;
+			OutMaxVelocityScale = 1.f;
+			UE_LOG(LogChaosCloth, VeryVerbose, TEXT("Cloth in group Id %d Needs teleport."), GroupId);
+		}
+		else
+		{
+			// Use the cloth config parameters
+			OutLinearVelocityScale = ConfigProperties.GetValue<FVector3f>(TEXT("LinearVelocityScale"), FVector3f(ClothingSimulationClothDefault::VelocityScale));
+			OutAngularVelocityScale = ConfigProperties.GetValue<float>(TEXT("AngularVelocityScale"), ClothingSimulationClothDefault::VelocityScale);
+			OutMaxVelocityScale = ConfigProperties.GetValue<float>(TEXT("MaxVelocityScale"), ClothingSimulationClothDefault::MaxVelocityScale);
+		}
+
+		// NOTE: Force-based solver doesn't actually use FictitiousAngularScale here. It gets it from the property collection directly.
+		const FRealSingle FictitiousAngularScale = ConfigProperties.GetValue<float>(TEXT("FictitiousAngularScale"), ClothingSimulationClothDefault::FictitiousAngularScale);
+		Solver->SetReferenceVelocityScale(
+			GroupId,
+			OldReferenceSpaceTransform,
+			ReferenceSpaceTransform,
+			OutLinearVelocityScale,
+			OutAngularVelocityScale,
+			FictitiousAngularScale,
+			OutMaxVelocityScale);
+		if (Solver->IsForceBasedSolver())
+		{
+			Solver->SetProperties(ParticleRangeId, ConfigProperties, LODData[LODIndex]->WeightMaps);
+			if (CoarseLODIndex != INDEX_NONE)
+			{
+				Solver->SetProperties(CoarseParticleRangeId, Config->GetProperties(CoarseLODIndex), LODData[CoarseLODIndex]->WeightMaps);
+			}
+		}
+		else
+		{
+			// Update gravity
+			// This code relies on the solver gravity property being already set.
+			// In order to use a cloth gravity override, it must first be enabled by the solver so that an override at solver level can still take precedence if needed.
+			// In all cases apart from when the cloth override is used, the gravity scale must be combined to the solver gravity value.
+			Solver->SetGravity(GroupId, GetGravity(Solver));
+
+			// External forces (legacy wind+field)
+			const bool bUsePointBasedWindModel = ConfigProperties.GetValue<bool>(TEXT("UsePointBasedWindModel"));
+			Solver->AddExternalForces(GroupId, bUsePointBasedWindModel);
+
+			const bool bPointBasedWindDisablesAccurateWind = ClothingSimulationClothConsoleVariables::CVarLegacyDisablesAccurateWind.GetValueOnAnyThread();
+			const bool bEnableAerodynamics = !(bUsePointBasedWindModel && bPointBasedWindDisablesAccurateWind);
+			Solver->SetWindAndPressureProperties(GroupId, ConfigProperties, LODData[LODIndex]->WeightMaps, bEnableAerodynamics);
+
+			constexpr float WorldScale = 100.f;  // VelocityField wind is in m/s in the config (same as the wind unit), but cm/s in the solver  TODO: Cleanup the Solver SetWindVelocity functions to be consistent with the unit
+			const FVec3f WindVelocity = ConfigProperties.GetValue<FVector3f>(TEXT("WindVelocity")) * WorldScale;
+			Solver->SetWindVelocity(GroupId, WindVelocity + Solver->GetWindVelocity());
+
+			// Update general solver properties
+			PRAGMA_DISABLE_DEPRECATION_WARNINGS  // TODO: CHAOS_IS_CLOTHINGSIMULATIONMESH_ABSTRACT
+			const Softs::FSolverReal MeshScale = Mesh->GetScale();
+			PRAGMA_ENABLE_DEPRECATION_WARNINGS
+
+			const FRealSingle DampingCoefficient = ConfigProperties.GetValue<float>(TEXT("DampingCoefficient"), ClothingSimulationClothDefault::DampingCoefficient);
+			const FRealSingle LocalDampingCoefficient = ConfigProperties.GetValue<float>(TEXT("LocalDampingCoefficient"));
+			const FRealSingle CollisionThickness = ConfigProperties.GetValue<float>(TEXT("CollisionThickness"), ClothingSimulationClothDefault::CollisionThickness);
+			const FRealSingle FrictionCoefficient = ConfigProperties.GetValue<float>(TEXT("FrictionCoefficient"), ClothingSimulationClothDefault::FrictionCoefficient);
+			Solver->SetProperties(GroupId, DampingCoefficient, LocalDampingCoefficient, CollisionThickness * MeshScale, FrictionCoefficient);
+
+			// Update use of continuous collision detection
+			const bool bUseCCD = ConfigProperties.GetValue<bool>(TEXT("UseCCD"));
+			Solver->SetUseCCD(GroupId, bUseCCD);
+		}
+
+		// This will be updated below if single legacy lod
+		if (!Config->IsLegacySingleLOD())
+		{
+			LODData[LODIndex]->Update(Solver, this);
+			if (CoarseLODIndex != INDEX_NONE)
+			{
+				LODData[CoarseLODIndex]->Update(Solver, this);
+			}
+		}
+	}
+
+	// Update all LODs dirty properties, since it is easier done than re-updating all properties when switching LODs
+	if (Config->IsLegacySingleLOD())
+	{
+		for (TUniquePtr<FLODData>& LODDatum : LODData)
+		{
+			LODDatum->Update(Solver, this);
+		}
+	}
+
+	// Reset trigger flags
+	bNeedsTeleport = false;
+	bNeedsReset = false;
+}
+
+void FClothingSimulationCloth::PostUpdate(FClothingSimulationSolver* Solver)
+{
+	check(Solver);
+
+	const int32 LODIndex = LODIndices.FindChecked(Solver);
+	if (LODIndex != INDEX_NONE)
+	{
+		// Update normals
+		LODData[LODIndex]->UpdateNormals(Solver);
+	}
+}
+
+void FClothingSimulationCloth::UpdateFromCache(const FClothingSimulationCacheData& CacheData)
+{
+	if (const FTransform* CachedReferenceSpaceTransform = CacheData.CachedReferenceSpaceTransforms.Find(GetGroupId()))
+	{
+		ReferenceSpaceTransform = *CachedReferenceSpaceTransform;
+		ReferenceSpaceTransform.SetScale3D(FVec3(1.f));
+	}
+}
+
+TConstArrayView<Softs::FSolverVec3> FClothingSimulationCloth::GetAnimationPositions(const FClothingSimulationSolver* Solver) const
+{
+	check(Solver);
+	const int32 LODIndex = LODIndices.FindChecked(Solver);
+	check(GetParticleRangeId(Solver, LODIndex) != INDEX_NONE);
+	return TConstArrayView<Softs::FSolverVec3>(Solver->GetAnimationPositions(GetParticleRangeId(Solver, LODIndex)), GetNumParticles(LODIndex));
+}
+
+TConstArrayView<Softs::FSolverVec3> FClothingSimulationCloth::GetAnimationNormals(const FClothingSimulationSolver* Solver) const
+{
+	check(Solver);
+	const int32 LODIndex = LODIndices.FindChecked(Solver);
+	check(GetParticleRangeId(Solver, LODIndex) != INDEX_NONE);
+	return TConstArrayView<Softs::FSolverVec3>(Solver->GetAnimationNormals(GetParticleRangeId(Solver, LODIndex)), GetNumParticles(LODIndex));
+}
+
+TConstArrayView<Softs::FSolverVec3> FClothingSimulationCloth::GetParticlePositions(const FClothingSimulationSolver* Solver) const
+{
+	check(Solver);
+	const int32 LODIndex = LODIndices.FindChecked(Solver);
+	check(GetParticleRangeId(Solver, LODIndex) != INDEX_NONE);
+	return TConstArrayView<Softs::FSolverVec3>(Solver->GetParticleXs(GetParticleRangeId(Solver, LODIndex)), GetNumParticles(LODIndex));
+}
+
+TConstArrayView<Softs::FSolverVec3> FClothingSimulationCloth::GetParticleNormals(const FClothingSimulationSolver* Solver) const
+{
+	check(Solver);
+	const int32 LODIndex = LODIndices.FindChecked(Solver);
+	check(GetParticleRangeId(Solver, LODIndex) != INDEX_NONE);
+	return TConstArrayView<Softs::FSolverVec3>(Solver->GetNormals(GetParticleRangeId(Solver, LODIndex)), GetNumParticles(LODIndex));
+}
+
+TConstArrayView<Softs::FSolverVec3> FClothingSimulationCloth::GetParticleVelocities(const FClothingSimulationSolver* Solver) const
+{
+	check(Solver);
+	const int32 LODIndex = LODIndices.FindChecked(Solver);
+	check(GetParticleRangeId(Solver, LODIndex) != INDEX_NONE);
+	return TConstArrayView<Softs::FSolverVec3>(Solver->GetParticleVs(GetParticleRangeId(Solver, LODIndex)), GetNumParticles(LODIndex));
+}
+
+TConstArrayView<Softs::FSolverReal> FClothingSimulationCloth::GetParticleInvMasses(const FClothingSimulationSolver* Solver) const
+{
+	check(Solver);
+	const int32 LODIndex = LODIndices.FindChecked(Solver);
+	check(GetParticleRangeId(Solver, LODIndex) != INDEX_NONE);
+	return TConstArrayView<Softs::FSolverReal>(Solver->GetParticleInvMasses(GetParticleRangeId(Solver, LODIndex)), GetNumParticles(LODIndex));
+}
+
+void FClothingSimulationCloth::SetMaterialProperties(const TVec2<FRealSingle>& InEdgeStiffness, const TVec2<FRealSingle>& InBendingStiffness, const TVec2<FRealSingle>& InAreaStiffness)
+{
+	Config->GetProperties().SetWeightedFloatValue(TEXT("EdgeSpringStiffness"), FVector2f(InEdgeStiffness));
+	Config->GetProperties().SetWeightedFloatValue(TEXT("BendingSpringStiffness"), FVector2f(InBendingStiffness));
+	Config->GetProperties().SetWeightedFloatValue(TEXT("AreaSpringStiffness"), FVector2f(InAreaStiffness));
+}
+
+void FClothingSimulationCloth::SetLongRangeAttachmentProperties(const TVec2<FRealSingle>& InTetherStiffness, const TVec2<FRealSingle>& InTetherScale)
+{
+	Config->GetProperties().SetWeightedFloatValue(TEXT("TetherStiffness"), FVector2f(InTetherStiffness));
+	Config->GetProperties().SetWeightedFloatValue(TEXT("TetherScale"), FVector2f(InTetherScale));
+}
+
+void FClothingSimulationCloth::SetCollisionProperties(FRealSingle InCollisionThickness, FRealSingle InFrictionCoefficient, bool bInUseCCD, FRealSingle InSelfCollisionThickness)
+{
+	Config->GetProperties().SetValue(TEXT("CollisionThickness"), (float)InCollisionThickness);
+	Config->GetProperties().SetValue(TEXT("FrictionCoefficient"), (float)InFrictionCoefficient);
+	Config->GetProperties().SetValue(TEXT("UseCCD"), bInUseCCD);
+	Config->GetProperties().SetValue(TEXT("SelfCollisionThickness"), (float)InSelfCollisionThickness);
+}
+
+void FClothingSimulationCloth::SetBackstopProperties(bool bInEnableBackstop)
+{
+	Config->GetProperties().SetEnabled(TEXT("BackstopRadius"), bInEnableBackstop);  // BackstopRadius controls whether the backstop is enabled or not
+}
+
+void FClothingSimulationCloth::SetDampingProperties(FRealSingle InDampingCoefficient, FRealSingle InLocalDampingCoefficient)
+{
+	Config->GetProperties().SetValue(TEXT("DampingCoefficient"), (float)InDampingCoefficient);
+	Config->GetProperties().SetValue(TEXT("LocalDampingCoefficient"), (float)InLocalDampingCoefficient);
+}
+
+void FClothingSimulationCloth::SetAerodynamicsProperties(const TVec2<FRealSingle>& InDrag, const TVec2<FRealSingle>& InLift, FRealSingle InAirDensity, const FVec3& InWindVelocity)
+{
+	constexpr float WorldScale = 100.f;  // Interactor values are setup in engine scale, kg/cm^3 for air density and cm/s for wind velocity, but the properties are stored in kg/m^3 and m/s in the UI.
+	Config->GetProperties().SetWeightedFloatValue(TEXT("Drag"), FVector2f(InDrag));
+	Config->GetProperties().SetWeightedFloatValue(TEXT("Lift"), FVector2f(InLift));
+	Config->GetProperties().SetValue(TEXT("FluidDensity"), (float)InAirDensity * FMath::Cube(WorldScale));
+	Config->GetProperties().SetValue(TEXT("WindVelocity"), FVector3f(InWindVelocity) / WorldScale);
+}
+
+void FClothingSimulationCloth::SetPressureProperties(const TVec2<FRealSingle>& InPressure)
+{
+	Config->GetProperties().SetWeightedFloatValue(TEXT("Pressure"), FVector2f(InPressure));
+}
+
+void FClothingSimulationCloth::SetGravityProperties(FRealSingle InGravityScale, bool bInUseGravityOverride, const FVec3& InGravityOverride)
+{
+	Config->GetProperties().SetValue(TEXT("GravityScale"), (float)InGravityScale);
+	Config->GetProperties().SetValue(TEXT("UseGravityOverride"), bInUseGravityOverride);
+	Config->GetProperties().SetValue(TEXT("GravityOverride"), FVector3f(InGravityOverride));
+}
+
+void FClothingSimulationCloth::SetAnimDriveProperties(const TVec2<FRealSingle>& InAnimDriveStiffness, const TVec2<FRealSingle>& InAnimDriveDamping)
+{
+	Config->GetProperties().SetWeightedFloatValue(TEXT("AnimDriveStiffness"), FVector2f(InAnimDriveStiffness));
+	Config->GetProperties().SetWeightedFloatValue(TEXT("AnimDriveDamping"), FVector2f(InAnimDriveDamping));
+}
+
+void FClothingSimulationCloth::GetAnimDriveProperties(TVec2<FRealSingle>& OutAnimDriveStiffness, TVec2<FRealSingle>& OutAnimDriveDamping)
+{
+	OutAnimDriveStiffness = TVec2<FRealSingle>(Config->GetProperties().GetWeightedFloatValue(TEXT("AnimDriveStiffness")));
+	OutAnimDriveDamping = TVec2<FRealSingle>(Config->GetProperties().GetWeightedFloatValue(TEXT("AnimDriveDamping")));
+}
+
+void FClothingSimulationCloth::SetVelocityScaleProperties(const FVec3& InLinearVelocityScale, FRealSingle InAngularVelocityScale, FRealSingle InFictitiousAngularScale)
+{
+	Config->GetProperties().SetValue(TEXT("LinearVelocityScale"), FVector3f(InLinearVelocityScale));
+	Config->GetProperties().SetValue(TEXT("AngularVelocityScale"), (float)InAngularVelocityScale);
+	Config->GetProperties().SetValue(TEXT("FictitiousAngularScale"), (float)InFictitiousAngularScale);
+}
+
+}  // End namespace Chaos
diff --git a/Engine/Plugins/ChaosCloth/Source/ChaosCloth/Private/ChaosCloth/ChaosClothingSimulationMesh.cpp b/Engine/Plugins/ChaosCloth/Source/ChaosCloth/Private/ChaosCloth/ChaosClothingSimulationMesh.cpp
new file mode 100644
--- /dev/null
+++ b/Engine/Plugins/ChaosCloth/Source/ChaosCloth/Private/ChaosCloth/ChaosClothingSimulationMesh.cpp
@@ -0,0 +1,536 @@
+// Copyright Epic Games, Inc. All Rights Reserved.
+#include "ChaosCloth/ChaosClothingSimulationMesh.h"
+#include "ChaosCloth/ChaosClothingSimulationSolver.h"
+#include "ChaosCloth/ChaosWeightMapTarget.h"
+#include "ChaosCloth/ChaosClothPrivate.h"
+#include "ClothingSimulation.h"
+#if !defined(CHAOS_IS_CLOTHINGSIMULATIONMESH_ABSTRACT) || !CHAOS_IS_CLOTHINGSIMULATIONMESH_ABSTRACT
+#include "Components/SkeletalMeshComponent.h"
+#include "ClothingAsset.h"
+#else
+#include "Components/SkinnedMeshComponent.h"
+#endif
+#include "Containers/ArrayView.h"
+#include "Async/ParallelFor.h"
+#if INTEL_ISPC
+#include "ChaosClothingSimulationMesh.ispc.generated.h"
+#endif
+
+#if INTEL_ISPC && !UE_BUILD_SHIPPING
+bool bChaos_SkinPhysicsMesh_ISPC_Enabled = CHAOS_SKIN_PHYSICS_MESH_ISPC_ENABLED_DEFAULT;
+FAutoConsoleVariableRef CVarChaosSkinPhysicsMeshISPCEnabled(TEXT("p.Chaos.SkinPhysicsMesh.ISPC"), bChaos_SkinPhysicsMesh_ISPC_Enabled, TEXT("Whether to use ISPC optimizations on skinned physics meshes"));
+
+static_assert(sizeof(ispc::FVector3f) == sizeof(Chaos::Softs::FSolverVec3), "sizeof(ispc::FVector3f) != sizeof(Chaos::Softs::FSolverVec3)");
+static_assert(sizeof(ispc::FVector3f) == sizeof(FVector3f), "sizeof(ispc::FVector3f) != sizeof(FVector3f)");
+static_assert(sizeof(ispc::FMatrix44f) == sizeof(FMatrix44f), "sizeof(ispc::FMatrix44f) != sizeof(FMatrix44f)");
+static_assert(sizeof(ispc::FTransform3f) == sizeof(FTransform3f), "sizeof(ispc::FTransform3f) != sizeof(FTransform3f)");
+static_assert(sizeof(ispc::FClothVertBoneData) == sizeof(FClothVertBoneData), "sizeof(ispc::FClothVertBoneData) != sizeof(Chaos::FClothVertBoneData)");
+#endif
+
+DECLARE_CYCLE_STAT(TEXT("Chaos Cloth Skin Physics Mesh"), STAT_ChaosClothSkinPhysicsMesh, STATGROUP_ChaosCloth);
+DECLARE_CYCLE_STAT(TEXT("Chaos Cloth Wrap Deform Mesh"), STAT_ChaosClothWrapDeformMesh, STATGROUP_ChaosCloth);
+DECLARE_CYCLE_STAT(TEXT("Chaos Cloth Wrap Deform Cloth LOD"), STAT_ChaosClothWrapDeformClothLOD, STATGROUP_ChaosCloth);
+
+namespace Chaos
+{
+
+#if !defined(CHAOS_IS_CLOTHINGSIMULATIONMESH_ABSTRACT) || !CHAOS_IS_CLOTHINGSIMULATIONMESH_ABSTRACT
+
+PRAGMA_DISABLE_DEPRECATION_WARNINGS  // TODO: CHAOS_IS_CLOTHINGSIMULATIONMESH_ABSTRACT
+
+FClothingSimulationMesh::FClothingSimulationMesh(const FString& InDebugName)
+	: Asset(nullptr)
+	, SkeletalMeshComponent(nullptr)
+#if !UE_BUILD_SHIPPING
+	, DebugName(InDebugName)
+#endif
+{
+}
+
+FClothingSimulationMesh::FClothingSimulationMesh(const UClothingAssetCommon* InAsset, const USkeletalMeshComponent* InSkeletalMeshComponent)
+	: Asset(InAsset)
+	, SkeletalMeshComponent(InSkeletalMeshComponent)
+#if !UE_BUILD_SHIPPING
+	, DebugName(InSkeletalMeshComponent->GetOwner() ?
+		FString::Format(TEXT("{0}|{1}"), { InSkeletalMeshComponent->GetOwner()->GetName(), InSkeletalMeshComponent->GetName() } ) :
+		InSkeletalMeshComponent->GetName())
+#endif
+{
+}
+
+FClothingSimulationMesh::~FClothingSimulationMesh() = default;
+
+int32 FClothingSimulationMesh::GetNumLODs() const
+{
+	return Asset ? Asset->LodData.Num() : 0;
+}
+
+int32 FClothingSimulationMesh::GetLODIndex() const
+{
+	int32 LODIndex = INDEX_NONE;
+
+	if (Asset && SkeletalMeshComponent)
+	{
+		const int32 MeshLODIndex = SkeletalMeshComponent->GetPredictedLODLevel();
+		if (Asset->LodMap.IsValidIndex(MeshLODIndex))
+		{
+			const int32 ClothLODIndex = Asset->LodMap[MeshLODIndex];
+			if (Asset->LodData.IsValidIndex(ClothLODIndex))
+			{
+				LODIndex = ClothLODIndex;
+			}
+		}
+	}
+	return LODIndex;
+}
+
+bool FClothingSimulationMesh::IsValidLODIndex(int32 LODIndex) const
+{
+	return Asset && Asset->LodData.IsValidIndex(LODIndex);
+}
+
+int32 FClothingSimulationMesh::GetOwnerLODIndex(int32 LODIndex) const
+{
+	const int32 OwnerLODIndex = Asset ? Asset->LodMap.Find(LODIndex) : INDEX_NONE;
+	return OwnerLODIndex != INDEX_NONE ? OwnerLODIndex : 0;  // Safer to return the default LOD 0 than INDEX_NONE in this case
+}
+
+int32 FClothingSimulationMesh::GetNumPoints(int32 LODIndex) const
+{
+	return IsValidLODIndex(LODIndex) ? Asset->LodData[LODIndex].PhysicalMeshData.Vertices.Num() : 0;
+}
+
+int32 FClothingSimulationMesh::GetNumPatternPoints(int32 LODIndex) const
+{
+	return 0;
+}
+
+TConstArrayView<FVector3f> FClothingSimulationMesh::GetPositions(int32 LODIndex) const
+{
+	if (IsValidLODIndex(LODIndex))
+	{
+		const FClothLODDataCommon& ClothLODData = Asset->LodData[LODIndex];
+		const FClothPhysicalMeshData& ClothPhysicalMeshData = ClothLODData.PhysicalMeshData;
+		return TConstArrayView<FVector3f>(ClothPhysicalMeshData.Vertices);
+	}
+	return TConstArrayView<FVector3f>();
+}
+
+TConstArrayView<FVector2f> FClothingSimulationMesh::GetPatternPositions(int32 LODIndex) const
+{
+	return TConstArrayView<FVector2f>();
+}
+
+TConstArrayView<FVector3f> FClothingSimulationMesh::GetNormals(int32 LODIndex) const
+{
+	if (IsValidLODIndex(LODIndex))
+	{
+		const FClothLODDataCommon& ClothLODData = Asset->LodData[LODIndex];
+		const FClothPhysicalMeshData& ClothPhysicalMeshData = ClothLODData.PhysicalMeshData;
+		return TConstArrayView<FVector3f>(ClothPhysicalMeshData.Normals);
+	}
+	return TConstArrayView<FVector3f>();
+}
+
+TConstArrayView<uint32> FClothingSimulationMesh::GetIndices(int32 LODIndex) const
+{
+	return IsValidLODIndex(LODIndex) ?
+		TConstArrayView<uint32>(Asset->LodData[LODIndex].PhysicalMeshData.Indices) :
+		TConstArrayView<uint32>();
+}
+
+TConstArrayView<uint32> FClothingSimulationMesh::GetPatternIndices(int32 LODIndex) const
+{
+	return TConstArrayView<uint32>();
+}
+
+TConstArrayView<uint32> FClothingSimulationMesh::GetPatternToWeldedIndices(int32 LODIndex) const
+{
+	return TConstArrayView<uint32>();
+}
+
+TArray<FName> FClothingSimulationMesh::GetWeightMapNames(int32 LODIndex) const
+{
+	TArray<FName> WeightMapNames;
+	if (IsValidLODIndex(LODIndex))
+	{
+		// This must match the order of GetWeightMaps
+		const FClothLODDataCommon& ClothLODData = Asset->LodData[LODIndex];
+		const FClothPhysicalMeshData& ClothPhysicalMeshData = ClothLODData.PhysicalMeshData;
+		WeightMapNames.Reserve(ClothPhysicalMeshData.WeightMaps.Num());
+
+		const UEnum* const ChaosWeightMapTargetEnum = StaticEnum<EChaosWeightMapTarget>();
+		for (const TPair<uint32, FPointWeightMap>& TargetIDAndMap : ClothPhysicalMeshData.WeightMaps)
+		{
+			WeightMapNames.Add(FName(ChaosWeightMapTargetEnum->GetNameStringByValue(TargetIDAndMap.Get<0>())));
+		}
+	}
+
+	return WeightMapNames;
+}
+
+TMap<FString, int32> FClothingSimulationMesh::GetWeightMapIndices(int32 LODIndex) const
+{
+	TMap<FString, int32> WeightMapIndices;
+	const TArray<FName> WeightMapNames = GetWeightMapNames(LODIndex);
+	WeightMapIndices.Reserve(WeightMapNames.Num());
+	for (int32 WeightMapIndex = 0; WeightMapIndex < WeightMapNames.Num(); ++WeightMapIndex)
+	{
+		const FName& WeightMapName = WeightMapNames[WeightMapIndex];
+		WeightMapIndices.Emplace(WeightMapName.ToString(), WeightMapIndex);
+	}
+	return WeightMapIndices;
+}
+
+TArray<TConstArrayView<FRealSingle>> FClothingSimulationMesh::GetWeightMaps(int32 LODIndex) const
+{
+	TArray<TConstArrayView<FRealSingle>> WeightMaps;
+	if (IsValidLODIndex(LODIndex))
+	{
+		// This must match the order of GetWeightMapNames
+		const FClothLODDataCommon& ClothLODData = Asset->LodData[LODIndex];
+		const FClothPhysicalMeshData& ClothPhysicalMeshData = ClothLODData.PhysicalMeshData;
+
+		WeightMaps.Reserve(ClothPhysicalMeshData.WeightMaps.Num());
+		for (const TPair<uint32, FPointWeightMap>& TargetIDAndMap : ClothPhysicalMeshData.WeightMaps)
+		{
+			WeightMaps.Add(TargetIDAndMap.Get<1>().Values);
+		}
+	}
+	return WeightMaps;
+}
+
+TMap<FString, const TSet<int32>*> FClothingSimulationMesh::GetVertexSets(int32 LODIndex) const
+{
+	// Not supported
+	return TMap<FString, const TSet<int32>*>();
+}
+
+TMap<FString, const TSet<int32>*> FClothingSimulationMesh::GetFaceSets(int32 LODIndex) const
+{
+	// Not supported
+	return TMap<FString, const TSet<int32>*>();
+}
+
+TMap<FString, TConstArrayView<int32>> FClothingSimulationMesh::GetFaceIntMaps(int32 LODIndex) const
+{
+	// Not supported
+	return TMap<FString, TConstArrayView<int32>>();
+}
+
+TArray<TConstArrayView<TTuple<int32, int32, float>>> FClothingSimulationMesh::GetTethers(int32 LODIndex, bool bUseGeodesicTethers) const
+{
+	TArray<TConstArrayView<TTuple<int32, int32, float>>> Tethers;
+	if (IsValidLODIndex(LODIndex))
+	{
+		const FClothLODDataCommon& ClothLODData = Asset->LodData[LODIndex];
+		const FClothPhysicalMeshData& ClothPhysicalMeshData = ClothLODData.PhysicalMeshData;
+		const FClothTetherData& ClothTetherData = bUseGeodesicTethers ? ClothPhysicalMeshData.GeodesicTethers : ClothPhysicalMeshData.EuclideanTethers;
+		
+		const int32 NumTetherBatches = ClothTetherData.Tethers.Num();
+		Tethers.Reserve(NumTetherBatches);
+		for (int32 Index = 0; Index < NumTetherBatches; ++Index)
+		{
+			Tethers.Emplace(TConstArrayView<TTuple<int32, int32, float>>(ClothTetherData.Tethers[Index]));
+		}
+	}
+	return Tethers;
+}
+
+int32 FClothingSimulationMesh::GetReferenceBoneIndex() const
+{
+	return Asset ? Asset->ReferenceBoneIndex : INDEX_NONE;
+}
+
+FTransform FClothingSimulationMesh::GetReferenceBoneTransform() const
+{
+	if (const FClothingSimulationContextCommon* const Context = GetContext())
+	{
+		const int32 ReferenceBoneIndex = GetReferenceBoneIndex();
+		const TArray<FTransform>& BoneTransforms = Context->BoneTransforms;
+
+		return BoneTransforms.IsValidIndex(ReferenceBoneIndex) ?
+			BoneTransforms[ReferenceBoneIndex] * Context->ComponentToWorld :
+			Context->ComponentToWorld;
+	}
+	return FTransform::Identity;
+}
+
+const TArray<FTransform>& FClothingSimulationMesh::GetBoneTransforms() const
+{
+	return GetContext()->BoneTransforms;
+}
+
+const FTransform& FClothingSimulationMesh::GetComponentToWorldTransform() const
+{
+	return GetContext() ? GetContext()->ComponentToWorld : FTransform::Identity;
+}
+
+const TArray<FMatrix44f>& FClothingSimulationMesh::GetRefToLocalMatrices() const
+{
+	static TArray<FMatrix44f> EmptyArray;
+	return GetContext() ? GetContext()->RefToLocals : EmptyArray;
+}
+
+TConstArrayView<int32> FClothingSimulationMesh::GetBoneMap() const
+{
+	return Asset ? Asset->UsedBoneIndices : TConstArrayView<int32>();
+}
+
+TConstArrayView<FClothVertBoneData> FClothingSimulationMesh::GetBoneData(int32 LODIndex) const
+{
+	if (IsValidLODIndex(LODIndex))
+	{
+		const FClothLODDataCommon& ClothLODData = Asset->LodData[LODIndex];
+		const FClothPhysicalMeshData& ClothPhysicalMeshData = ClothLODData.PhysicalMeshData;
+		UE_CLOG(ClothPhysicalMeshData.MaxBoneWeights > 12, LogChaosCloth, Warning, TEXT("The cloth physics mesh skinning code can't cope with more than 12 bone influences."));
+		return ClothPhysicalMeshData.BoneData;
+	}
+	static TArray<FClothVertBoneData> EmptyArray;
+	return EmptyArray;
+}
+
+TConstArrayView<FMeshToMeshVertData> FClothingSimulationMesh::GetTransitionUpSkinData(int32 LODIndex) const
+{
+	return IsValidLODIndex(LODIndex) ? TConstArrayView<FMeshToMeshVertData>(Asset->LodData[LODIndex].TransitionUpSkinData) : TConstArrayView<FMeshToMeshVertData>();
+}
+
+TConstArrayView<FMeshToMeshVertData> FClothingSimulationMesh::GetTransitionDownSkinData(int32 LODIndex) const
+{
+	return IsValidLODIndex(LODIndex) ? TConstArrayView<FMeshToMeshVertData>(Asset->LodData[LODIndex].TransitionDownSkinData) : TConstArrayView<FMeshToMeshVertData>();
+}
+
+const FClothingSimulationContextCommon* FClothingSimulationMesh::GetContext() const
+{
+	return GetSkeletalMeshComponent() ?
+		static_cast<const FClothingSimulationContextCommon*>(GetSkeletalMeshComponent()->GetClothingSimulationContext()) :
+		nullptr;
+}
+
+PRAGMA_ENABLE_DEPRECATION_WARNINGS
+
+#else
+FClothingSimulationMesh::FClothingSimulationMesh() = default;
+
+FClothingSimulationMesh::~FClothingSimulationMesh() = default;
+#endif
+
+Softs::FSolverReal FClothingSimulationMesh::GetScale() const
+{
+	return GetComponentToWorldTransform().GetScale3D().GetMax();
+}
+
+bool FClothingSimulationMesh::WrapDeformLOD(
+	int32 PrevLODIndex,
+	int32 LODIndex,
+	const FSolverVec3* Normals,
+	const FSolverVec3* Positions,
+	FSolverVec3* OutPositions) const
+{
+	TRACE_CPUPROFILER_EVENT_SCOPE(FClothingSimulationMesh_WrapDeformLOD);
+	SCOPE_CYCLE_COUNTER(STAT_ChaosClothWrapDeformMesh);
+
+	const int32 NumLODsPassed = FMath::Abs(LODIndex - PrevLODIndex);
+	if (NumLODsPassed != 1 || !IsValidLODIndex(PrevLODIndex) || !IsValidLODIndex(LODIndex))
+	{
+		return false;
+	}
+
+	const int32 NumPoints = GetNumPoints(LODIndex);
+	const TConstArrayView<FMeshToMeshVertData> SkinData = (PrevLODIndex < LODIndex) ?
+		GetTransitionUpSkinData(LODIndex) :
+		GetTransitionDownSkinData(LODIndex);
+
+	for (int32 Index = 0; Index < NumPoints; ++Index)  // TODO: Profile for parallel for
+	{
+		const FMeshToMeshVertData& VertData = SkinData[Index];
+
+		const int32 VertIndex0 = (int32)VertData.SourceMeshVertIndices[0];  // Note: The source is uint16. Watch out for large mesh sections!
+		const int32 VertIndex1 = (int32)VertData.SourceMeshVertIndices[1];
+		const int32 VertIndex2 = (int32)VertData.SourceMeshVertIndices[2];
+
+		OutPositions[Index] = 
+			Positions[VertIndex0] * (FSolverReal)VertData.PositionBaryCoordsAndDist.X + Normals[VertIndex0] * (FSolverReal)VertData.PositionBaryCoordsAndDist.W +
+			Positions[VertIndex1] * (FSolverReal)VertData.PositionBaryCoordsAndDist.Y + Normals[VertIndex1] * (FSolverReal)VertData.PositionBaryCoordsAndDist.W +
+			Positions[VertIndex2] * (FSolverReal)VertData.PositionBaryCoordsAndDist.Z + Normals[VertIndex2] * (FSolverReal)VertData.PositionBaryCoordsAndDist.W;
+	}
+
+	return true;
+}
+
+bool FClothingSimulationMesh::WrapDeformLOD(
+	int32 PrevLODIndex,
+	int32 LODIndex,
+	const Softs::FSolverVec3* Normals,
+	const Softs::FPAndInvM* PositionAndInvMs,
+	const Softs::FSolverVec3* Velocities,
+	Softs::FPAndInvM* OutPositionAndInvMs0,
+	Softs::FSolverVec3* OutPositions1,
+	Softs::FSolverVec3* OutVelocities) const
+{
+	TRACE_CPUPROFILER_EVENT_SCOPE(FClothingSimulationMesh_WrapDeformLOD);
+	SCOPE_CYCLE_COUNTER(STAT_ChaosClothWrapDeformClothLOD);
+
+	const int32 NumLODsPassed = FMath::Abs(LODIndex - PrevLODIndex);
+	if (NumLODsPassed != 1 || !IsValidLODIndex(PrevLODIndex) || !IsValidLODIndex(LODIndex))
+	{
+		return false;
+	}
+
+	const int32 NumPoints = GetNumPoints(LODIndex);
+	const TConstArrayView<FMeshToMeshVertData> SkinData = (PrevLODIndex < LODIndex) ?
+		GetTransitionUpSkinData(LODIndex) :
+		GetTransitionDownSkinData(LODIndex);
+
+	for (int32 Index = 0; Index < NumPoints; ++Index)  // TODO: Profile for parallel for
+	{
+		const FMeshToMeshVertData& VertData = SkinData[Index];
+
+		const int32 VertIndex0 = (int32)VertData.SourceMeshVertIndices[0];  // Note: The source is uint16. Watch out for large mesh sections!
+		const int32 VertIndex1 = (int32)VertData.SourceMeshVertIndices[1];
+		const int32 VertIndex2 = (int32)VertData.SourceMeshVertIndices[2];
+
+		OutPositionAndInvMs0[Index].P = OutPositions1[Index] =
+			PositionAndInvMs[VertIndex0].P * (FSolverReal)VertData.PositionBaryCoordsAndDist.X + Normals[VertIndex0] * (FSolverReal)VertData.PositionBaryCoordsAndDist.W +
+			PositionAndInvMs[VertIndex1].P * (FSolverReal)VertData.PositionBaryCoordsAndDist.Y + Normals[VertIndex1] * (FSolverReal)VertData.PositionBaryCoordsAndDist.W +
+			PositionAndInvMs[VertIndex2].P * (FSolverReal)VertData.PositionBaryCoordsAndDist.Z + Normals[VertIndex2] * (FSolverReal)VertData.PositionBaryCoordsAndDist.W;
+
+		OutVelocities[Index] = 
+			Velocities[VertIndex0] * (FSolverReal)VertData.PositionBaryCoordsAndDist.X +
+			Velocities[VertIndex1] * (FSolverReal)VertData.PositionBaryCoordsAndDist.Y +
+			Velocities[VertIndex2] * (FSolverReal)VertData.PositionBaryCoordsAndDist.Z;
+	}
+
+	return true;
+}
+
+// Inline function used to force the unrolling of the skinning loop, LWC: note skinning is all done in float to match the asset data type
+FORCEINLINE static void AddInfluence(FVector3f& OutPosition, FVector3f& OutNormal, const FVector3f& RefParticle, const FVector3f& RefNormal, const FMatrix44f& BoneMatrix, const float Weight)
+{
+	OutPosition += BoneMatrix.TransformPosition(RefParticle) * Weight;
+	OutNormal += BoneMatrix.TransformVector(RefNormal) * Weight;
+}
+
+void FClothingSimulationMesh::SkinPhysicsMesh(int32 LODIndex, const FVec3& LocalSpaceLocation, FSolverVec3* OutPositions, FSolverVec3* OutNormals) const
+{
+	TRACE_CPUPROFILER_EVENT_SCOPE(FClothingSimulationMesh_SkinPhysicsMesh);
+	SCOPE_CYCLE_COUNTER(STAT_ChaosClothSkinPhysicsMesh);
+	SCOPE_CYCLE_COUNTER(STAT_ClothSkinPhysMesh);
+
+	FTransform ComponentToLocalSpaceReal = GetComponentToWorldTransform();
+	ComponentToLocalSpaceReal.AddToTranslation(-LocalSpaceLocation);
+	const FTransform3f ComponentToLocalSpace(ComponentToLocalSpaceReal);  // LWC: Now in local space, therefore it is safe to use single precision which is the asset data format
+
+	const int32* const RESTRICT BoneMap = GetBoneMap().GetData();
+	const FMatrix44f* const RESTRICT RefToLocalMatrices = GetRefToLocalMatrices().GetData();
+
+	check(IsValidLODIndex(LODIndex));
+	const uint32 NumPoints = GetNumPoints(LODIndex);
+	const TConstArrayView<FClothVertBoneData> BoneData = GetBoneData(LODIndex);
+	const TConstArrayView<FVector3f> Positions = GetPositions(LODIndex);
+	const TConstArrayView<FVector3f> Normals = GetNormals(LODIndex);
+
+#if INTEL_ISPC
+	if (bChaos_SkinPhysicsMesh_ISPC_Enabled)
+	{
+		ispc::SkinPhysicsMesh(
+			(ispc::FVector3f*)OutPositions,
+			(ispc::FVector3f*)OutNormals,
+			(ispc::FVector3f*)Positions.GetData(),
+			(ispc::FVector3f*)Normals.GetData(),
+			(ispc::FClothVertBoneData*)BoneData.GetData(),
+			BoneMap,
+			(ispc::FMatrix44f*)RefToLocalMatrices,
+			(ispc::FTransform3f&)ComponentToLocalSpace,
+			NumPoints);
+	}
+	else
+#endif
+	{
+		static const uint32 MinParallelVertices = 500;  // 500 seems to be the lowest threshold still giving gains even on profiled assets that are only using a small number of influences
+
+		ParallelFor(NumPoints, [&BoneData, &Positions, &Normals, &ComponentToLocalSpace, BoneMap, RefToLocalMatrices, &OutPositions, &OutNormals](uint32 VertIndex)
+		{
+			const uint16* const RESTRICT BoneIndices = BoneData[VertIndex].BoneIndices;
+			const float* const RESTRICT BoneWeights = BoneData[VertIndex].BoneWeights;
+	
+			// WARNING - HORRIBLE UNROLLED LOOP + JUMP TABLE BELOW
+			// done this way because this is a pretty tight and performance critical loop. essentially
+			// rather than checking each influence we can just jump into this switch and fall through
+			// everything to compose the final skinned data
+			const FVector3f& RefParticle = Positions[VertIndex];
+			const FVector3f& RefNormal = Normals[VertIndex];
+
+			FVector3f Position(ForceInitToZero);
+			FVector3f Normal(ForceInitToZero);
+			switch (BoneData[VertIndex].NumInfluences)
+			{
+			default:  // Intentional fall through
+			case 12: AddInfluence(Position, Normal, RefParticle, RefNormal, RefToLocalMatrices[BoneMap[BoneIndices[11]]], BoneWeights[11]);  // Intentional fall through
+			case 11: AddInfluence(Position, Normal, RefParticle, RefNormal, RefToLocalMatrices[BoneMap[BoneIndices[10]]], BoneWeights[10]);  // Intentional fall through
+			case 10: AddInfluence(Position, Normal, RefParticle, RefNormal, RefToLocalMatrices[BoneMap[BoneIndices[ 9]]], BoneWeights[ 9]);  // Intentional fall through
+			case  9: AddInfluence(Position, Normal, RefParticle, RefNormal, RefToLocalMatrices[BoneMap[BoneIndices[ 8]]], BoneWeights[ 8]);  // Intentional fall through
+			case  8: AddInfluence(Position, Normal, RefParticle, RefNormal, RefToLocalMatrices[BoneMap[BoneIndices[ 7]]], BoneWeights[ 7]);  // Intentional fall through
+			case  7: AddInfluence(Position, Normal, RefParticle, RefNormal, RefToLocalMatrices[BoneMap[BoneIndices[ 6]]], BoneWeights[ 6]);  // Intentional fall through
+			case  6: AddInfluence(Position, Normal, RefParticle, RefNormal, RefToLocalMatrices[BoneMap[BoneIndices[ 5]]], BoneWeights[ 5]);  // Intentional fall through
+			case  5: AddInfluence(Position, Normal, RefParticle, RefNormal, RefToLocalMatrices[BoneMap[BoneIndices[ 4]]], BoneWeights[ 4]);  // Intentional fall through
+			case  4: AddInfluence(Position, Normal, RefParticle, RefNormal, RefToLocalMatrices[BoneMap[BoneIndices[ 3]]], BoneWeights[ 3]);  // Intentional fall through
+			case  3: AddInfluence(Position, Normal, RefParticle, RefNormal, RefToLocalMatrices[BoneMap[BoneIndices[ 2]]], BoneWeights[ 2]);  // Intentional fall through
+			case  2: AddInfluence(Position, Normal, RefParticle, RefNormal, RefToLocalMatrices[BoneMap[BoneIndices[ 1]]], BoneWeights[ 1]);  // Intentional fall through
+			case  1: AddInfluence(Position, Normal, RefParticle, RefNormal, RefToLocalMatrices[BoneMap[BoneIndices[ 0]]], BoneWeights[ 0]);  // Intentional fall through
+			case  0: break;
+			}
+
+			OutPositions[VertIndex] = FSolverVec3(ComponentToLocalSpace.TransformPosition(Position));
+			OutNormals[VertIndex] = FSolverVec3(ComponentToLocalSpace.TransformVector(Normal).GetSafeNormal());
+
+		}, NumPoints > MinParallelVertices ? EParallelForFlags::None : EParallelForFlags::ForceSingleThread);
+	}
+}
+
+void FClothingSimulationMesh::Update(
+	FClothingSimulationSolver* Solver,
+	int32 PrevLODIndex,
+	int32 LODIndex,
+	int32 PrevParticleRangeId,
+	int32 ParticleRangeId)
+{
+	check(Solver);
+
+	// Exit if any inputs are missing or not ready, and if the LOD is invalid
+	if (!IsValidLODIndex(LODIndex))
+	{
+		return;
+	}
+
+	// Skin current LOD positions
+	const FVec3& LocalSpaceLocation = Solver->GetLocalSpaceLocation();
+	FSolverVec3* const OutPositions = Solver->GetAnimationPositions(ParticleRangeId);
+	FSolverVec3* const OutNormals = Solver->GetAnimationNormals(ParticleRangeId);
+	
+	SkinPhysicsMesh(LODIndex, LocalSpaceLocation, OutPositions, OutNormals);
+
+	// Update old positions after LOD Switching
+	if (LODIndex != PrevLODIndex)
+	{
+		// TODO: Using the more accurate skinning method here would require double buffering the context at the skeletal mesh level
+		const FSolverVec3* const SrcWrapNormals = Solver->GetAnimationNormals(PrevParticleRangeId);  // No need to keep an old normals array around, since the LOD has just changed
+		const FSolverVec3* const SrcWrapPositions = Solver->GetOldAnimationPositions(PrevParticleRangeId);
+		FSolverVec3* const OutOldPositions = Solver->GetOldAnimationPositions(ParticleRangeId);
+
+		const bool bValidWrap = WrapDeformLOD(PrevLODIndex, LODIndex, SrcWrapNormals, SrcWrapPositions, OutOldPositions);
+	
+		if (!bValidWrap)
+		{
+			// The previous LOD is invalid, reset old positions with the new LOD
+			const int32 NumPoints = GetNumPoints(LODIndex);
+
+			for (int32 Index = 0; Index < NumPoints; ++Index)
+			{
+				OutOldPositions[Index] = OutPositions[Index];
+			}
+		}
+	}
+}
+
+}  // End namespace Chaos
diff --git a/Engine/Plugins/ChaosCloth/Source/ChaosCloth/Public/ChaosCloth/ChaosClothingSimulationCloth.h b/Engine/Plugins/ChaosCloth/Source/ChaosCloth/Public/ChaosCloth/ChaosClothingSimulationCloth.h
new file mode 100644
--- /dev/null
+++ b/Engine/Plugins/ChaosCloth/Source/ChaosCloth/Public/ChaosCloth/ChaosClothingSimulationCloth.h
@@ -0,0 +1,246 @@
+// Copyright Epic Games, Inc. All Rights Reserved.
+#pragma once
+
+#include "Chaos/Transform.h"
+#include "Chaos/TriangleMesh.h"
+#include "Chaos/AABB.h"
+#include "Containers/ContainersFwd.h"
+#include "ChaosCloth/ChaosClothConstraints.h"
+
+struct FManagedArrayCollection;
+struct FClothingSimulationCacheData;
+
+namespace Chaos
+{
+	class FClothingSimulationSolver;
+	class FClothingSimulationMesh;
+	class FClothingSimulationCollider;
+	class FClothingSimulationConfig;
+
+	// Cloth simulation node
+	class FClothingSimulationCloth final
+	{
+	public:
+		enum EMassMode
+		{
+			UniformMass,
+			TotalMass,
+			Density
+		};
+
+		UE_DEPRECATED(5.3, "ETetherMode has been replaced with bUseGeodesicTethers.")
+PRAGMA_DISABLE_DEPRECATION_WARNINGS
+		typedef FClothConstraints::ETetherMode ETetherMode;
+CHAOSCLOTH_API PRAGMA_ENABLE_DEPRECATION_WARNINGS
+
+		FClothingSimulationCloth(
+			FClothingSimulationConfig* InConfig,
+PRAGMA_DISABLE_DEPRECATION_WARNINGS  // TODO: CHAOS_IS_CLOTHINGSIMULATIONMESH_ABSTRACT
+			FClothingSimulationMesh* InMesh,
+PRAGMA_ENABLE_DEPRECATION_WARNINGS
+			TArray<FClothingSimulationCollider*>&& InColliders,
+			uint32 InGroupId);
+
+		UE_DEPRECATED(5.2, "Use config based constructor instead.")
+		FClothingSimulationCloth(
+PRAGMA_DISABLE_DEPRECATION_WARNINGS  // TODO: CHAOS_IS_CLOTHINGSIMULATIONMESH_ABSTRACT
+			FClothingSimulationMesh* InMesh,
+PRAGMA_ENABLE_DEPRECATION_WARNINGS
+			TArray<FClothingSimulationCollider*>&& InColliders,
+			uint32 InGroupId,
+			EMassMode InMassMode,
+			FRealSingle InMassValue,
+			FRealSingle InMinPerParticleMass,
+			const TVec2<FRealSingle>& InEdgeStiffness,
+			const TVec2<FRealSingle>& InBendingStiffness,
+			FRealSingle InBucklingRatio,
+			const TVec2<FRealSingle>& InBucklingStiffness,
+			bool bInUseBendingElements,
+			const TVec2<FRealSingle>& InAreaStiffness,
+			FRealSingle InVolumeStiffness,
+			bool bInUseThinShellVolumeConstraints,
+			const TVec2<FRealSingle>& InTetherStiffness,
+			const TVec2<FRealSingle>& InTetherScale,
+PRAGMA_DISABLE_DEPRECATION_WARNINGS
+			ETetherMode InTetherMode,
+PRAGMA_ENABLE_DEPRECATION_WARNINGS
+			FRealSingle InMaxDistancesMultiplier,
+			const TVec2<FRealSingle>& InAnimDriveStiffness,
+			const TVec2<FRealSingle>& InAnimDriveDamping,
+			FRealSingle InShapeTargetStiffness,
+			bool bInUseXPBDEdgeConstraints,
+			bool bInUseXPBDBendingConstraints,
+			bool bInUseXPBDAreaConstraints,
+			FRealSingle InGravityScale,
+			bool bIsGravityOverridden,
+			const TVec3<FRealSingle>& InGravityOverride,
+			const TVec3<FRealSingle>& InLinearVelocityScale,
+			FRealSingle InAngularVelocityScale,
+			FRealSingle InFictitiousAngularScale,
+			const TVec2<FRealSingle>& InDrag,
+			const TVec2<FRealSingle>& InLift,
+			bool bInUsePointBasedWindModel,
+			const TVec2<FRealSingle>& InPressure,
+			FRealSingle InDampingCoefficient,
+			FRealSingle InLocalDampingCoefficient,
+			FRealSingle InCollisionThickness,
+			FRealSingle InFrictionCoefficient,
+			bool bInUseCCD,
+			bool bInUseSelfCollisions,
+			FRealSingle InSelfCollisionThickness,
+			FRealSingle InSelfCollisionFrictionCoefficient,
+			bool bInUseSelfIntersections,
+			bool bInUseLegacyBackstop,
+			bool bInUseLODIndexOverride,
+			int32 InLODIndexOverride,
+			const TVec2<FRealSingle>& EdgeDampingRatio = TVec2<FRealSingle>(0.f),
+			const TVec2<FRealSingle>& BendDampingRatio = TVec2<FRealSingle>(0.f));
+		CHAOSCLOTH_API ~FClothingSimulationCloth();
+
+		FClothingSimulationCloth(const FClothingSimulationCloth&) = delete;
+		FClothingSimulationCloth(FClothingSimulationCloth&&) = delete;
+		FClothingSimulationCloth& operator=(const FClothingSimulationCloth&) = delete;
+		FClothingSimulationCloth& operator=(FClothingSimulationCloth&&) = delete;
+
+		uint32 GetGroupId() const { return GroupId; }
+		uint32 GetLODIndex(const FClothingSimulationSolver* Solver) const { return LODIndices.FindChecked(Solver); }
+
+		int32 GetNumActiveKinematicParticles() const { return NumActiveKinematicParticles; }
+		int32 GetNumActiveDynamicParticles() const { return NumActiveDynamicParticles; }
+
+		// ---- Animatable property setters ----
+		void SetMaxDistancesMultiplier(FRealSingle InMaxDistancesMultiplier) { MaxDistancesMultiplier = InMaxDistancesMultiplier; }
+		UE_DEPRECATED(5.3, "Set properties directly through FClothingSimulationConfig")
+		CHAOSCLOTH_API void SetMaterialProperties(const TVec2<FRealSingle>& InEdgeStiffness, const TVec2<FRealSingle>& InBendingStiffness, const TVec2<FRealSingle>& InAreaStiffness);
+		UE_DEPRECATED(5.3, "Set properties directly through FClothingSimulationConfig")
+		CHAOSCLOTH_API void SetLongRangeAttachmentProperties(const TVec2<FRealSingle>& InTetherStiffness, const TVec2<FRealSingle>& InTetherScale);
+		UE_DEPRECATED(5.3, "Set properties directly through FClothingSimulationConfig")
+		CHAOSCLOTH_API void SetCollisionProperties(FRealSingle InCollisionThickness, FRealSingle InFrictionCoefficient, bool bInUseCCD, FRealSingle InSelfCollisionThickness);
+		UE_DEPRECATED(5.3, "Set properties directly through FClothingSimulationConfig")
+		CHAOSCLOTH_API void SetBackstopProperties(bool bInEnableBackstop);
+		UE_DEPRECATED(5.3, "Set properties directly through FClothingSimulationConfig")
+		CHAOSCLOTH_API void SetDampingProperties(FRealSingle InDampingCoefficient, FRealSingle InLocalDampingCoefficient = 0.f);
+		UE_DEPRECATED(5.3, "Set properties directly through FClothingSimulationConfig")
+		CHAOSCLOTH_API void SetAerodynamicsProperties(const TVec2<FRealSingle>& InDrag, const TVec2<FRealSingle>& InLift, FRealSingle InAirDensity, const FVec3& InWindVelocity);  // AirDensity is here in kg/cm^3 for legacy reason (kg/m^3 in UI)
+		UE_DEPRECATED(5.3, "Set properties directly through FClothingSimulationConfig")
+		CHAOSCLOTH_API void SetPressureProperties(const TVec2<FRealSingle>& InPressure);
+		UE_DEPRECATED(5.3, "Set properties directly through FClothingSimulationConfig")
+		CHAOSCLOTH_API void SetGravityProperties(FRealSingle InGravityScale, bool bInUseGravityOverride, const FVec3& InGravityOverride);
+		UE_DEPRECATED(5.3, "Set properties directly through FClothingSimulationConfig")
+		CHAOSCLOTH_API void SetAnimDriveProperties(const TVec2<FRealSingle>& InAnimDriveStiffness, const TVec2<FRealSingle>& InAnimDriveDamping);
+		UE_DEPRECATED(5.3, "Get properties directly through FClothingSimulationConfig")
+		CHAOSCLOTH_API void GetAnimDriveProperties(TVec2<FRealSingle>& OutAnimDriveStiffness, TVec2<FRealSingle>& OutAnimDriveDamping);
+		UE_DEPRECATED(5.3, "Set properties directly through FClothingSimulationConfig")
+		CHAOSCLOTH_API void SetVelocityScaleProperties(const FVec3& InLinearVelocityScale, FRealSingle InAngularVelocityScale, FRealSingle InFictitiousAngularScale);
+
+		void Reset() { bNeedsReset = true; }
+		void Teleport() { bNeedsTeleport = true; }
+		// ---- End of the animatable property setters ----
+
+		// ---- Node property getters/setters
+		PRAGMA_DISABLE_DEPRECATION_WARNINGS  // TODO: CHAOS_IS_CLOTHINGSIMULATIONMESH_ABSTRACT
+		FClothingSimulationMesh* GetMesh() const { return Mesh; }
+		CHAOSCLOTH_API void SetMesh(FClothingSimulationMesh* InMesh);
+		PRAGMA_ENABLE_DEPRECATION_WARNINGS
+
+		FClothingSimulationConfig* GetConfig() const { return Config; }
+		CHAOSCLOTH_API void SetConfig(FClothingSimulationConfig* InConfig);
+
+		const TArray<FClothingSimulationCollider*>& GetColliders() const { return Colliders; }
+		CHAOSCLOTH_API void SetColliders(TArray<FClothingSimulationCollider*>&& InColliders);
+		CHAOSCLOTH_API void AddCollider(FClothingSimulationCollider* InCollider);
+		CHAOSCLOTH_API void RemoveCollider(FClothingSimulationCollider* InCollider);
+		CHAOSCLOTH_API void RemoveColliders();
+		// ---- End of the Node property getters/setters
+
+		// ---- Debugging/visualization functions
+		// Return the solver's input positions for this cloth source current LOD, not thread safe, call must be done right after the solver update.
+		CHAOSCLOTH_API TConstArrayView<Softs::FSolverVec3> GetAnimationPositions(const FClothingSimulationSolver* Solver) const;
+		// Return the solver's input normals for this cloth source current LOD, not thread safe, call must be done right after the solver update.
+		CHAOSCLOTH_API TConstArrayView<Softs::FSolverVec3> GetAnimationNormals(const FClothingSimulationSolver* Solver) const;
+		// Return the solver's positions for this cloth current LOD, not thread safe, call must be done right after the solver update.
+		CHAOSCLOTH_API TConstArrayView<Softs::FSolverVec3> GetParticlePositions(const FClothingSimulationSolver* Solver) const;
+		// Return the solver's velocities for this cloth current LOD, not thread safe, call must be done right after the solver update.
+		CHAOSCLOTH_API TConstArrayView<Softs::FSolverVec3> GetParticleVelocities(const FClothingSimulationSolver* Solver) const;
+		// Return the solver's normals for this cloth current LOD, not thread safe, call must be done right after the solver update.
+		CHAOSCLOTH_API TConstArrayView<Softs::FSolverVec3> GetParticleNormals(const FClothingSimulationSolver* Solver) const;
+		// Return the solver's inverse masses for this cloth current LOD, not thread safe, call must be done right after the solver update.
+		CHAOSCLOTH_API TConstArrayView<Softs::FSolverReal> GetParticleInvMasses(const FClothingSimulationSolver* Solver) const;
+		// Return the current gravity as applied by the solver using the various overrides, not thread safe, call must be done right after the solver update.
+		// Does not have GravityScale applied when using Force-based solver (Get the per-particle value directly from cloth constraints' external forces).
+		CHAOSCLOTH_API TVec3<FRealSingle> GetGravity(const FClothingSimulationSolver* Solver) const;
+		// Return the current bounding box based on a given solver, not thread safe, call must be done right after the solver update.
+		CHAOSCLOTH_API FAABB3 CalculateBoundingBox(const FClothingSimulationSolver* Solver) const;
+		// Return the current LOD ParticleRangeId, or INDEX_NONE if no LOD is currently selected.
+		CHAOSCLOTH_API int32 GetParticleRangeId(const FClothingSimulationSolver* Solver) const;
+		UE_DEPRECATED(5.4, "Offset has been renamed ParticleRangeId to reflect that it is no longer an offset.")
+		int32 GetOffset(const FClothingSimulationSolver* Solver) const { return GetParticleRangeId(Solver); }
+		// Return the current LOD num particles, or 0 if no LOD is currently selected.
+		CHAOSCLOTH_API int32 GetNumParticles(const FClothingSimulationSolver* Solver) const;
+		// Return the current LOD mesh.
+		CHAOSCLOTH_API const FTriangleMesh& GetTriangleMesh(const FClothingSimulationSolver* Solver) const;
+		// Return the weight map of the specified name if available on the current LOD, or an empty array view otherwise.
+		CHAOSCLOTH_API TConstArrayView<FRealSingle> GetWeightMapByName(const FClothingSimulationSolver* Solver, const FString& Name) const;
+		// Return the weight map of the specified property name if it exists and is available on the current LOD, or an empty array view otherwise.
+		CHAOSCLOTH_API TConstArrayView<FRealSingle> GetWeightMapByProperty(const FClothingSimulationSolver* Solver, const FString& Property) const;
+		// Return the face int map of the specified name if available on the current LOD, or an empty array view otherwise.
+		CHAOSCLOTH_API TConstArrayView<int32> GetFaceIntMapByName(const FClothingSimulationSolver* Solver, const FString& Name) const;
+		// Return the face int map of the specified property name if it exists and is available on the current LOD, or an empty array view otherwise.
+		CHAOSCLOTH_API TConstArrayView<int32> GetFaceIntMapByProperty(const FClothingSimulationSolver* Solver, const FString& Property) const;
+		UE_DEPRECATED(5.3, "Returns an empty array from 5.3. Update your code with GetWeightMapByName and GetWeightMapByProperty to return the current LOD weight map instead.")
+		CHAOSCLOTH_API const TArray<TConstArrayView<FRealSingle>>& GetWeightMaps(const FClothingSimulationSolver* Solver) const;
+		// Return the current LOD tethers.
+		CHAOSCLOTH_API const TArray<TConstArrayView<TTuple<int32, int32, float>>>& GetTethers(const FClothingSimulationSolver* Solver) const;
+		// Return the reference bone index for this cloth.
+		CHAOSCLOTH_API int32 GetReferenceBoneIndex() const;
+		// Return the local reference space transform for this cloth.
+		const FRigidTransform3& GetReferenceSpaceTransform() const { return ReferenceSpaceTransform;  }
+		// ---- End of the debugging/visualization functions
+
+		// ---- Solver interface ----
+		CHAOSCLOTH_API void Add(FClothingSimulationSolver* Solver);
+		CHAOSCLOTH_API void Remove(FClothingSimulationSolver* Solver);
+
+		CHAOSCLOTH_API void PreUpdate(FClothingSimulationSolver* Solver);
+		CHAOSCLOTH_API void Update(FClothingSimulationSolver* Solver);
+		CHAOSCLOTH_API void PostUpdate(FClothingSimulationSolver* Solver);
+
+		CHAOSCLOTH_API void UpdateFromCache(const FClothingSimulationCacheData& CacheData);
+		// ---- End of the Solver interface ----
+
+	private:
+		CHAOSCLOTH_API int32 GetNumParticles(int32 InLODIndex) const;
+		CHAOSCLOTH_API int32 GetParticleRangeId(const FClothingSimulationSolver* Solver, int32 InLODIndex) const;
+
+	private:
+		struct FLODData;
+
+		// Cloth parameters
+PRAGMA_DISABLE_DEPRECATION_WARNINGS  // TODO: CHAOS_IS_CLOTHINGSIMULATIONMESH_ABSTRACT
+		FClothingSimulationMesh* Mesh = nullptr;
+PRAGMA_ENABLE_DEPRECATION_WARNINGS
+		FClothingSimulationConfig* Config = nullptr;
+		TArray<FClothingSimulationCollider*> Colliders;
+		uint32 GroupId = 0;
+
+		TSharedPtr<FManagedArrayCollection> PropertyCollection;  // Used for backward compatibility only, otherwise the properties are owned by the Config
+
+		FRealSingle MaxDistancesMultiplier = 1.f;  // Legacy multiplier
+
+		bool bUseLODIndexOverride = false;
+		int32 LODIndexOverride = 0;
+		bool bNeedsReset = false;
+		bool bNeedsTeleport = false;
+
+		// Reference space transform
+		FRigidTransform3 ReferenceSpaceTransform;  // TODO: Add override in the style of LODIndexOverride
+
+		// LOD data
+		TArray<TUniquePtr<FLODData>> LODData;
+		TMap<FClothingSimulationSolver*, int32> LODIndices;
+
+		// Stats
+		int32 NumActiveKinematicParticles = 0;
+		int32 NumActiveDynamicParticles = 0;
+	};
+} // namespace Chaos
