Chaos Cloth:

diff --git a/Engine/Plugins/ChaosCloth/Source/ChaosClothEditor/Private/ChaosClothEditor/ChaosSimulationEditorExtender.cpp b/Engine/Plugins/ChaosCloth/Source/ChaosClothEditor/Private/ChaosClothEditor/ChaosSimulationEditorExtender.cpp
new file mode 100644
--- /dev/null
+++ b/Engine/Plugins/ChaosCloth/Source/ChaosClothEditor/Private/ChaosClothEditor/ChaosSimulationEditorExtender.cpp
@@ -0,0 +1,214 @@
+// Copyright Epic Games, Inc. All Rights Reserved.
+
+#include "ChaosClothEditor/ChaosSimulationEditorExtender.h"
+
+#include "ChaosCloth/ChaosClothingSimulationFactory.h"
+#include "ChaosCloth/ChaosClothingSimulation.h"
+#include "Framework/MultiBox/MultiBoxBuilder.h"
+#include "Animation/DebugSkelMeshComponent.h"
+#include "Rendering/SkeletalMeshRenderData.h"
+#include "IPersonaPreviewScene.h"
+
+#define LOCTEXT_NAMESPACE "ChaosSimulationEditorExtender"
+
+namespace Chaos
+{
+
+struct FVisualizationOption
+{
+	// Actual option entries
+	static const FVisualizationOption OptionData[];
+	static const uint32 Count;
+
+	// Chaos debug draw function
+	typedef void (Chaos::FClothingSimulation::*FDebugDrawFunction)(FPrimitiveDrawInterface*) const;
+	typedef void (Chaos::FClothingSimulation::*FDebugDrawTextsFunction)(FCanvas*, const FSceneView*) const;
+	FDebugDrawFunction DebugDrawFunction;
+	FDebugDrawTextsFunction DebugDrawTextsFunction;
+
+	FText DisplayName;         // Text for menu entries.
+	FText ToolTip;             // Text for menu tooltips.
+	bool bDisablesSimulation;  // Whether or not this option requires the simulation to be disabled.
+	bool bHidesClothSections;  // Hides the cloth section to avoid zfighting with the debug geometry.
+
+	// Console override
+	IConsoleVariable* const ConsoleVariable;
+
+	FVisualizationOption(FDebugDrawFunction InDebugDrawFunction, const TCHAR* ConsoleName, const FText& InDisplayName, const FText& InToolTip, bool bInDisablesSimulation = false, bool bInHidesClothSections = false)
+		: DebugDrawFunction(InDebugDrawFunction)
+		, DisplayName(InDisplayName)
+		, ToolTip(InToolTip)
+		, bDisablesSimulation(bInDisablesSimulation)
+		, bHidesClothSections(bInHidesClothSections)
+		, ConsoleVariable(IConsoleManager::Get().RegisterConsoleVariable(ConsoleName, false, *InToolTip.ToString(), 0))
+	{}
+
+	FVisualizationOption(FDebugDrawTextsFunction InDebugDrawTextsFunction, const TCHAR* ConsoleName, const FText& InDisplayName, const FText& InToolTip, bool bInDisablesSimulation = false, bool bInHidesClothSections = false)
+		: DebugDrawTextsFunction(InDebugDrawTextsFunction)
+		, DisplayName(InDisplayName)
+		, ToolTip(InToolTip)
+		, bDisablesSimulation(bInDisablesSimulation)
+		, bHidesClothSections(bInHidesClothSections)
+		, ConsoleVariable(IConsoleManager::Get().RegisterConsoleVariable(ConsoleName, false, *InToolTip.ToString(), 0))
+	{}
+
+	~FVisualizationOption()
+	{
+		IConsoleManager::Get().UnregisterConsoleObject(ConsoleVariable);
+	}
+
+	bool IsConsoleVariableEnabled() const { return ConsoleVariable->AsVariableBool()->GetValueOnGameThread(); }
+};
+
+const FVisualizationOption FVisualizationOption::OptionData[] = 
+{
+	FVisualizationOption(&Chaos::FClothingSimulation::DebugDrawPhysMeshShaded      , TEXT("p.ChaosClothEditor.DebugDrawPhysMeshShaded"      ), LOCTEXT("ChaosVisName_PhysMesh"            , "Physical Mesh (Flat Shaded)"), LOCTEXT("ChaosVisName_PhysMeshShaded_ToolTip"      , "Draws the current physical result as a doubled sided flat shaded mesh"), /*bDisablesSimulation =*/false, /*bHidesClothSections=*/true),
+	FVisualizationOption(&Chaos::FClothingSimulation::DebugDrawPhysMeshWired       , TEXT("p.ChaosClothEditor.DebugDrawPhysMeshWired"       ), LOCTEXT("ChaosVisName_PhysMeshWire"        , "Physical Mesh (Wireframe)"  ), LOCTEXT("ChaosVisName_PhysMeshWired_ToolTip"       , "Draws the current physical mesh result in wireframe")),
+	FVisualizationOption(&Chaos::FClothingSimulation::DebugDrawAnimMeshWired       , TEXT("p.ChaosClothEditor.DebugDrawAnimMeshWired"       ), LOCTEXT("ChaosVisName_AnimMeshWire"        , "Animated Mesh (Wireframe)"  ), LOCTEXT("ChaosVisName_AnimMeshWired_ToolTip"       , "Draws the current animated mesh input in wireframe")),
+	FVisualizationOption(&Chaos::FClothingSimulation::DebugDrawParticleIndices     , TEXT("p.ChaosClothEditor.DebugDrawParticleIndices"     ), LOCTEXT("ChaosVisName_ParticleIndices"     , "Particle Indices"           ), LOCTEXT("ChaosVisName_ParticleIndices_ToolTip"     , "Draws the particle indices as instantiated by the solver")),
+	FVisualizationOption(&Chaos::FClothingSimulation::DebugDrawElementIndices      , TEXT("p.ChaosClothEditor.DebugDrawElementIndices"      ), LOCTEXT("ChaosVisName_ElementIndices"      , "Element Indices"            ), LOCTEXT("ChaosVisName_ElementIndices_ToolTip"      , "Draws the element's (triangle or other) indices as instantiated by the solver")),
+	FVisualizationOption(&Chaos::FClothingSimulation::DebugDrawPointNormals        , TEXT("p.ChaosClothEditor.DebugDrawPointNormals"        ), LOCTEXT("ChaosVisName_PointNormals"        , "Physical Mesh Normals"      ), LOCTEXT("ChaosVisName_PointNormals_ToolTip"        , "Draws the current point normals for the simulation mesh")),
+	FVisualizationOption(&Chaos::FClothingSimulation::DebugDrawPointVelocities     , TEXT("p.ChaosClothEditor.DebugDrawPointVelocities"     ), LOCTEXT("ChaosVisName_PointVelocities"     , "Point Velocities"           ), LOCTEXT("ChaosVisName_PointVelocities_ToolTip"     , "Draws the current point velocities for the simulation mesh")),
+	FVisualizationOption(&Chaos::FClothingSimulation::DebugDrawAnimNormals         , TEXT("p.ChaosClothEditor.DebugDrawAnimNormals"         ), LOCTEXT("ChaosVisName_AnimNormals"         , "Animated Mesh Normals"      ), LOCTEXT("ChaosVisName_AnimNormals_ToolTip"         , "Draws the current point normals for the animated mesh")),
+	FVisualizationOption(&Chaos::FClothingSimulation::DebugDrawCollision           , TEXT("p.ChaosClothEditor.DebugDrawCollision"           ), LOCTEXT("ChaosVisName_Collision"           , "Collisions"                 ), LOCTEXT("ChaosVisName_Collision_ToolTip"           , "Draws the collision bodies the simulation is currently using")),
+	FVisualizationOption(&Chaos::FClothingSimulation::DebugDrawBackstops           , TEXT("p.ChaosClothEditor.DebugDrawBackstops"           ), LOCTEXT("ChaosVisName_Backstop"            , "Backstops"                  ), LOCTEXT("ChaosVisName_Backstop_ToolTip"            , "Draws the backstop radius and position for each simulation particle")),
+	FVisualizationOption(&Chaos::FClothingSimulation::DebugDrawBackstopDistances   , TEXT("p.ChaosClothEditor.DebugDrawBackstopDistances"   ), LOCTEXT("ChaosVisName_BackstopDistance"    , "Backstop Distances"         ), LOCTEXT("ChaosVisName_BackstopDistance_ToolTip"    , "Draws the backstop distance offset for each simulation particle"), /*bDisablesSimulation =*/true),
+	FVisualizationOption(&Chaos::FClothingSimulation::DebugDrawMaxDistances        , TEXT("p.ChaosClothEditor.DebugDrawMaxDistances"        ), LOCTEXT("ChaosVisName_MaxDistance"         , "Max Distances"              ), LOCTEXT("ChaosVisName_MaxDistance_ToolTip"         , "Draws the current max distances for the sim particles as a line along its normal"), /*bDisablesSimulation =*/true),
+	FVisualizationOption(&Chaos::FClothingSimulation::DebugDrawMaxDistanceValues   , TEXT("p.ChaosClothEditor.DebugDrawMaxDistanceValues"   ), LOCTEXT("ChaosVisName_MaxDistanceValue"    , "Max Distances As Numbers"   ), LOCTEXT("ChaosVisName_MaxDistanceValue_ToolTip"    , "Draws the current max distances as numbers")),
+	FVisualizationOption(&Chaos::FClothingSimulation::DebugDrawAnimDrive           , TEXT("p.ChaosClothEditor.DebugDrawAnimDrive"           ), LOCTEXT("ChaosVisName_AnimDrive"           , "Anim Drive"                 ), LOCTEXT("ChaosVisName_AnimDrive_Tooltip"           , "Draws the current skinned reference mesh for the simulation which anim drive will attempt to reach if enabled")),
+	FVisualizationOption(&Chaos::FClothingSimulation::DebugDrawEdgeConstraint      , TEXT("p.ChaosClothEditor.DebugDrawEdgeConstraint"      ), LOCTEXT("ChaosVisName_EdgeConstraint"      , "Edge Constraint"            ), LOCTEXT("ChaosVisName_EdgeConstraint_Tooltip"      , "Draws the edge spring constraints")),
+	FVisualizationOption(&Chaos::FClothingSimulation::DebugDrawBendingConstraint   , TEXT("p.ChaosClothEditor.DebugDrawBendingConstraint"   ), LOCTEXT("ChaosVisName_BendingConstraint"   , "Bending Constraint"         ), LOCTEXT("ChaosVisName_BendingConstraint_Tooltip"   , "Draws the bending spring constraints")),
+	FVisualizationOption(&Chaos::FClothingSimulation::DebugDrawLongRangeConstraint , TEXT("p.ChaosClothEditor.DebugDrawLongRangeConstraint" ), LOCTEXT("ChaosVisName_LongRangeConstraint" , "Long Range Constraint"      ), LOCTEXT("ChaosVisName_LongRangeConstraint_Tooltip" , "Draws the long range attachment constraint distances")),
+	FVisualizationOption(&Chaos::FClothingSimulation::DebugDrawWindAndPressureForces, TEXT("p.ChaosClothEditor.DebugDrawWindAndPressureForces"), LOCTEXT("ChaosVisName_WindAndPressureForces", "Wind Aerodynamic And Pressure Forces"), LOCTEXT("ChaosVisName_WindAndPressure_Tooltip", "Draws the Wind drag and lift and pressure forces")),
+	FVisualizationOption(&Chaos::FClothingSimulation::DebugDrawLocalSpace          , TEXT("p.ChaosClothEditor.DebugDrawLocalSpace"          ), LOCTEXT("ChaosVisName_LocalSpace"          , "Local Space Reference Bone" ), LOCTEXT("ChaosVisName_LocalSpace_Tooltip"          , "Draws the local space reference bone")),
+	FVisualizationOption(&Chaos::FClothingSimulation::DebugDrawSelfCollision       , TEXT("p.ChaosClothEditor.DebugDrawSelfCollision"       ), LOCTEXT("ChaosVisName_SelfCollision"       , "Self Collision"             ), LOCTEXT("ChaosVisName_SelfCollision_Tooltip"       , "Draws the self collision thickness/debugging information")),
+	FVisualizationOption(&Chaos::FClothingSimulation::DebugDrawSelfIntersection    , TEXT("p.ChaosClothEditor.DebugDrawSelfIntersection"    ), LOCTEXT("ChaosVisName_SelfIntersection"    , "Self Intersection"          ), LOCTEXT("ChaosVisName_SelfIntersection_Tooltip"    , "Draws the self intersection contour/region information")),
+};
+const uint32 FVisualizationOption::Count = sizeof(OptionData) / sizeof(FVisualizationOption);
+
+FSimulationEditorExtender::FSimulationEditorExtender()
+	: Flags(false, FVisualizationOption::Count)
+{
+}
+
+UClass* FSimulationEditorExtender::GetSupportedSimulationFactoryClass()
+{
+	return UChaosClothingSimulationFactory::StaticClass();
+}
+
+void FSimulationEditorExtender::ExtendViewportShowMenu(FMenuBuilder& MenuBuilder, TSharedRef<IPersonaPreviewScene> PreviewScene)
+{
+	MenuBuilder.BeginSection(TEXT("ChaosSimulation_Visualizations"), LOCTEXT("VisualizationSection", "Visualizations"));
+	{
+		for (uint32 OptionIndex = 0; OptionIndex < FVisualizationOption::Count; ++OptionIndex)
+		{
+			// Handler for visualization entry being clicked
+			const FExecuteAction ExecuteAction = FExecuteAction::CreateLambda([this, OptionIndex, PreviewScene]()
+			{
+				Flags[OptionIndex] = !Flags[OptionIndex];
+
+				// If we need to toggle the disabled or visibility states, handle it
+				if (UDebugSkelMeshComponent* const MeshComponent = PreviewScene->GetPreviewMeshComponent())
+				{
+					// Disable simulation
+					const bool bShouldDisableSimulation = ShouldDisableSimulation();
+					if (bShouldDisableSimulation && MeshComponent->bDisableClothSimulation != bShouldDisableSimulation)
+					{
+						MeshComponent->bDisableClothSimulation = !MeshComponent->bDisableClothSimulation;
+					}
+					// Hide cloth section
+					if (FVisualizationOption::OptionData[OptionIndex].bHidesClothSections)
+					{
+						const bool bIsClothSectionsVisible = !Flags[OptionIndex];
+						ShowClothSections(MeshComponent, bIsClothSectionsVisible);
+					}
+				}
+			});
+
+			// Checkstate function for visualization entries
+			const FIsActionChecked IsActionChecked = FIsActionChecked::CreateLambda([this, OptionIndex]()
+			{
+				return Flags[OptionIndex];
+			});
+
+			const FUIAction Action(ExecuteAction, FCanExecuteAction(), IsActionChecked);
+
+			MenuBuilder.AddMenuEntry(FVisualizationOption::OptionData[OptionIndex].DisplayName, FVisualizationOption::OptionData[OptionIndex].ToolTip, FSlateIcon(), Action, NAME_None, EUserInterfaceActionType::ToggleButton);
+		}
+	}
+	MenuBuilder.EndSection();
+}
+
+void FSimulationEditorExtender::DebugDrawSimulation(const IClothingSimulation* Simulation, USkeletalMeshComponent* /*OwnerComponent*/, FPrimitiveDrawInterface* PDI)
+{
+	if (!ensure(Simulation)) { return; }
+
+	const FClothingSimulation* const ChaosSimulation = static_cast<const FClothingSimulation*>(Simulation);
+
+	for (int32 OptionIndex = 0; OptionIndex < FVisualizationOption::Count; ++OptionIndex)
+	{
+		if (FVisualizationOption::OptionData[OptionIndex].DebugDrawFunction &&
+			(FVisualizationOption::OptionData[OptionIndex].IsConsoleVariableEnabled() || Flags[OptionIndex]))
+		{
+			(ChaosSimulation->*(FVisualizationOption::OptionData[OptionIndex].DebugDrawFunction))(PDI);
+		}
+	}
+}
+
+void FSimulationEditorExtender::DebugDrawSimulationTexts(const IClothingSimulation* Simulation, USkeletalMeshComponent* /*OwnerComponent*/, FCanvas* Canvas, const FSceneView* SceneView)
+{
+	if (!ensure(Simulation)) { return; }
+
+	const FClothingSimulation* const ChaosSimulation = static_cast<const FClothingSimulation*>(Simulation);
+
+	for (int32 OptionIndex = 0; OptionIndex < FVisualizationOption::Count; ++OptionIndex)
+	{
+		if (FVisualizationOption::OptionData[OptionIndex].DebugDrawTextsFunction &&
+			(FVisualizationOption::OptionData[OptionIndex].IsConsoleVariableEnabled() || Flags[OptionIndex]))
+		{
+			(ChaosSimulation->*(FVisualizationOption::OptionData[OptionIndex].DebugDrawTextsFunction))(Canvas, SceneView);
+		}
+	}
+}
+
+bool FSimulationEditorExtender::ShouldDisableSimulation() const
+{
+	for (uint32 OptionIndex = 0; OptionIndex < FVisualizationOption::Count; ++OptionIndex)
+	{
+		if (Flags[OptionIndex])
+		{
+			const FVisualizationOption& Data = FVisualizationOption::OptionData[OptionIndex];
+
+			if (Data.bDisablesSimulation)
+			{
+				return true;
+			}
+		}
+	}
+	return false;
+}
+
+void FSimulationEditorExtender::ShowClothSections(USkeletalMeshComponent* MeshComponent, bool bIsClothSectionsVisible) const
+{
+	if (FSkeletalMeshRenderData* const SkeletalMeshRenderData = MeshComponent->GetSkeletalMeshRenderData())
+	{
+		for (int32 LODIndex = 0; LODIndex < SkeletalMeshRenderData->LODRenderData.Num(); ++LODIndex)
+		{
+			FSkeletalMeshLODRenderData& SkeletalMeshLODRenderData = SkeletalMeshRenderData->LODRenderData[LODIndex];
+
+			for (int32 SectionIndex = 0; SectionIndex < SkeletalMeshLODRenderData.RenderSections.Num(); ++SectionIndex)
+			{
+				FSkelMeshRenderSection& SkelMeshRenderSection = SkeletalMeshLODRenderData.RenderSections[SectionIndex];
+
+				if (SkelMeshRenderSection.HasClothingData())
+				{
+					MeshComponent->ShowMaterialSection(SkelMeshRenderSection.MaterialIndex, SectionIndex, bIsClothSectionsVisible, LODIndex);
+				}
+			}
+		}
+	}
+}
+
+}  // End namespace Chaos
+
+#undef LOCTEXT_NAMESPACE
diff --git a/Engine/Plugins/ChaosCloth/Source/ChaosClothEditor/Public/ChaosClothEditor/ChaosSimulationEditorExtender.h b/Engine/Plugins/ChaosCloth/Source/ChaosClothEditor/Public/ChaosClothEditor/ChaosSimulationEditorExtender.h
new file mode 100644
--- /dev/null
+++ b/Engine/Plugins/ChaosCloth/Source/ChaosClothEditor/Public/ChaosClothEditor/ChaosSimulationEditorExtender.h
@@ -0,0 +1,32 @@
+// Copyright Epic Games, Inc. All Rights Reserved.
+
+#pragma once
+
+#include "SimulationEditorExtender.h"
+
+namespace Chaos
+{
+	/** Chaos extension to the asset editor. */
+	class FSimulationEditorExtender : public ISimulationEditorExtender
+	{
+	public:
+		FSimulationEditorExtender();
+
+		// ISimulationEditorExtender Interface
+		virtual UClass* GetSupportedSimulationFactoryClass() override;
+		virtual void ExtendViewportShowMenu(FMenuBuilder& MenuBuilder, TSharedRef<IPersonaPreviewScene> PreviewScene) override;
+		virtual void DebugDrawSimulation(const IClothingSimulation* Simulation, USkeletalMeshComponent* OwnerComponent, FPrimitiveDrawInterface* PDI) override;
+		virtual void DebugDrawSimulationTexts(const IClothingSimulation* Simulation, USkeletalMeshComponent* OwnerComponent, FCanvas* Canvas, const FSceneView* SceneView) override;
+		// End ISimulationEditorExtender Interface
+
+	private:
+		/** Return whether or not - given the current enabled options - the simulation should be disabled. */
+		bool ShouldDisableSimulation() const;
+		/** Show/hide all cloth sections for the specified mesh compoment. */
+		void ShowClothSections(USkeletalMeshComponent* MeshComponent, bool bIsClothSectionsVisible) const;
+
+	private:
+		/** Flags used to store the checked status for the visualization options. */
+		TBitArray<> Flags;
+	};
+}
diff --git a/Engine/Plugins/Experimental/ChaosClothEditor/ChaosClothEditor.uplugin b/Engine/Plugins/Experimental/ChaosClothEditor/ChaosClothEditor.uplugin
--- a/Engine/Plugins/Experimental/ChaosClothEditor/ChaosClothEditor.uplugin
+++ b/Engine/Plugins/Experimental/ChaosClothEditor/ChaosClothEditor.uplugin
@@ -3,23 +3,17 @@
 	"Version": 1,
 	"VersionName": "1.0",
 	"FriendlyName": "Chaos Cloth Editor",
-	"Description": "Editor module accompanying the Chaos Cloth runtime module.",
+	"Description": "Deprecated plugin, please use the Chaos Cloth plugin instead.",
 	"Category": "Physics",
 	"CreatedBy": "Epic Games, Inc.",
 	"CreatedByURL": "https://epicgames.com",
 	"DocsURL": "",
 	"MarketplaceURL": "",
 	"SupportURL": "",
-	"EnabledByDefault": true,
+	"EnabledByDefault": false,
 	"CanContainContent": true,
 	"Installed": false,
-	"Modules": [
-		{
-			"Name": "ChaosClothEditor",
-			"Type": "Editor",
-			"LoadingPhase": "Default"
-		}
-	],
+	"Modules": [],
 	"Plugins": [
 		{
 			"Name": "ChaosCloth",
