CIS - Remove unacceptable word that was in the GaussSeidel constraint.

diff --git a/Engine/Plugins/ChaosCloth/Source/ChaosCloth/Private/ChaosCloth/ChaosClothConstraints.cpp b/Engine/Plugins/ChaosCloth/Source/ChaosCloth/Private/ChaosCloth/ChaosClothConstraints.cpp
--- a/Engine/Plugins/ChaosCloth/Source/ChaosCloth/Private/ChaosCloth/ChaosClothConstraints.cpp
+++ b/Engine/Plugins/ChaosCloth/Source/ChaosCloth/Private/ChaosCloth/ChaosClothConstraints.cpp
@@ -27,7 +27,7 @@
 #include "Chaos/SoftsEvolution.h"
 #include "Chaos/VelocityField.h"
 #include "Chaos/CollectionPropertyFacade.h"
-#include "Chaos/Deformable/GaussSeidelMasterConstraint.h"
+#include "Chaos/Deformable/GaussSeidelMainConstraint.h"
 #include "Chaos/Deformable/GaussSeidelCorotatedCodimensionalConstraints.h"
 #include "PhysicsProxy/PerSolverFieldSystem.h"
 #include "Utils/ClothingMeshUtils.h"
@@ -557,21 +557,21 @@ void FClothConstraints::CreateGSRules()
 
 	int32 ConstraintPostProcessingIndex = 0;
 
-	GSMasterConstraint = MakeShared<Chaos::Softs::FGaussSeidelMasterConstraint<Softs::FSolverReal, Softs::FSolverParticles>>(PBDEvolution->Particles(), bClothDoQuasistatics, bUseSOR, SOROmega, 100);
+	GSMainConstraint = MakeShared<Chaos::Softs::FGaussSeidelMainConstraint<Softs::FSolverReal, Softs::FSolverParticles>>(PBDEvolution->Particles(), bClothDoQuasistatics, bUseSOR, SOROmega, 100);
 
 	if (bClothDoQuasistatics)
 	{
 		PBDEvolution->SetQuasistatics(bClothDoQuasistatics);
 		const Softs::FSolverVec3& GravityTerm = PBDEvolution->GetGravity(0);
-		GSMasterConstraint->ExternalForce[0] = 0.f;
-		GSMasterConstraint->ExternalForce[1] = 0.f;
-		GSMasterConstraint->ExternalForce[2] = GravityTerm[2];
+		GSMainConstraint->ExternalForce[0] = 0.f;
+		GSMainConstraint->ExternalForce[1] = 0.f;
+		GSMainConstraint->ExternalForce[2] = GravityTerm[2];
 	}
 
 	ConstraintInits[ConstraintInitIndex++] =
 		[this](Softs::FSolverParticles& InParticles, const Softs::FSolverReal Dt)
 	{
-		this->GSMasterConstraint->Init(Dt, InParticles);
+		this->GSMainConstraint->Init(Dt, InParticles);
 	};
 
 	if (XStretchBiasConstraints)
@@ -629,7 +629,7 @@ void FClothConstraints::CreateGSRules()
 					}
 					if (bbDisplayResidual)
 					{
-						TArray<Chaos::Softs::FSolverVec3> Residual = this->GSMasterConstraint->ComputeNewtonResiduals(Particles, Dt);
+						TArray<Chaos::Softs::FSolverVec3> Residual = this->GSMainConstraint->ComputeNewtonResiduals(Particles, Dt);
 					}
 				}
 			};
@@ -644,7 +644,7 @@ void FClothConstraints::CreateGSRules()
 			{
 				for (int32 i = 0; i < GSIterNums; i++)
 				{
-					this->GSMasterConstraint->ApplyCG(InParticles, Dt);
+					this->GSMainConstraint->ApplyCG(InParticles, Dt);
 				}
 			};
 		}
@@ -655,7 +655,7 @@ void FClothConstraints::CreateGSRules()
 			{
 				for (int32 i = 0; i < GSIterNums; i++)
 				{
-					this->GSMasterConstraint->Apply(InParticles, Dt, MaxResidualIters, bWriteResidual2File);
+					this->GSMainConstraint->Apply(InParticles, Dt, MaxResidualIters, bWriteResidual2File);
 				}
 			};
 		}
@@ -668,7 +668,7 @@ void FClothConstraints::CreateGSRules()
 			ConstraintRules[ConstraintRuleIndex++] =
 				[this](Softs::FSolverParticles& InParticles, const Softs::FSolverReal Dt)
 			{
-				this->GSMasterConstraint->ApplyCG(InParticles, Dt);
+				this->GSMainConstraint->ApplyCG(InParticles, Dt);
 			};
 		}
 		else
@@ -676,7 +676,7 @@ void FClothConstraints::CreateGSRules()
 			ConstraintRules[ConstraintRuleIndex++] =
 				[this](Softs::FSolverParticles& InParticles, const Softs::FSolverReal Dt)
 			{
-				this->GSMasterConstraint->Apply(InParticles, Dt, MaxResidualIters, bWriteResidual2File);
+				this->GSMainConstraint->Apply(InParticles, Dt, MaxResidualIters, bWriteResidual2File);
 			};
 		}
 	}
@@ -693,10 +693,10 @@ void FClothConstraints::CreateGSRules()
 
 			GSCorotatedCodimensionalConstraint = MakeShared<Chaos::Softs::FGaussSeidelCorotatedCodimensionalConstraints<Softs::FSolverReal, Softs::FSolverParticles>>(PBDEvolution->Particles(), CodimensionalMesh, false, YoungsModulus);
 
-			GSMasterConstraint->AddStaticConstraints(XStretchBiasConstraints->GetConstraintsArray(), IncidentElements, IncidentElementsLocal);
-			const int32 StaticIndex = GSMasterConstraint->AddStaticConstraintResidualAndHessianRange(1);
+			GSMainConstraint->AddStaticConstraints(XStretchBiasConstraints->GetConstraintsArray(), IncidentElements, IncidentElementsLocal);
+			const int32 StaticIndex = GSMainConstraint->AddStaticConstraintResidualAndHessianRange(1);
 
-			GSMasterConstraint->StaticConstraintResidualAndHessian()[StaticIndex] = [this](const Softs::FSolverParticles& Particles, const int32 ElementIndex, const int32 ElementIndexLocal, const Softs::FSolverReal Dt, TVec3<Softs::FSolverReal>& ParticleResidual, Chaos::PMatrix<Softs::FSolverReal, 3, 3>& ParticleHessian)
+			GSMainConstraint->StaticConstraintResidualAndHessian()[StaticIndex] = [this](const Softs::FSolverParticles& Particles, const int32 ElementIndex, const int32 ElementIndexLocal, const Softs::FSolverReal Dt, TVec3<Softs::FSolverReal>& ParticleResidual, Chaos::PMatrix<Softs::FSolverReal, 3, 3>& ParticleHessian)
 			{
 				this->GSCorotatedCodimensionalConstraint->AddHyperelasticResidualAndHessian(Particles, ElementIndex, ElementIndexLocal, Dt, ParticleResidual, ParticleHessian);
 			};
@@ -704,19 +704,19 @@ void FClothConstraints::CreateGSRules()
 		else
 		{
 			TArray<TArray<int32>> IncidentElements, IncidentElementsLocal;
-			GSMasterConstraint->AddStaticConstraints(XStretchBiasConstraints->GetConstraintsArray(), IncidentElements, IncidentElementsLocal);
-			const int32 StaticIndex = GSMasterConstraint->AddStaticConstraintResidualAndHessianRange(1);
+			GSMainConstraint->AddStaticConstraints(XStretchBiasConstraints->GetConstraintsArray(), IncidentElements, IncidentElementsLocal);
+			const int32 StaticIndex = GSMainConstraint->AddStaticConstraintResidualAndHessianRange(1);
 
-			GSMasterConstraint->StaticConstraintResidualAndHessian()[StaticIndex] = [this](const Softs::FSolverParticles& Particles, const int32 ElementIndex, const int32 ElementIndexLocal, const Softs::FSolverReal Dt, TVec3<Softs::FSolverReal>& ParticleResidual, Chaos::PMatrix<Softs::FSolverReal, 3, 3>& ParticleHessian)
+			GSMainConstraint->StaticConstraintResidualAndHessian()[StaticIndex] = [this](const Softs::FSolverParticles& Particles, const int32 ElementIndex, const int32 ElementIndexLocal, const Softs::FSolverReal Dt, TVec3<Softs::FSolverReal>& ParticleResidual, Chaos::PMatrix<Softs::FSolverReal, 3, 3>& ParticleHessian)
 			{
 				this->XStretchBiasConstraints->AddStretchBiasElementResidualAndHessian(Particles, ElementIndex, ElementIndexLocal, Dt, ParticleResidual, ParticleHessian);
 			};
 
 			if (bEnableCG)
 			{
-				const int32 ForceDifferentialRange = GSMasterConstraint->AddAddInternalForceDifferentialsRange(1);
+				const int32 ForceDifferentialRange = GSMainConstraint->AddAddInternalForceDifferentialsRange(1);
 
-				GSMasterConstraint->InternalForceDifferentials()[ForceDifferentialRange] = [this](const Softs::FSolverParticles& Particles, const TArray<Softs::FSolverVec3>& Deltax, TArray<Softs::FSolverVec3>& ndf)
+				GSMainConstraint->InternalForceDifferentials()[ForceDifferentialRange] = [this](const Softs::FSolverParticles& Particles, const TArray<Softs::FSolverVec3>& Deltax, TArray<Softs::FSolverVec3>& ndf)
 				{
 					this->XStretchBiasConstraints->AddInternalForceDifferential(Particles, Deltax, ndf);
 				};
@@ -731,19 +731,19 @@ void FClothConstraints::CreateGSRules()
 	if (XAnisoBendingElementConstraints)
 	{
 		TArray<TArray<int32>> IncidentElements, IncidentElementsLocal;
-		GSMasterConstraint->AddStaticConstraints(XAnisoBendingElementConstraints->GetConstraintsArray(), IncidentElements, IncidentElementsLocal);
-		const int32 StaticIndex = GSMasterConstraint->AddStaticConstraintResidualAndHessianRange(1);
+		GSMainConstraint->AddStaticConstraints(XAnisoBendingElementConstraints->GetConstraintsArray(), IncidentElements, IncidentElementsLocal);
+		const int32 StaticIndex = GSMainConstraint->AddStaticConstraintResidualAndHessianRange(1);
 
-		GSMasterConstraint->StaticConstraintResidualAndHessian()[StaticIndex] = [this](const Softs::FSolverParticles& Particles, const int32 ElementIndex, const int32 ElementIndexLocal, const Softs::FSolverReal Dt, TVec3<Softs::FSolverReal>& ParticleResidual, Chaos::PMatrix<Softs::FSolverReal, 3, 3>& ParticleHessian)
+		GSMainConstraint->StaticConstraintResidualAndHessian()[StaticIndex] = [this](const Softs::FSolverParticles& Particles, const int32 ElementIndex, const int32 ElementIndexLocal, const Softs::FSolverReal Dt, TVec3<Softs::FSolverReal>& ParticleResidual, Chaos::PMatrix<Softs::FSolverReal, 3, 3>& ParticleHessian)
 		{
 			this->XAnisoBendingElementConstraints->AddAnisotropicBendingResidualAndHessian(Particles, ElementIndex, ElementIndexLocal, Dt, ParticleResidual, ParticleHessian);
 		};
 
 		if (bEnableCG)
 		{
-			const int32 ForceDifferentialRange = GSMasterConstraint->AddAddInternalForceDifferentialsRange(1);
+			const int32 ForceDifferentialRange = GSMainConstraint->AddAddInternalForceDifferentialsRange(1);
 
-			GSMasterConstraint->InternalForceDifferentials()[ForceDifferentialRange] = [this](const Softs::FSolverParticles& Particles, const TArray<Softs::FSolverVec3>& Deltax, TArray<Softs::FSolverVec3>& ndf)
+			GSMainConstraint->InternalForceDifferentials()[ForceDifferentialRange] = [this](const Softs::FSolverParticles& Particles, const TArray<Softs::FSolverVec3>& Deltax, TArray<Softs::FSolverVec3>& ndf)
 			{
 				this->XAnisoBendingElementConstraints->AddInternalForceDifferential(Particles, Deltax, ndf);
 			};
@@ -754,10 +754,10 @@ void FClothConstraints::CreateGSRules()
 	{
 
 		TArray<TArray<int32>> IncidentElements, IncidentElementsLocal;
-		GSMasterConstraint->AddStaticConstraints(XBendingElementConstraints->GetConstraintsArray(), IncidentElements, IncidentElementsLocal);
-		const int32 StaticIndex = GSMasterConstraint->AddStaticConstraintResidualAndHessianRange(1);
+		GSMainConstraint->AddStaticConstraints(XBendingElementConstraints->GetConstraintsArray(), IncidentElements, IncidentElementsLocal);
+		const int32 StaticIndex = GSMainConstraint->AddStaticConstraintResidualAndHessianRange(1);
 
-		GSMasterConstraint->StaticConstraintResidualAndHessian()[StaticIndex] = [this](const Softs::FSolverParticles& Particles, const int32 ElementIndex, const int32 ElementIndexLocal, const Softs::FSolverReal Dt, TVec3<Softs::FSolverReal>& ParticleResidual, Chaos::PMatrix<Softs::FSolverReal, 3, 3>& ParticleHessian)
+		GSMainConstraint->StaticConstraintResidualAndHessian()[StaticIndex] = [this](const Softs::FSolverParticles& Particles, const int32 ElementIndex, const int32 ElementIndexLocal, const Softs::FSolverReal Dt, TVec3<Softs::FSolverReal>& ParticleResidual, Chaos::PMatrix<Softs::FSolverReal, 3, 3>& ParticleHessian)
 		{
 			this->XBendingElementConstraints->AddBendingResidualAndHessian(Particles, ElementIndex, ElementIndexLocal, Dt, ParticleResidual, ParticleHessian);
 		};
@@ -765,16 +765,16 @@ void FClothConstraints::CreateGSRules()
 
 		if (bEnableCG)
 		{
-			const int32 ForceDifferentialRange = GSMasterConstraint->AddAddInternalForceDifferentialsRange(1);
+			const int32 ForceDifferentialRange = GSMainConstraint->AddAddInternalForceDifferentialsRange(1);
 
-			GSMasterConstraint->InternalForceDifferentials()[ForceDifferentialRange] = [this](const Softs::FSolverParticles& Particles, const TArray<Softs::FSolverVec3>& Deltax, TArray<Softs::FSolverVec3>& ndf)
+			GSMainConstraint->InternalForceDifferentials()[ForceDifferentialRange] = [this](const Softs::FSolverParticles& Particles, const TArray<Softs::FSolverVec3>& Deltax, TArray<Softs::FSolverVec3>& ndf)
 			{
 				this->XBendingElementConstraints->AddInternalForceDifferential(Particles, Deltax, ndf);
 			};
 		}
 	}
 
-	GSMasterConstraint->InitStaticColor(PBDEvolution->Particles());
+	GSMainConstraint->InitStaticColor(PBDEvolution->Particles());
 
 	if (bWriteFinalResiduals)
 	{
@@ -784,7 +784,7 @@ void FClothConstraints::CreateGSRules()
 		PostprocessingRules[ConstraintPostProcessingIndex++] =
 			[this](Softs::FSolverParticles& InParticles, const Softs::FSolverReal Dt)
 		{
-			this->GSMasterConstraint->ComputeNewtonResiduals(InParticles, Dt, true, nullptr);
+			this->GSMainConstraint->ComputeNewtonResiduals(InParticles, Dt, true, nullptr);
 		};
 	}
 
@@ -1825,12 +1825,12 @@ void FClothConstraints::CreatePBDRules()
 	#if !UE_BUILD_SHIPPING	
 	if (bDisplayResidual || bWriteFinalResiduals)
 	{
-		GSMasterConstraint = MakeShared<Chaos::Softs::FGaussSeidelMasterConstraint<Softs::FSolverReal, Softs::FSolverParticles>>(PBDEvolution->Particles(), false, false, (Softs::FSolverReal)1.2, 100);
+		GSMainConstraint = MakeShared<Chaos::Softs::FGaussSeidelMainConstraint<Softs::FSolverReal, Softs::FSolverParticles>>(PBDEvolution->Particles(), false, false, (Softs::FSolverReal)1.2, 100);
 
 		ConstraintInits[ConstraintInitIndex++] =
 			[this](Softs::FSolverParticles& InParticles, const Softs::FSolverReal Dt)
 		{
-			this->GSMasterConstraint->Init(Dt, InParticles);
+			this->GSMainConstraint->Init(Dt, InParticles);
 		};
 		if (bDisplayResidual)
 		{
@@ -1838,11 +1838,11 @@ void FClothConstraints::CreatePBDRules()
 			ConstraintRules[ConstraintRuleIndex++] =
 				[this](Softs::FSolverParticles& InParticles, const Softs::FSolverReal Dt)
 			{
-				if (this->GSMasterConstraint->DebugResidual && this->GSMasterConstraint->PassedIters < MaxResidualIters)
+				if (this->GSMainConstraint->DebugResidual && this->GSMainConstraint->PassedIters < MaxResidualIters)
 				{
 					UE_LOG(LogTemp, Warning, TEXT("Current Iteration As From Evolution: %d"), this->Evolution->GetIterations());
-					this->GSMasterConstraint->ComputeNewtonResiduals(InParticles, Dt, bWriteResidual2File);
-					this->GSMasterConstraint->PassedIters += 1;
+					this->GSMainConstraint->ComputeNewtonResiduals(InParticles, Dt, bWriteResidual2File);
+					this->GSMainConstraint->PassedIters += 1;
 				}
 			};
 		}
@@ -1850,10 +1850,10 @@ void FClothConstraints::CreatePBDRules()
 		if (XStretchBiasConstraints)
 		{
 			TArray<TArray<int32>> IncidentElements, IncidentElementsLocal;
-			GSMasterConstraint->AddStaticConstraints(XStretchBiasConstraints->GetConstraintsArray(), IncidentElements, IncidentElementsLocal);
-			int32 StaticIndex = GSMasterConstraint->AddStaticConstraintResidualAndHessianRange(1);
+			GSMainConstraint->AddStaticConstraints(XStretchBiasConstraints->GetConstraintsArray(), IncidentElements, IncidentElementsLocal);
+			int32 StaticIndex = GSMainConstraint->AddStaticConstraintResidualAndHessianRange(1);
 
-			GSMasterConstraint->StaticConstraintResidualAndHessian()[StaticIndex] = [this](const Softs::FSolverParticles& Particles, const int32 ElementIndex, const int32 ElementIndexLocal, const Softs::FSolverReal Dt, TVec3<Softs::FSolverReal>& ParticleResidual, Chaos::PMatrix<Softs::FSolverReal, 3, 3>& ParticleHessian)
+			GSMainConstraint->StaticConstraintResidualAndHessian()[StaticIndex] = [this](const Softs::FSolverParticles& Particles, const int32 ElementIndex, const int32 ElementIndexLocal, const Softs::FSolverReal Dt, TVec3<Softs::FSolverReal>& ParticleResidual, Chaos::PMatrix<Softs::FSolverReal, 3, 3>& ParticleHessian)
 			{
 				this->XStretchBiasConstraints->AddStretchBiasElementResidualAndHessian(Particles, ElementIndex, ElementIndexLocal, Dt, ParticleResidual, ParticleHessian);
 			};
@@ -1862,10 +1862,10 @@ void FClothConstraints::CreatePBDRules()
 		if (XAnisoBendingElementConstraints)
 		{
 			TArray<TArray<int32>> IncidentElements, IncidentElementsLocal;
-			GSMasterConstraint->AddStaticConstraints(XAnisoBendingElementConstraints->GetConstraintsArray(), IncidentElements, IncidentElementsLocal);
-			int32 StaticIndex = GSMasterConstraint->AddStaticConstraintResidualAndHessianRange(1);
+			GSMainConstraint->AddStaticConstraints(XAnisoBendingElementConstraints->GetConstraintsArray(), IncidentElements, IncidentElementsLocal);
+			int32 StaticIndex = GSMainConstraint->AddStaticConstraintResidualAndHessianRange(1);
 
-			GSMasterConstraint->StaticConstraintResidualAndHessian()[StaticIndex] = [this](const Softs::FSolverParticles& Particles, const int32 ElementIndex, const int32 ElementIndexLocal, const Softs::FSolverReal Dt, TVec3<Softs::FSolverReal>& ParticleResidual, Chaos::PMatrix<Softs::FSolverReal, 3, 3>& ParticleHessian)
+			GSMainConstraint->StaticConstraintResidualAndHessian()[StaticIndex] = [this](const Softs::FSolverParticles& Particles, const int32 ElementIndex, const int32 ElementIndexLocal, const Softs::FSolverReal Dt, TVec3<Softs::FSolverReal>& ParticleResidual, Chaos::PMatrix<Softs::FSolverReal, 3, 3>& ParticleHessian)
 			{
 				this->XAnisoBendingElementConstraints->AddAnisotropicBendingResidualAndHessian(Particles, ElementIndex, ElementIndexLocal, Dt, ParticleResidual, ParticleHessian);
 			};
@@ -1875,16 +1875,16 @@ void FClothConstraints::CreatePBDRules()
 		if (XBendingElementConstraints)
 		{
 			TArray<TArray<int32>> IncidentElements, IncidentElementsLocal;
-			GSMasterConstraint->AddStaticConstraints(XBendingElementConstraints->GetConstraintsArray(), IncidentElements, IncidentElementsLocal);
-			int32 StaticIndex = GSMasterConstraint->AddStaticConstraintResidualAndHessianRange(1);
+			GSMainConstraint->AddStaticConstraints(XBendingElementConstraints->GetConstraintsArray(), IncidentElements, IncidentElementsLocal);
+			int32 StaticIndex = GSMainConstraint->AddStaticConstraintResidualAndHessianRange(1);
 
-			GSMasterConstraint->StaticConstraintResidualAndHessian()[StaticIndex] = [this](const Softs::FSolverParticles& Particles, const int32 ElementIndex, const int32 ElementIndexLocal, const Softs::FSolverReal Dt, TVec3<Softs::FSolverReal>& ParticleResidual, Chaos::PMatrix<Softs::FSolverReal, 3, 3>& ParticleHessian)
+			GSMainConstraint->StaticConstraintResidualAndHessian()[StaticIndex] = [this](const Softs::FSolverParticles& Particles, const int32 ElementIndex, const int32 ElementIndexLocal, const Softs::FSolverReal Dt, TVec3<Softs::FSolverReal>& ParticleResidual, Chaos::PMatrix<Softs::FSolverReal, 3, 3>& ParticleHessian)
 			{
 				this->XBendingElementConstraints->AddBendingResidualAndHessian(Particles, ElementIndex, ElementIndexLocal, Dt, ParticleResidual, ParticleHessian);
 			};
 		}
 
-		GSMasterConstraint->InitStaticColor(PBDEvolution->Particles());
+		GSMainConstraint->InitStaticColor(PBDEvolution->Particles());
 
 		if (bWriteFinalResiduals)
 		{
@@ -1894,7 +1894,7 @@ void FClothConstraints::CreatePBDRules()
 			PostprocessingConstraintRules[PostprocessingConstraintRuleIndex++] =
 				[this](Softs::FSolverParticles& InParticles, const Softs::FSolverReal Dt)
 			{
-				this->GSMasterConstraint->ComputeNewtonResiduals(InParticles, Dt, true, nullptr);
+				this->GSMainConstraint->ComputeNewtonResiduals(InParticles, Dt, true, nullptr);
 			};
 		}
 	}
diff --git a/Engine/Plugins/ChaosCloth/Source/ChaosCloth/Public/ChaosCloth/ChaosClothConstraints.h b/Engine/Plugins/ChaosCloth/Source/ChaosCloth/Public/ChaosCloth/ChaosClothConstraints.h
--- a/Engine/Plugins/ChaosCloth/Source/ChaosCloth/Public/ChaosCloth/ChaosClothConstraints.h
+++ b/Engine/Plugins/ChaosCloth/Source/ChaosCloth/Public/ChaosCloth/ChaosClothConstraints.h
@@ -7,7 +7,7 @@
 #include "Chaos/PBDSoftsEvolutionFwd.h"
 #include "Chaos/PBDLongRangeConstraints.h"
 #include "Chaos/PBDCollisionSpringConstraintsBase.h"
-#include "Chaos/Deformable/GaussSeidelMasterConstraint.h"
+#include "Chaos/Deformable/GaussSeidelMainConstraint.h"
 #include "Chaos/Deformable/GaussSeidelCorotatedCodimensionalConstraints.h"
 
 namespace Chaos
@@ -238,7 +238,7 @@ PRAGMA_ENABLE_DEPRECATION_WARNINGS
 		TSharedPtr<Softs::FPBDCollisionSpringConstraints> SelfCollisionConstraints;
 		TSharedPtr<Softs::FPBDTriangleMeshIntersections> SelfIntersectionConstraints;
 		TSharedPtr<Softs::FPBDSelfCollisionSphereConstraints> SelfCollisionSphereConstraints;
-		TSharedPtr<Softs::FGaussSeidelMasterConstraint<Softs::FSolverReal, Softs::FSolverParticles>> GSMasterConstraint;
+		TSharedPtr<Softs::FGaussSeidelMainConstraint<Softs::FSolverReal, Softs::FSolverParticles>> GSMainConstraint;
 		TSharedPtr<Softs::FGaussSeidelCorotatedCodimensionalConstraints<Softs::FSolverReal, Softs::FSolverParticles>> GSCorotatedCodimensionalConstraint;
 		TSharedPtr<Softs::FMultiResConstraints> MultiResConstraints;
 		//~ Begin Force-based solver only constraints
diff --git a/Engine/Source/Runtime/Experimental/Chaos/Private/Chaos/Deformable/ChaosDeformableSolver.cpp b/Engine/Source/Runtime/Experimental/Chaos/Private/Chaos/Deformable/ChaosDeformableSolver.cpp
--- a/Engine/Source/Runtime/Experimental/Chaos/Private/Chaos/Deformable/ChaosDeformableSolver.cpp
+++ b/Engine/Source/Runtime/Experimental/Chaos/Private/Chaos/Deformable/ChaosDeformableSolver.cpp
@@ -1201,7 +1201,7 @@ namespace Chaos::Softs
 	{
 		PERF_SCOPE(STAT_ChaosDeformableSolver_InitializeGaussSeidelConstraintVariables);
 
-		GSMasterConstraint.Reset(new Chaos::Softs::FGaussSeidelMasterConstraint<FSolverReal, FSolverParticles>(Evolution->Particles(), Property.bDoQuasistatics, Property.bUseSOR, Property.OmegaSOR, GSParallelMax));
+		GSMainConstraint.Reset(new Chaos::Softs::FGaussSeidelMainConstraint<FSolverReal, FSolverParticles>(Evolution->Particles(), Property.bDoQuasistatics, Property.bUseSOR, Property.OmegaSOR, GSParallelMax));
 
 		if (Property.bUseGSNeohookean)
 		{
@@ -1209,26 +1209,26 @@ namespace Chaos::Softs
 				Evolution->Particles(), *AllElements, *AllTetEMeshArray, *AllTetNuMeshArray, MoveTemp(*AllTetAlphaJArray), MoveTemp(*AllIncidentElements), MoveTemp(*AllIncidentElementsLocal), 0, Evolution->Particles().Size(), Property.bDoQuasistatics, Property.bUseSOR, Property.OmegaSOR, GDeformableXPBDCorotatedParams));
 			Evolution->ResetConstraintRules();
 			
-			GSMasterConstraint->AddStaticConstraints(GSNeohookeanConstraints->GetMeshArray(), GSNeohookeanConstraints->GetIncidentElements(), GSNeohookeanConstraints->GetIncidentElementsLocal());
+			GSMainConstraint->AddStaticConstraints(GSNeohookeanConstraints->GetMeshArray(), GSNeohookeanConstraints->GetIncidentElements(), GSNeohookeanConstraints->GetIncidentElementsLocal());
 
 			int32 InitIndex1 = Evolution->AddConstraintInitRange(1, true);
 			Evolution->ConstraintInits()[InitIndex1] =
 				[this](FSolverParticles& InParticles, const FSolverReal Dt)
 			{
-				this->GSMasterConstraint->Init(Dt, InParticles);
+				this->GSMainConstraint->Init(Dt, InParticles);
 			};
 
 			int32 ConstraintIndex1 = Evolution->AddConstraintRuleRange(1, true);
 			Evolution->ConstraintRules()[ConstraintIndex1] =
 				[this](FSolverParticles& InParticles, const FSolverReal Dt)
 			{
-				this->GSMasterConstraint->Apply(InParticles, Dt);
+				this->GSMainConstraint->Apply(InParticles, Dt);
 			};
 
 
-			int32 StaticIndex = GSMasterConstraint->AddStaticConstraintResidualAndHessianRange(1);
+			int32 StaticIndex = GSMainConstraint->AddStaticConstraintResidualAndHessianRange(1);
 
-			GSMasterConstraint->StaticConstraintResidualAndHessian()[StaticIndex] = [this](const FSolverParticles& Particles, const int32 ElementIndex, const int32 ElementIndexLocal, const FSolverReal Dt, TVec3<FSolverReal>& ParticleResidual, Chaos::PMatrix<FSolverReal, 3, 3>& ParticleHessian)
+			GSMainConstraint->StaticConstraintResidualAndHessian()[StaticIndex] = [this](const FSolverParticles& Particles, const int32 ElementIndex, const int32 ElementIndexLocal, const FSolverReal Dt, TVec3<FSolverReal>& ParticleResidual, Chaos::PMatrix<FSolverReal, 3, 3>& ParticleHessian)
 			{
 				this->GSNeohookeanConstraints->AddHyperelasticResidualAndHessian(Particles, ElementIndex, ElementIndexLocal, Dt, ParticleResidual, ParticleHessian);
 			};
@@ -1241,16 +1241,16 @@ namespace Chaos::Softs
 
 				TArray<TArray<int32>> StaticIncidentElements, StaticIncidentElementsLocal;
 				TArray<TArray<int32>> StaticConstraints = GSWeakConstraints->GetStaticConstraintArrays(StaticIncidentElements, StaticIncidentElementsLocal);
-				GSMasterConstraint->AddStaticConstraints(StaticConstraints, StaticIncidentElements, StaticIncidentElementsLocal);
+				GSMainConstraint->AddStaticConstraints(StaticConstraints, StaticIncidentElements, StaticIncidentElementsLocal);
 
-				int32 StaticIndex1 = GSMasterConstraint->AddStaticConstraintResidualAndHessianRange(1);
-				GSMasterConstraint->StaticConstraintResidualAndHessian()[StaticIndex1] = [this](const FSolverParticles& Particles, const int32 ConstraintIndex, const int32 ConstraintIndexLocal, const FSolverReal Dt, TVec3<FSolverReal>& ParticleResidual, Chaos::PMatrix<FSolverReal, 3, 3>& ParticleHessian)
+				int32 StaticIndex1 = GSMainConstraint->AddStaticConstraintResidualAndHessianRange(1);
+				GSMainConstraint->StaticConstraintResidualAndHessian()[StaticIndex1] = [this](const FSolverParticles& Particles, const int32 ConstraintIndex, const int32 ConstraintIndexLocal, const FSolverReal Dt, TVec3<FSolverReal>& ParticleResidual, Chaos::PMatrix<FSolverReal, 3, 3>& ParticleHessian)
 				{
 					this->GSWeakConstraints->AddWCResidualAndHessian(Particles, ConstraintIndex, ConstraintIndexLocal, Dt, ParticleResidual, ParticleHessian);
 				};
 
-				int32 PerNodeIndex = GSMasterConstraint->AddPerNodeHessianRange(1);
-				GSMasterConstraint->PerNodeHessian()[PerNodeIndex] = [this](const int32 p, const FSolverReal Dt, Chaos::PMatrix<FSolverReal, 3, 3>& ParticleHessian)
+				int32 PerNodeIndex = GSMainConstraint->AddPerNodeHessianRange(1);
+				GSMainConstraint->PerNodeHessian()[PerNodeIndex] = [this](const int32 p, const FSolverReal Dt, Chaos::PMatrix<FSolverReal, 3, 3>& ParticleHessian)
 				{
 					this->GSWeakConstraints->AddWCHessian(p, Dt, ParticleHessian);
 				};
@@ -1262,26 +1262,26 @@ namespace Chaos::Softs
 				Evolution->Particles(), *AllElements, *AllTetEMeshArray, *AllTetNuMeshArray, MoveTemp(*AllTetAlphaJArray), MoveTemp(*AllIncidentElements), MoveTemp(*AllIncidentElementsLocal), 0, Evolution->Particles().Size(), Property.bDoQuasistatics, Property.bUseSOR, Property.OmegaSOR, GDeformableXPBDCorotatedParams));
 			Evolution->ResetConstraintRules();
 
-			GSMasterConstraint->AddStaticConstraints(GSCorotatedConstraints->GetMeshArray(), GSCorotatedConstraints->GetIncidentElements(), GSCorotatedConstraints->GetIncidentElementsLocal());
+			GSMainConstraint->AddStaticConstraints(GSCorotatedConstraints->GetMeshArray(), GSCorotatedConstraints->GetIncidentElements(), GSCorotatedConstraints->GetIncidentElementsLocal());
 			
 			int32 InitIndex1 = Evolution->AddConstraintInitRange(1, true);
 			Evolution->ConstraintInits()[InitIndex1] =
 				[this](FSolverParticles& InParticles, const FSolverReal Dt)
 			{
-				this->GSMasterConstraint->Init(Dt, InParticles);
+				this->GSMainConstraint->Init(Dt, InParticles);
 			};
 
 			int32 ConstraintIndex1 = Evolution->AddConstraintRuleRange(1, true);
 			Evolution->ConstraintRules()[ConstraintIndex1] =
 				[this](FSolverParticles& InParticles, const FSolverReal Dt)
 			{
-				this->GSMasterConstraint->Apply(InParticles, Dt);
+				this->GSMainConstraint->Apply(InParticles, Dt);
 			};
 
 
-			int32 StaticIndex = GSMasterConstraint->AddStaticConstraintResidualAndHessianRange(1);
+			int32 StaticIndex = GSMainConstraint->AddStaticConstraintResidualAndHessianRange(1);
 
-			GSMasterConstraint->StaticConstraintResidualAndHessian()[StaticIndex] = [this](const FSolverParticles& Particles, const int32 ElementIndex, const int32 ElementIndexLocal, const FSolverReal Dt, TVec3<FSolverReal>& ParticleResidual, Chaos::PMatrix<FSolverReal, 3, 3>& ParticleHessian)
+			GSMainConstraint->StaticConstraintResidualAndHessian()[StaticIndex] = [this](const FSolverParticles& Particles, const int32 ElementIndex, const int32 ElementIndexLocal, const FSolverReal Dt, TVec3<FSolverReal>& ParticleResidual, Chaos::PMatrix<FSolverReal, 3, 3>& ParticleHessian)
 			{
 				this->GSCorotatedConstraints->AddHyperelasticResidualAndHessian(Particles, ElementIndex, ElementIndexLocal, Dt, ParticleResidual, ParticleHessian);
 			};
@@ -1294,16 +1294,16 @@ namespace Chaos::Softs
 
 				TArray<TArray<int32>> StaticIncidentElements, StaticIncidentElementsLocal;
 				TArray<TArray<int32>> StaticConstraints = GSWeakConstraints->GetStaticConstraintArrays(StaticIncidentElements, StaticIncidentElementsLocal);
-				GSMasterConstraint->AddStaticConstraints(StaticConstraints, StaticIncidentElements, StaticIncidentElementsLocal);
+				GSMainConstraint->AddStaticConstraints(StaticConstraints, StaticIncidentElements, StaticIncidentElementsLocal);
 
-				int32 StaticIndex1 = GSMasterConstraint->AddStaticConstraintResidualAndHessianRange(1);
-				GSMasterConstraint->StaticConstraintResidualAndHessian()[StaticIndex1] = [this](const FSolverParticles& Particles, const int32 ConstraintIndex, const int32 ConstraintIndexLocal, const FSolverReal Dt, TVec3<FSolverReal>& ParticleResidual, Chaos::PMatrix<FSolverReal, 3, 3>& ParticleHessian)
+				int32 StaticIndex1 = GSMainConstraint->AddStaticConstraintResidualAndHessianRange(1);
+				GSMainConstraint->StaticConstraintResidualAndHessian()[StaticIndex1] = [this](const FSolverParticles& Particles, const int32 ConstraintIndex, const int32 ConstraintIndexLocal, const FSolverReal Dt, TVec3<FSolverReal>& ParticleResidual, Chaos::PMatrix<FSolverReal, 3, 3>& ParticleHessian)
 				{
 					this->GSWeakConstraints->AddWCResidualAndHessian(Particles, ConstraintIndex, ConstraintIndexLocal, Dt, ParticleResidual, ParticleHessian);
 				};
 
-				int32 PerNodeIndex = GSMasterConstraint->AddPerNodeHessianRange(1);
-				GSMasterConstraint->PerNodeHessian()[PerNodeIndex] = [this](const int32 p, const FSolverReal Dt, Chaos::PMatrix<FSolverReal, 3, 3>& ParticleHessian)
+				int32 PerNodeIndex = GSMainConstraint->AddPerNodeHessianRange(1);
+				GSMainConstraint->PerNodeHessian()[PerNodeIndex] = [this](const int32 p, const FSolverReal Dt, Chaos::PMatrix<FSolverReal, 3, 3>& ParticleHessian)
 				{
 					this->GSWeakConstraints->AddWCHessian(p, Dt, ParticleHessian);
 				};
@@ -1311,7 +1311,7 @@ namespace Chaos::Softs
 			}
 		}
 
-		GSMasterConstraint->InitStaticColor(Evolution->Particles());
+		GSMainConstraint->InitStaticColor(Evolution->Particles());
 
 		if (Property.bEnablePositionTargets)
 		{
@@ -1325,16 +1325,16 @@ namespace Chaos::Softs
 
 		if (Property.bDoSelfCollision)
 		{
-			int32 DynamicIndex = GSMasterConstraint->AddDynamicConstraintResidualAndHessianRange(1);
-			GSMasterConstraint->DynamicConstraintResidualAndHessian()[DynamicIndex] = [this](const FSolverParticles& Particles, const int32 ConstraintIndex, const int32 ConstraintIndexLocal, const FSolverReal Dt, TVec3<FSolverReal>& ParticleResidual, Chaos::PMatrix<FSolverReal, 3, 3>& ParticleHessian)
+			int32 DynamicIndex = GSMainConstraint->AddDynamicConstraintResidualAndHessianRange(1);
+			GSMainConstraint->DynamicConstraintResidualAndHessian()[DynamicIndex] = [this](const FSolverParticles& Particles, const int32 ConstraintIndex, const int32 ConstraintIndexLocal, const FSolverReal Dt, TVec3<FSolverReal>& ParticleResidual, Chaos::PMatrix<FSolverReal, 3, 3>& ParticleHessian)
 			{
 				this->GSWeakConstraints->AddWCResidualAndHessian(Particles, ConstraintIndex + this->GSWeakConstraints->InitialWCSize, ConstraintIndexLocal, Dt, ParticleResidual, ParticleHessian);
 			};
 
 			if (!Property.bEnablePositionTargets)
 			{
-				int32 PerNodeIndex = GSMasterConstraint->AddPerNodeHessianRange(1);
-				GSMasterConstraint->PerNodeHessian()[PerNodeIndex] = [this](const int32 p, const FSolverReal Dt, Chaos::PMatrix<FSolverReal, 3, 3>& ParticleHessian)
+				int32 PerNodeIndex = GSMainConstraint->AddPerNodeHessianRange(1);
+				GSMainConstraint->PerNodeHessian()[PerNodeIndex] = [this](const int32 p, const FSolverReal Dt, Chaos::PMatrix<FSolverReal, 3, 3>& ParticleHessian)
 				{
 					this->GSWeakConstraints->AddWCHessian(p, Dt, ParticleHessian);
 				};
@@ -1363,8 +1363,8 @@ namespace Chaos::Softs
 			{
 				TArray<TArray<int32>> WCCollisionConstraints, WCCollisionIncidentElements, WCCollisionIncidentElementsLocal;
 				this->GSWeakConstraints->ComputeCollisionWCDataSimplified(WCCollisionConstraints, WCCollisionIncidentElements, WCCollisionIncidentElementsLocal);
-				this->GSMasterConstraint->AddDynamicConstraints(WCCollisionConstraints, WCCollisionIncidentElements, WCCollisionIncidentElementsLocal);
-				this->GSMasterConstraint->InitDynamicColor(InParticles);
+				this->GSMainConstraint->AddDynamicConstraints(WCCollisionConstraints, WCCollisionIncidentElements, WCCollisionIncidentElementsLocal);
+				this->GSMainConstraint->InitDynamicColor(InParticles);
 			}; 
 			
 		}
diff --git a/Engine/Source/Runtime/Experimental/Chaos/Private/Chaos/Deformable/GaussSeidelMainConstraint.cpp b/Engine/Source/Runtime/Experimental/Chaos/Private/Chaos/Deformable/GaussSeidelMainConstraint.cpp
new file mode 100644
--- /dev/null
+++ b/Engine/Source/Runtime/Experimental/Chaos/Private/Chaos/Deformable/GaussSeidelMainConstraint.cpp
@@ -0,0 +1,229 @@
+// Copyright Epic Games, Inc. All Rights Reserved.
+#include "Chaos/Deformable/GaussSeidelMainConstraint.h"
+#include "Chaos/Math/Krylov.h"
+
+namespace Chaos::Softs
+{
+	int32 MaxItCG = 50;
+
+	FAutoConsoleVariableRef CVarClothMaxItCG(TEXT("p.Chaos.Cloth.MaxItCG"), MaxItCG, TEXT("Max iter for CG [def: 50]"));
+
+	FSolverReal CGTol = 1e-4f;
+
+	FAutoConsoleVariableRef CVarClothCGTol(TEXT("p.Chaos.Cloth.CGTol"), MaxItCG, TEXT("CG Tolerance [def: 1e-4]"));
+
+	template <typename T, typename ParticleType>
+	void FGaussSeidelMainConstraint<T, ParticleType>::AddStaticConstraints(const TArray<TArray<int32>>& ExtraConstraints, TArray<TArray<int32>>& ExtraIncidentElements, TArray<TArray<int32>>& ExtraIncidentElementsLocal)
+	{	
+		if (!IsClean(ExtraConstraints, ExtraIncidentElements, ExtraIncidentElementsLocal))
+		{
+			ExtraIncidentElements = Chaos::Utilities::ComputeIncidentElements(ExtraConstraints, &ExtraIncidentElementsLocal);
+		}
+
+		int32 Offset = StaticConstraints.Num();
+		StaticConstraints += ExtraConstraints;
+		for (int32 i = 0; i < ExtraIncidentElements.Num(); i++)
+		{
+			if (ExtraIncidentElements[i].Num() > 0)
+			{
+				TArray<int32> ExtraIncidentElementsWithOffset = ExtraIncidentElements[i];
+				for (int32 j = 0; j < ExtraIncidentElementsWithOffset.Num(); j++)
+				{
+					ExtraIncidentElementsWithOffset[j] += Offset;
+				}
+				StaticIncidentElements[i] += ExtraIncidentElementsWithOffset;
+				StaticIncidentElementsLocal[i] += ExtraIncidentElementsLocal[i];
+			}
+		}
+		if (StaticIncidentElementsOffsets.Num() > 0)
+		{
+			StaticIncidentElementsOffsets.RemoveAt(StaticIncidentElementsOffsets.Num() - 1);
+		}
+		StaticIncidentElementsOffsets.Add(Offset);
+		StaticIncidentElementsOffsets.Add(StaticConstraints.Num());
+	}
+
+	template <typename T, typename ParticleType>
+	void FGaussSeidelMainConstraint<T, ParticleType>::AddDynamicConstraints(const TArray<TArray<int32>>& ExtraConstraints, TArray<TArray<int32>>& ExtraIncidentElements, TArray<TArray<int32>>& ExtraIncidentElementsLocal, bool CheckIncidentElements)
+	{
+		if (CheckIncidentElements)
+		{
+			if (!IsClean(ExtraConstraints, ExtraIncidentElements, ExtraIncidentElementsLocal))
+			{
+				ExtraIncidentElements = Chaos::Utilities::ComputeIncidentElements(ExtraConstraints, &ExtraIncidentElementsLocal);
+			}
+		}
+
+		int32 Offset = DynamicConstraints.Num();
+		DynamicConstraints += ExtraConstraints;
+		for (int32 i = 0; i < ExtraIncidentElements.Num(); i++)
+		{
+			if (ExtraIncidentElements[i].Num() > 0)
+			{
+				TArray<int32> ExtraIncidentElementsWithOffset = ExtraIncidentElements[i];
+				for (int32 j = 0; j < ExtraIncidentElementsWithOffset.Num(); j++)
+				{
+					ExtraIncidentElementsWithOffset[j] += Offset;
+				}
+				DynamicIncidentElements[i] += ExtraIncidentElementsWithOffset;
+				DynamicIncidentElementsLocal[i] += ExtraIncidentElementsLocal[i];
+			}
+		}
+		if (DynamicIncidentElementsOffsets.Num() > 0)
+		{
+			DynamicIncidentElementsOffsets.RemoveAt(StaticIncidentElementsOffsets.Num() - 1);
+		}
+		DynamicIncidentElementsOffsets.Add(Offset);
+		DynamicIncidentElementsOffsets.Add(DynamicConstraints.Num());
+	
+	}
+
+	template <typename T, typename ParticleType>
+	TArray<TVec3<T>> FGaussSeidelMainConstraint<T, ParticleType>::ComputeNewtonResiduals(const ParticleType& Particles, const T Dt, const bool Write2File, TArray<PMatrix<T, 3, 3>>* AllParticleHessian)
+	{
+		TRACE_CPUPROFILER_EVENT_SCOPE(STAT_ChaosGaussSeidelComputeNewtonResidual);
+
+		TArray<TVec3<T>> NewtonResidual;
+		NewtonResidual.Init(TVec3<T>((T)0.), xtilde.Num());
+
+		if (AllParticleHessian)
+		{
+			AllParticleHessian->Init(PMatrix<T, 3, 3>((T)0.), Particles.Size());
+		}
+
+		for (int32 k = 0; k < ParticlesPerColor.Num(); k++)
+		{
+			PhysicsParallelFor(ParticlesPerColor[k].Num(), [this, &AllParticleHessian, Dt, k, &NewtonResidual, &Particles](const int32 j)
+				{
+					const int32 p = ParticlesPerColor[k][j];
+
+					if (Particles.InvM(p) != (T)0.)
+					{
+						int32 ConstraintIndex = 0;
+						Chaos::PMatrix<T, 3, 3> ParticleHessian((T)0., (T)0., (T)0.);
+
+						this->ComputeInitialResidualAndHessian(Particles, p, Dt, NewtonResidual[p], ParticleHessian);
+
+						for (int32 i = 0; i < StaticIncidentElements[p].Num(); i++)
+						{
+							while (StaticIncidentElements[p][i] >= StaticIncidentElementsOffsets[ConstraintIndex + 1] && ConstraintIndex < StaticIncidentElementsOffsets.Num() - 1)
+							{
+								ConstraintIndex += 1;
+							}
+
+							this->AddStaticConstraintResidualAndHessian[ConstraintIndex](Particles, StaticIncidentElements[p][i] - StaticIncidentElementsOffsets[ConstraintIndex], StaticIncidentElementsLocal[p][i], Dt, NewtonResidual[p], ParticleHessian);
+						}
+
+						ConstraintIndex = 0;
+
+						for (int32 i = 0; i < DynamicIncidentElements[p].Num(); i++)
+						{
+							while (DynamicIncidentElements[p][i] >= DynamicIncidentElementsOffsets[ConstraintIndex + 1] && ConstraintIndex < DynamicIncidentElementsOffsets.Num() - 1)
+							{
+								ConstraintIndex += 1;
+							}
+
+							this->AddDynamicConstraintResidualAndHessian[ConstraintIndex](Particles, DynamicIncidentElements[p][i] - DynamicIncidentElementsOffsets[ConstraintIndex], DynamicIncidentElementsLocal[p][i], Dt, NewtonResidual[p], ParticleHessian);
+						}
+
+						if (AllParticleHessian)
+						{
+							(*AllParticleHessian)[p] = ParticleHessian;
+						}
+					}
+			}, ParticlesPerColor[k].Num() < 1000);
+		}
+
+		T NewtonNorm = (T)0.;
+		for (int32 p = 0; p < NewtonResidual.Num(); p++)
+		{
+			NewtonNorm += NewtonResidual[p].SizeSquared();
+		}
+
+		NewtonNorm = FMath::Sqrt(NewtonNorm);
+
+		UE_LOG(LogTemp, Warning, TEXT("Current Iteration: %d"), CurrentIt);
+
+		UE_LOG(LogTemp, Warning, TEXT("Newton Residual is %f"), NewtonNorm);
+
+		PassedIters++;
+
+
+		if (Write2File)
+		{
+			FString file = FPaths::ProjectDir();
+			file.Append(TEXT("/DebugOutput/NewtonResidual.txt"));
+			if (PassedIters == 0)
+			{
+				FFileHelper::SaveStringToFile(FString(TEXT("Newton Norm\r\n")), *file);
+			}
+			FFileHelper::SaveStringToFile((FString::SanitizeFloat(NewtonNorm) + FString(TEXT(",\r\n"))), *file, FFileHelper::EEncodingOptions::AutoDetect, &IFileManager::Get(), EFileWrite::FILEWRITE_Append);
+		}
+
+		return NewtonResidual;
+
+	}
+
+	template <typename T, typename ParticleType>
+	void FGaussSeidelMainConstraint<T, ParticleType>::ApplyCG(ParticleType& Particles, const T Dt)
+	{
+		TRACE_CPUPROFILER_EVENT_SCOPE(STAT_ChaosGaussSeidelApplyCG);
+		TFunction<void(TArray<FSolverVec3>&)> ProjectBCs = [&Particles](TArray<FSolverVec3>& y)
+		{
+			for (int32 i = 0; i < y.Num(); i++)
+			{
+				if (Particles.InvM(i) == 0.f)
+				{
+					y[i] = FSolverVec3(0.f);
+				}
+			}
+		};
+
+		int32 max_it_cg = MaxItCG;
+		FSolverReal cg_tol = CGTol;
+		int32 MaxNewtonIt = 1;
+
+		for (int32 It = 0; It < MaxNewtonIt; It++)
+		{
+			const TArray<FSolverVec3> Residual = ComputeNewtonResiduals(Particles, Dt, false, nullptr);
+
+			auto multiply = [this, &ProjectBCs, &Dt, &Particles](TArray<FSolverVec3>& y, const TArray<FSolverVec3>& x) 
+			{
+				TArray<FSolverVec3> XProj = x;
+				ProjectBCs(XProj);
+
+				y.Init(FSolverVec3((FSolverReal)0.), XProj.Num());
+
+				for (int32 i = 0; i < this->AddInternalForceDifferentials.Num(); i++)
+				{
+					this->AddInternalForceDifferentials[i](Particles, XProj, y);
+				}
+
+				PhysicsParallelFor(y.Num(), 
+					[&y, &Particles, &XProj, Dt](const int32 i) 
+					{
+						y[i] = Particles.M(i) * XProj[i] + Dt * Dt * y[i];
+					});
+
+				ProjectBCs(y);
+			};
+
+			TArray<FSolverVec3> Deltax;
+			Deltax.Init(FSolverVec3(0.f), Particles.Size());
+
+			Chaos::LanczosCG<FSolverReal>(multiply, Deltax, Residual, MaxItCG, CGTol, use_list.Get()); 
+
+			PhysicsParallelFor(Particles.Size(), [&Particles, &Deltax](const int32 i) 
+				{
+					Particles.P(i) -= Deltax[i];
+				});
+		}
+	}
+}
+
+
+template class Chaos::Softs::FGaussSeidelMainConstraint<Chaos::Softs::FSolverReal, Chaos::Softs::FSolverParticles>;
+
+template CHAOS_API void Chaos::Softs::FGaussSeidelMainConstraint<Chaos::FRealDouble, Chaos::TDynamicParticles<Chaos::FRealDouble, 3>>::AddStaticConstraints(const TArray<TArray<int32>>& ExtraConstraints, TArray<TArray<int32>>& ExtraIncidentElements, TArray<TArray<int32>>& ExtraIncidentElementsLocal);
+
+template CHAOS_API void Chaos::Softs::FGaussSeidelMainConstraint<Chaos::FRealDouble, Chaos::TDynamicParticles<Chaos::FRealDouble, 3>>::AddDynamicConstraints(const TArray<TArray<int32>>& ExtraConstraints, TArray<TArray<int32>>& ExtraIncidentElements, TArray<TArray<int32>>& ExtraIncidentElementsLocal, bool CheckIncidentElements);
diff --git a/Engine/Source/Runtime/Experimental/Chaos/Private/Chaos/Deformable/GaussSeidelMasterConstraint.cpp b/Engine/Source/Runtime/Experimental/Chaos/Private/Chaos/Deformable/GaussSeidelMasterConstraint.cpp
deleted file mode 100644
diff --git a/Engine/Source/Runtime/Experimental/Chaos/Public/Chaos/Deformable/ChaosDeformableSolver.h b/Engine/Source/Runtime/Experimental/Chaos/Public/Chaos/Deformable/ChaosDeformableSolver.h
--- a/Engine/Source/Runtime/Experimental/Chaos/Public/Chaos/Deformable/ChaosDeformableSolver.h
+++ b/Engine/Source/Runtime/Experimental/Chaos/Public/Chaos/Deformable/ChaosDeformableSolver.h
@@ -14,7 +14,7 @@
 #include "Chaos/Deformable/GaussSeidelCorotatedConstraints.h"
 #include "Chaos/Deformable/GaussSeidelNeohookeanConstraints.h"
 #include "Chaos/Deformable/GaussSeidelWeakConstraints.h"
-#include "Chaos/Deformable/GaussSeidelMasterConstraint.h"
+#include "Chaos/Deformable/GaussSeidelMainConstraint.h"
 #include "Chaos/XPBDWeakConstraints.h"
 #include "Chaos/BlendedXPBDCorotatedConstraints.h"
 #include "Chaos/XPBDGridBasedCorotatedConstraints.h"
@@ -184,7 +184,7 @@ namespace Chaos::Softs
 		TArray<TUniquePtr<Softs::FXPBDWeakConstraints<Softs::FSolverReal, Softs::FSolverParticles>>> WeakConstraints;
 		TArray<TUniquePtr<Softs::FBlendedXPBDCorotatedConstraints<Softs::FSolverReal, Softs::FSolverParticles>>> BlendedCorotatedConstraints;
 		TUniquePtr<Softs::FXPBDGridBasedCorotatedConstraints<Softs::FSolverReal, Softs::FSolverParticles>> GridBasedCorotatedConstraint;
-		TUniquePtr<Softs::FGaussSeidelMasterConstraint<Softs::FSolverReal, Softs::FSolverParticles>> GSMasterConstraint;
+		TUniquePtr<Softs::FGaussSeidelMainConstraint<Softs::FSolverReal, Softs::FSolverParticles>> GSMainConstraint;
 		TUniquePtr<Softs::FPBDCollisionSpringConstraints> CollisionSpringConstraint;
 		TUniquePtr<Softs::FPBDTriangleMeshCollisions> TriangleMeshCollisions;
 		TArrayCollectionArray<const UObject*> MObjects;
diff --git a/Engine/Source/Runtime/Experimental/Chaos/Public/Chaos/Deformable/GaussSeidelMainConstraint.h b/Engine/Source/Runtime/Experimental/Chaos/Public/Chaos/Deformable/GaussSeidelMainConstraint.h
new file mode 100644
--- /dev/null
+++ b/Engine/Source/Runtime/Experimental/Chaos/Public/Chaos/Deformable/GaussSeidelMainConstraint.h
@@ -0,0 +1,354 @@
+// Copyright Epic Games, Inc. All Rights Reserved.
+#pragma once
+
+#include "Chaos/PBDSpringConstraintsBase.h"
+#include "Chaos/XPBDCorotatedConstraints.h"
+#include "ChaosStats.h"
+#include "Chaos/ImplicitQRSVD.h"
+#include "Chaos/GraphColoring.h"
+#include "Chaos/NewtonCorotatedCache.h"
+#include "Chaos/Framework/Parallel.h"
+#include "Chaos/Deformable/GaussSeidelWeakConstraints.h"
+#include "Chaos/Deformable/GaussSeidelCorotatedConstraints.h"
+
+
+#define PERF_SCOPE(X) SCOPE_CYCLE_COUNTER(X); TRACE_CPUPROFILER_EVENT_SCOPE(X);
+DECLARE_CYCLE_STAT(TEXT("Chaos.Deformable.GSMainConstraint.Apply"), STAT_ChaosGSMainConstraint_Apply, STATGROUP_Chaos);
+DECLARE_CYCLE_STAT(TEXT("Chaos.Deformable.GSMainConstraint.Acceleration"), STAT_ChaosGSMainConstraint_Acceleration, STATGROUP_Chaos);
+DECLARE_CYCLE_STAT(TEXT("Chaos.Deformable.GSMainConstraint.Init"), STAT_ChaosGSMainConstraint_Init, STATGROUP_Chaos);
+DECLARE_CYCLE_STAT(TEXT("Chaos.Deformable.GSMainConstraint.InitDynamicColor"), STAT_ChaosGSMainConstraint_InitDynamicColor, STATGROUP_Chaos);
+
+namespace Chaos::Softs
+{
+	template <typename T, typename ParticleType>
+	class FGaussSeidelMainConstraint
+	{
+
+	public:
+		FGaussSeidelMainConstraint(
+			const ParticleType& InParticles, 
+			const bool bDoQuasistaticsIn = false,
+			const bool bDoSORIn = true,
+			const T InOmegaSOR = (T)1.6,
+			const int32 ParallelMaxIn = 1000, 
+			const FDeformableXPBDCorotatedParams& InParams = FDeformableXPBDCorotatedParams())
+			: bDoQuasistatics(bDoQuasistaticsIn)
+			, bDoAcceleration(bDoSORIn)
+			, OmegaSOR(InOmegaSOR)
+			, ParallelMax(ParallelMaxIn)
+			, CorotatedParams(InParams)
+		{
+			Resize((int32)InParticles.Size());
+
+			InitializeLambdas();
+		}
+
+		virtual ~FGaussSeidelMainConstraint() {}
+
+		void Resize(const int32 NewSize)
+		{
+			xtilde.SetNum(NewSize);
+			StaticIncidentElements.SetNum(NewSize);
+			StaticIncidentElementsLocal.SetNum(NewSize);
+			DynamicIncidentElements.SetNum(NewSize);
+			DynamicIncidentElementsLocal.SetNum(NewSize);
+			X_k_1.Init(TVector<T, 3>(T(0.)), NewSize);
+			X_k.Init(TVector<T, 3>(T(0.)), NewSize);
+		}
+
+		const TArray<TFunction<void(const ParticleType&, const int32, const int32, const T, TVec3<T>&, Chaos::PMatrix<T, 3, 3>&)>>& StaticConstraintResidualAndHessian() const { return AddStaticConstraintResidualAndHessian; }
+		TArray<TFunction<void(const ParticleType&, const int32, const int32, const T, TVec3<T>&, Chaos::PMatrix<T, 3, 3>&)>>& StaticConstraintResidualAndHessian() { return AddStaticConstraintResidualAndHessian; }
+		const TArray<TFunction<void(const ParticleType&, const int32, const int32, const T, TVec3<T>&, Chaos::PMatrix<T, 3, 3>&)>>& DynamicConstraintResidualAndHessian() const { return AddDynamicConstraintResidualAndHessian; }
+		TArray<TFunction<void(const ParticleType&, const int32, const int32, const T, TVec3<T>&, Chaos::PMatrix<T, 3, 3>&)>>& DynamicConstraintResidualAndHessian() { return AddDynamicConstraintResidualAndHessian; }
+		const TArray<TFunction<void(const int32, const T, Chaos::PMatrix<T, 3, 3>&)>>& PerNodeHessian() const { return AddPerNodeHessian; }
+		TArray<TFunction<void(const int32, const T, Chaos::PMatrix<T, 3, 3>&)>>& PerNodeHessian() { return AddPerNodeHessian; }
+		const TArray<TFunction<void(const ParticleType&, const TArray<TVec3<T>>&, TArray<TVec3<T>>&)>>& InternalForceDifferentials() const { return AddInternalForceDifferentials; }
+		TArray<TFunction<void(const ParticleType&, const TArray<TVec3<T>>&, TArray<TVec3<T>>&)>>& InternalForceDifferentials() { return AddInternalForceDifferentials; }
+
+		int32 AddStaticConstraintResidualAndHessianRange(int32 NumConstraints)
+		{
+			int32 CurrentSize = AddStaticConstraintResidualAndHessian.Num();
+			AddStaticConstraintResidualAndHessian.AddDefaulted(NumConstraints);
+			return CurrentSize;
+		}
+
+		int32 AddDynamicConstraintResidualAndHessianRange(int32 NumConstraints)
+		{
+			int32 CurrentSize = AddDynamicConstraintResidualAndHessian.Num();
+			AddDynamicConstraintResidualAndHessian.AddDefaulted(NumConstraints);
+			return CurrentSize;
+		}
+
+		int32 AddPerNodeHessianRange(int32 NumConstraints)
+		{
+			int32 CurrentSize = AddPerNodeHessian.Num();
+			AddPerNodeHessian.AddDefaulted(NumConstraints);
+			return CurrentSize;
+		}
+
+		int32 AddAddInternalForceDifferentialsRange(int32 NumConstraints)
+		{
+			int32 CurrentSize = AddInternalForceDifferentials.Num();
+			AddInternalForceDifferentials.AddDefaulted(NumConstraints);
+			return CurrentSize;
+		}
+
+		CHAOS_API void AddStaticConstraints(const TArray<TArray<int32>>& ExtraConstraints, TArray<TArray<int32>>& ExtraIncidentElements, TArray<TArray<int32>>& ExtraIncidentElementsLocal);
+
+		CHAOS_API void AddDynamicConstraints(const TArray<TArray<int32>>& ExtraConstraints, TArray<TArray<int32>>& ExtraIncidentElements, TArray<TArray<int32>>& ExtraIncidentElementsLocal, bool CheckIncidentElements = false);
+
+		void Apply(ParticleType& Particles, const T Dt, const int32 MaxWriteIters = 10, const bool Write2File = false)
+		{
+ 			PERF_SCOPE(STAT_ChaosGSMainConstraint_Apply);
+			
+			if (DebugResidual && PassedIters < MaxWriteIters)
+			{
+				ComputeNewtonResiduals(Particles, Dt, Write2File);
+			}
+			
+			
+			for (int32 i = 0; i < ParticlesPerColor.Num(); i++)
+			{
+				int32 NumBatch = ParticlesPerColor[i].Num() / CorotatedParams.XPBDCorotatedBatchSize;
+				if (ParticlesPerColor[i].Num() % CorotatedParams.XPBDCorotatedBatchSize != 0)
+				{
+					NumBatch ++;
+				}
+			
+				PhysicsParallelFor(NumBatch, [&](const int32 BatchIndex)
+					{
+						for (int32 BatchSubIndex = 0; BatchSubIndex < CorotatedParams.XPBDCorotatedBatchSize; BatchSubIndex++) {
+							int32 TaskIndex = CorotatedParams.XPBDCorotatedBatchSize * BatchIndex + BatchSubIndex;
+							if (TaskIndex < ParticlesPerColor[i].Num())
+							{
+								int32 ParticleIndex = ParticlesPerColor[i][TaskIndex];
+								ApplySingleParticle(ParticleIndex, Dt, Particles);
+							}
+						}
+					}, NumBatch < CorotatedParams.XPBDCorotatedBatchThreshold);
+			
+			}
+			
+						
+			if (bDoAcceleration)
+			{
+				AccelerationTechnique(Particles);
+			}
+			
+			CurrentIt ++;
+		}
+
+		void InitStaticColor(const ParticleType& Particles)
+		{
+			StaticParticlesPerColor = ComputeNodalColoring(StaticConstraints, Particles, 0, Particles.Size(), StaticIncidentElements, StaticIncidentElementsLocal, &StaticParticleColors);
+			ParticleColors = StaticParticleColors;
+			ParticlesPerColor = StaticParticlesPerColor;
+		}
+
+		void InitDynamicColor(const ParticleType& Particles)
+		{
+			PERF_SCOPE(STAT_ChaosGSMainConstraint_InitDynamicColor);
+			ParticleColors = StaticParticleColors;
+			ParticlesPerColor = StaticParticlesPerColor;
+			Chaos::ComputeExtraNodalColoring(StaticConstraints, DynamicConstraints, Particles, StaticIncidentElements, DynamicIncidentElements, ParticleColors, ParticlesPerColor);
+		}
+
+		void Init(const T Dt, const ParticleType& Particles)
+		{
+			Resize((int32)Particles.Size());
+
+			PERF_SCOPE(STAT_ChaosGSMainConstraint_Init);
+			DynamicConstraints.SetNum(0);
+			for (int32 p = 0; p < DynamicIncidentElements.Num(); p++)
+			{
+				DynamicIncidentElements[p].SetNum(0);
+				DynamicIncidentElementsLocal[p].SetNum(0);
+			}
+			DynamicIncidentElementsOffsets.SetNum(0);
+			if (!bDoQuasistatics)
+			{
+				for (int32 i = 0; i < xtilde.Num(); i++)
+				{
+					//xtilde[i] = Particles.X(i) + Dt * Particles.V(i);
+					xtilde[i] = Particles.P(i);
+				}
+			}
+			CurrentIt = 0;
+		}
+
+		static bool IsClean(const TArray<TArray<int32>>& ConstraintsIn, const TArray<TArray<int32>>& IncidentElementsIn, const TArray<TArray<int32>>& IncidentElementsLocalIn)
+		{
+			if (IncidentElementsIn.Num() == IncidentElementsLocalIn.Num())
+			{
+				int32 TotalEntries = 0;
+				for (int32 i = 0; i < IncidentElementsIn.Num(); i++)
+				{
+					TotalEntries += IncidentElementsIn[i].Num();
+					if (IncidentElementsIn[i].Num() != IncidentElementsLocalIn[i].Num())
+					{
+						return false;
+					}
+					for (int32 j = 0; j < IncidentElementsIn[i].Num(); j++)
+					{
+						if (IncidentElementsIn[i][j] > ConstraintsIn.Num() || IncidentElementsLocalIn[i][j] > ConstraintsIn[IncidentElementsIn[i][j]].Num())
+						{
+							return false;
+						}
+					}
+				}
+				if (TotalEntries > 0)
+				{
+					return true;
+				}
+			}
+			return false;
+		}
+
+		CHAOS_API TArray<TVec3<T>> ComputeNewtonResiduals(const ParticleType& Particles, const T Dt, const bool Write2File = false, TArray<PMatrix<T, 3, 3>>* AllParticleHessian = nullptr);
+		
+		CHAOS_API void ApplyCG(ParticleType& Particles, const T Dt);
+
+	private:
+
+		void ApplySingleParticle(const int32 p, const T Dt, ParticleType& Particles)
+		{
+			int32 ConstraintIndex = 0;
+			Chaos::TVector<T, 3> ParticleResidual((T)0.);
+			Chaos::PMatrix<T, 3, 3> ParticleHessian((T)0., (T)0., (T)0.);
+
+			ComputeInitialResidualAndHessian(Particles, p, Dt, ParticleResidual, ParticleHessian);
+
+			for (int32 i = 0; i < StaticIncidentElements[p].Num(); i++)
+			{
+				while (StaticIncidentElements[p][i] >= StaticIncidentElementsOffsets[ConstraintIndex + 1] && ConstraintIndex < StaticIncidentElementsOffsets.Num() - 1)
+				{
+					ConstraintIndex ++;
+				}
+
+				AddStaticConstraintResidualAndHessian[ConstraintIndex](Particles, StaticIncidentElements[p][i] - StaticIncidentElementsOffsets[ConstraintIndex], StaticIncidentElementsLocal[p][i], Dt, ParticleResidual, ParticleHessian);
+			}
+
+			ConstraintIndex = 0;
+
+			for (int32 i = 0; i < DynamicIncidentElements[p].Num(); i++)
+			{
+				while (DynamicIncidentElements[p][i] >= DynamicIncidentElementsOffsets[ConstraintIndex + 1] && ConstraintIndex < DynamicIncidentElementsOffsets.Num() - 1)
+				{
+					ConstraintIndex ++;
+				}
+
+				AddDynamicConstraintResidualAndHessian[ConstraintIndex](Particles, DynamicIncidentElements[p][i] - DynamicIncidentElementsOffsets[ConstraintIndex], DynamicIncidentElementsLocal[p][i], Dt, ParticleResidual, ParticleHessian);
+			}
+
+			for (int32 i = 0; i < AddPerNodeHessian.Num(); i++)
+			{
+				AddPerNodeHessian[i](p, Dt, ParticleHessian);
+			}
+
+
+			T HessianDet = ParticleHessian.Determinant();
+			if (HessianDet > UE_SMALL_NUMBER && ParticleResidual.Size() > UE_SMALL_NUMBER)
+			{
+				Chaos::PMatrix<T, 3, 3> HessianInv = ParticleHessian.SymmetricCofactorMatrix();
+				HessianInv *= T(1) / HessianDet;
+				Chaos::TVector<T, 3> Dx = HessianInv.GetTransposed() * (-ParticleResidual);
+
+				Particles.P(p) += Dx;
+			}
+		}
+		
+		void InitializeLambdas()
+		{
+			ComputeInitialResidualAndHessian = [this](const ParticleType& Particles, const int32 p, const T Dt, TVec3<T>& ParticleResidual, Chaos::PMatrix<T, 3, 3>& ParticleHessian)
+			{
+				if (!this->bDoQuasistatics) {
+					for (int32 alpha = 0; alpha < 3; alpha++) {
+						ParticleResidual[alpha] = Particles.M(p) * (Particles.P(p)[alpha] - this->xtilde[p][alpha]);
+					}
+					for (int32 alpha = 0; alpha < 3; alpha++) {
+						ParticleHessian.SetAt(alpha, alpha, Particles.M(p));
+					}
+				}
+				else
+				{
+					for (int32 alpha = 0; alpha < 3; alpha++) {
+						ParticleResidual[alpha] = -Dt * Dt * ExternalForce[alpha] * Particles.M(p);
+					}
+				}
+			};
+
+			AccelerationTechnique = [this](ParticleType& Particles)
+			{
+				PERF_SCOPE(STAT_ChaosGSMainConstraint_Acceleration);
+				PhysicsParallelFor(Particles.Size(), [&](const int32 ParticleIndex)
+					{
+						if (Particles.InvM(ParticleIndex) != T(0) && CurrentIt > SORStart)
+						{
+							Particles.P(ParticleIndex) = OmegaSOR * (Particles.P(ParticleIndex) - this->X_k_1[ParticleIndex]) + this->X_k_1[ParticleIndex];
+						}
+						this->X_k_1[ParticleIndex] = this->X_k[ParticleIndex];
+						this->X_k[ParticleIndex] = Particles.P(ParticleIndex);
+					}, Particles.Size() < 1000);
+			};
+		}
+
+
+		//Constraints storage:
+		TArray<TArray<int32>> StaticConstraints;
+		TArray<TArray<int32>> StaticIncidentElements;
+		TArray<TArray<int32>> StaticIncidentElementsLocal;
+		TArray<TArray<int32>> DynamicConstraints;
+		TArray<TArray<int32>> DynamicIncidentElements;
+		TArray<TArray<int32>> DynamicIncidentElementsLocal;
+
+		//Lambds for spcifying residual/hessian computations:
+		TFunction<void(const ParticleType&, const int32, const T, TVec3<T>&, Chaos::PMatrix<T, 3, 3>&)> ComputeInitialResidualAndHessian;
+		TArray<TFunction<void(const ParticleType&, const int32, const int32, const T, TVec3<T>&, Chaos::PMatrix<T, 3, 3>&)>> AddStaticConstraintResidualAndHessian;
+		TArray<TFunction<void(const ParticleType&, const int32, const int32, const T, TVec3<T>&, Chaos::PMatrix<T, 3, 3>&)>> AddDynamicConstraintResidualAndHessian;
+		TArray<TFunction<void(const int32, const T, Chaos::PMatrix<T, 3, 3>&)>> AddPerNodeHessian;
+
+		//Coloring information:
+		TArray<int32> StaticParticleColors;
+		TArray<TArray<int32>> StaticParticlesPerColor;
+
+		TArray<int32> ParticleColors;
+		TArray<TArray<int32>> ParticlesPerColor;
+
+		TArray<int32> StaticIncidentElementsOffsets;
+		TArray<int32> DynamicIncidentElementsOffsets;
+
+		bool bDoQuasistatics = false;
+		mutable TArray<Chaos::TVector<T, 3>> xtilde;
+
+		//SOR variables:
+		TFunction<void(ParticleType&)> AccelerationTechnique;
+		mutable TArray<Chaos::TVector<T, 3>> X_k_1;
+		mutable TArray<Chaos::TVector<T, 3>> X_k;
+		int32 CurrentIt = 0;
+		bool bDoAcceleration = true;
+		T OmegaSOR = T(1.6);
+		int32 SORStart = 1; //1 is smallest
+
+		int32 ParallelMax = 1000;
+
+		FDeformableXPBDCorotatedParams CorotatedParams;
+
+		//Newton solver variables:
+		TArray<TFunction<void(const ParticleType&, const TArray<TVec3<T>>&, TArray<TVec3<T>>&)>> AddInternalForceDifferentials;
+		TUniquePtr<TArray<int32>> use_list;
+
+		int32 NumTotalParticles;
+		TArray<TVec3<T>> ReorderedPs;
+
+		public:
+		bool DebugResidual = false;
+		bool IsFirstFrame = true;
+		int32 PassedIters = 0;
+
+		TVec3<T> ExternalForce = TVec3<T>((T)0.);
+
+	};
+
+	//template class FGaussSeidelMainConstraint<FSolverReal, FSolverParticles>;
+}
+
diff --git a/Engine/Source/Runtime/Experimental/Chaos/Public/Chaos/Deformable/GaussSeidelMasterConstraint.h b/Engine/Source/Runtime/Experimental/Chaos/Public/Chaos/Deformable/GaussSeidelMasterConstraint.h
deleted file mode 100644
