ChaosClothAsset: use SoftsEvolution by default. SKM cloth will continue to UseLegacySolver (PBDEvolution) for now.

Removed all of the ForceBasedSolver attributes from the SolverConfig node. We're currently just using the SoftsEvolution as a PBD Solver. The force based stuff is still not fully implemented. SoftsEvolution still knows about the ForceBasedSolver parameters; there is just no UI to enable it.

diff --git a/Engine/Plugins/ChaosCloth/Source/ChaosCloth/Private/ChaosCloth/ChaosClothVisualization.cpp b/Engine/Plugins/ChaosCloth/Source/ChaosCloth/Private/ChaosCloth/ChaosClothVisualization.cpp
--- a/Engine/Plugins/ChaosCloth/Source/ChaosCloth/Private/ChaosCloth/ChaosClothVisualization.cpp
+++ b/Engine/Plugins/ChaosCloth/Source/ChaosCloth/Private/ChaosCloth/ChaosClothVisualization.cpp
@@ -148,8 +148,8 @@ FLinearColor PseudoRandomColor(int32 NumColorRotations)
 			{
 				continue;
 			}
-			// Elements are local indexed for force based solver
-			const int32 Offset = Solver->IsForceBasedSolver() ? 0 : ParticleRangeId;
+			// Elements are local indexed for new solver
+			const int32 Offset = Solver->IsLegacySolver() ? ParticleRangeId : 0;
 
 			const TConstArrayView<TVec3<int32>> Elements = Cloth->GetTriangleMesh(Solver).GetElements();
 			const TConstArrayView<Softs::FSolverVec3> Positions = Cloth->GetParticlePositions(Solver);
@@ -249,8 +249,8 @@ FLinearColor PseudoRandomColor(int32 NumColorRotations)
 			{
 				continue;
 			}
-			// Elements are local indexed for force based solver
-			const int32 Offset = Solver->IsForceBasedSolver() ? 0 : ParticleRangeId;
+			// Elements are local indexed for new solver
+			const int32 Offset = Solver->IsLegacySolver() ? ParticleRangeId : 0;
 
 			const TArray<TVec3<int32>>& Elements = Cloth->GetTriangleMesh(Solver).GetElements();
 			const TConstArrayView<Softs::FSolverVec3> Positions = Cloth->GetParticlePositions(Solver);
@@ -342,8 +342,8 @@ FLinearColor PseudoRandomColor(int32 NumColorRotations)
 			{
 				continue;
 			}
-			// Elements are local indexed for force based solver
-			const int32 Offset = Solver->IsForceBasedSolver() ? 0 : ParticleRangeId;
+			// Elements are local indexed for new solver
+			const int32 Offset = Solver->IsLegacySolver() ? ParticleRangeId : 0;
 			const TConstArrayView<TVec3<int32>> Elements = Cloth->GetTriangleMesh(Solver).GetElements();
 			const TConstArrayView<Softs::FSolverVec3> Positions = Cloth->GetParticlePositions(Solver);
 			const TConstArrayView<FRealSingle>& WeightMap = Cloth->GetWeightMapByName(Solver, Name);
@@ -414,8 +414,8 @@ FLinearColor PseudoRandomColor(int32 NumColorRotations)
 			{
 				continue;
 			}
-			// Elements are local indexed for force based solver
-			const int32 Offset = Solver->IsForceBasedSolver() ? 0 : ParticleRangeId;
+			// Elements are local indexed for new solver
+			const int32 Offset = Solver->IsLegacySolver() ? ParticleRangeId : 0;
 			const TConstArrayView<TVec3<int32>> Elements = Cloth->GetTriangleMesh(Solver).GetElements();
 			const TConstArrayView<Softs::FSolverVec3> Positions = Cloth->GetParticlePositions(Solver);
 			const TConstArrayView<int32>& WeightMap = Cloth->GetFaceIntMapByProperty(Solver, TEXT("SelfCollisionLayers"));
@@ -825,7 +825,7 @@ FLinearColor PseudoRandomColor(int32 NumColorRotations)
 			const FClothConstraints& ClothConstraints = Solver->GetClothConstraints(ParticleRangeId);
 			if (const Softs::FExternalForces* const ExternalForces = ClothConstraints.GetExternalForces().Get())
 			{
-				check(Solver->IsForceBasedSolver());
+				check(!Solver->IsLegacySolver());
 				if (ExternalForces->HasPerParticleGravity())
 				{
 					const FVec3& LocalSpaceLocation = Solver->GetLocalSpaceLocation();
@@ -875,7 +875,7 @@ FLinearColor PseudoRandomColor(int32 NumColorRotations)
 			const FClothConstraints& ClothConstraints = Solver->GetClothConstraints(ParticleRangeId);
 			if (const Softs::FExternalForces* const ExternalForces = ClothConstraints.GetExternalForces().Get())
 			{
-				check(Solver->IsForceBasedSolver());
+				check(!Solver->IsLegacySolver());
 				{
 					const FVec3& LocalSpaceLocation = Solver->GetLocalSpaceLocation();
 					const TConstArrayView<Softs::FSolverVec3> Positions = Solver->GetParticleXsView(ParticleRangeId);
@@ -914,8 +914,8 @@ FLinearColor PseudoRandomColor(int32 NumColorRotations)
 			{
 				continue;
 			}
-			// Elements are local indexed for force based solver
-			const int32 Offset = Solver->IsForceBasedSolver() ? 0 : ParticleRangeId;
+			// Elements are local indexed for new solver
+			const int32 Offset = Solver->IsLegacySolver() ? ParticleRangeId : 0;
 			const TConstArrayView<TVec3<int32>> Elements = Cloth->GetTriangleMesh(Solver).GetElements();
 			const TConstArrayView<Softs::FSolverVec3> Positions = Cloth->GetParticlePositions(Solver);
 			const TConstArrayView<Softs::FSolverReal> InvMasses = Cloth->GetParticleInvMasses(Solver);
@@ -958,8 +958,8 @@ FLinearColor PseudoRandomColor(int32 NumColorRotations)
 			{
 				continue;
 			}
-			// Elements are local indexed for force based solver
-			const int32 Offset = Solver->IsForceBasedSolver() ? 0 : ParticleRangeId;
+			// Elements are local indexed for new solver
+			const int32 Offset = Solver->IsLegacySolver() ? ParticleRangeId : 0;
 
 			const TConstArrayView<TVec3<int32>> Elements = Cloth->GetTriangleMesh(Solver).GetElements();
 			const TConstArrayView<Softs::FSolverVec3> Positions = Cloth->GetAnimationPositions(Solver);
@@ -1024,8 +1024,8 @@ FLinearColor PseudoRandomColor(int32 NumColorRotations)
 			{
 				continue;
 			}
-			// Elements are local indexed for force based solver
-			const int32 Offset = Solver->IsForceBasedSolver() ? 0 : ParticleRangeId;
+			// Elements are local indexed for new solver
+			const int32 Offset = Solver->IsLegacySolver() ? ParticleRangeId : 0;
 
 			const TConstArrayView<TVec3<int32>> Elements = Cloth->GetTriangleMesh(Solver).GetElements();
 			const TConstArrayView<Softs::FSolverVec3> Positions = Cloth->GetAnimationPositions(Solver);
@@ -1911,9 +1911,8 @@ FLinearColor PseudoRandomColor(int32 NumColorRotations)
 
 			const FClothConstraints& ClothConstraints = Solver->GetClothConstraints(ParticleRangeId);
 
-			// Constraints are locally indexed for force based solver
-			const TConstArrayView<Softs::FSolverVec3> Positions = Solver->IsForceBasedSolver() ? Solver->GetParticleXsView(ParticleRangeId)
-				: TConstArrayView<Softs::FSolverVec3>(Solver->GetParticleXs());
+			// Constraints are locally indexed for new solver
+			const TConstArrayView<Softs::FSolverVec3> Positions = Solver->IsLegacySolver() ? TConstArrayView<Softs::FSolverVec3>(Solver->GetParticleXs()) : Solver->GetParticleXsView(ParticleRangeId);
 
 			if (const Softs::FPBDEdgeSpringConstraints* const EdgeConstraints = ClothConstraints.GetEdgeSpringConstraints().Get())
 			{
@@ -2025,9 +2024,8 @@ FLinearColor PseudoRandomColor(int32 NumColorRotations)
 			const FClothConstraints& ClothConstraints = Solver->GetClothConstraints(ParticleRangeId);
 
 
-			// Constraints are locally indexed for force based solver
-			const TConstArrayView<Softs::FSolverVec3> Positions = Solver->IsForceBasedSolver() ? Solver->GetParticleXsView(ParticleRangeId)
-				: TConstArrayView<Softs::FSolverVec3>(Solver->GetParticleXs());
+			// Constraints are locally indexed for new solver
+			const TConstArrayView<Softs::FSolverVec3> Positions = Solver->IsLegacySolver() ? TConstArrayView<Softs::FSolverVec3>(Solver->GetParticleXs()) : Solver->GetParticleXsView(ParticleRangeId);
 
 			if (const Softs::FPBDBendingSpringConstraints* const BendingConstraints = ClothConstraints.GetBendingSpringConstraints().Get())
 			{
@@ -2199,7 +2197,7 @@ FLinearColor PseudoRandomColor(int32 NumColorRotations)
 			}
 
 			const Softs::FVelocityAndPressureField* VelocityField = nullptr;
-			if (Solver->IsForceBasedSolver())
+			if (!Solver->IsLegacySolver())
 			{
 				const FClothConstraints& ClothConstraints = Solver->GetClothConstraints(ParticleRangeId);
 				VelocityField = ClothConstraints.GetVelocityAndPressureField().Get();
@@ -2213,11 +2211,9 @@ FLinearColor PseudoRandomColor(int32 NumColorRotations)
 				VelocityField = &Solver->GetWindVelocityAndPressureField(Cloth->GetGroupId());
 			}
 
-			// Constraints are locally indexed for force based solver
-			const TConstArrayView<Softs::FSolverVec3> Positions = Solver->IsForceBasedSolver() ? Solver->GetParticleXsView(ParticleRangeId)
-				: TConstArrayView<Softs::FSolverVec3>(Solver->GetParticleXs());
-			const TConstArrayView<Softs::FSolverReal> InvMasses = Solver->IsForceBasedSolver() ? Solver->GetParticleInvMassesView(ParticleRangeId)
-				: TConstArrayView<Softs::FSolverReal>(Solver->GetParticleInvMasses());
+			// Constraints are locally indexed for new solver
+			const TConstArrayView<Softs::FSolverVec3> Positions = Solver->IsLegacySolver() ? TConstArrayView<Softs::FSolverVec3>(Solver->GetParticleXs()) : Solver->GetParticleXsView(ParticleRangeId);
+			const TConstArrayView<Softs::FSolverReal> InvMasses = Solver->IsLegacySolver() ? TConstArrayView<Softs::FSolverReal>(Solver->GetParticleInvMasses()) : Solver->GetParticleInvMassesView(ParticleRangeId);
 
 			const TConstArrayView<TVec3<int32>>& Elements = VelocityField->GetElements();
 			const TConstArrayView<Softs::FSolverVec3> Forces = VelocityField->GetForces();
@@ -2286,10 +2282,9 @@ FLinearColor PseudoRandomColor(int32 NumColorRotations)
 
 			if (const Softs::FPBDCollisionSpringConstraints* const SelfCollisionConstraints = ClothConstraints.GetSelfCollisionConstraints().Get())
 			{
-				// Constraints are locally indexed for force based solver
-				const TConstArrayView<Softs::FSolverVec3> Positions = Solver->IsForceBasedSolver() ? Solver->GetParticleXsView(ParticleRangeId)
-					: TConstArrayView<Softs::FSolverVec3>(Solver->GetParticleXs());
-				const int32 Offset = Solver->IsForceBasedSolver() ? 0 : ParticleRangeId;
+				// Constraints are locally indexed for new solver
+				const TConstArrayView<Softs::FSolverVec3> Positions = Solver->IsLegacySolver() ? TConstArrayView<Softs::FSolverVec3>(Solver->GetParticleXs()) : Solver->GetParticleXsView(ParticleRangeId);
+				const int32 Offset = Solver->IsLegacySolver() ? ParticleRangeId : 0;
 				const TArray<TVec4<int32>>& Constraints = SelfCollisionConstraints->GetConstraints();
 				const TArray<Softs::FSolverVec3>& Barys = SelfCollisionConstraints->GetBarys();
 				const TArray<bool>& FlipNormals = SelfCollisionConstraints->GetFlipNormals();
@@ -2359,10 +2354,9 @@ FLinearColor PseudoRandomColor(int32 NumColorRotations)
 			if (const Softs::FPBDSelfCollisionSphereConstraints* const SelfCollisionSphereConstraints =
 				ClothConstraints.GetSelfCollisionSphereConstraints().Get())
 			{
-				// Constraints are locally indexed for force based solver
-				const TConstArrayView<Softs::FSolverVec3> Positions = Solver->IsForceBasedSolver() ? Solver->GetParticleXsView(ParticleRangeId)
-					: TConstArrayView<Softs::FSolverVec3>(Solver->GetParticleXs());
-				const int32 Offset = Solver->IsForceBasedSolver() ? 0 : ParticleRangeId;
+				// Constraints are locally indexed for new solver
+				const TConstArrayView<Softs::FSolverVec3> Positions = Solver->IsLegacySolver() ? TConstArrayView<Softs::FSolverVec3>(Solver->GetParticleXs()) : Solver->GetParticleXsView(ParticleRangeId);
+				const int32 Offset = Solver->IsLegacySolver() ? ParticleRangeId : 0;
 				const TArray<TVec2<int32>>& Constraints = SelfCollisionSphereConstraints->GetConstraints();
 				for (int32 Index = 0; Index < Constraints.Num(); ++Index)
 				{
@@ -2417,11 +2411,9 @@ FLinearColor PseudoRandomColor(int32 NumColorRotations)
 
 			if (const Softs::FPBDTriangleMeshCollisions* const SelfCollisionInit = ClothConstraints.GetSelfCollisionInit().Get())
 			{
-				// Constraints are locally indexed for force based solver
-				const TConstArrayView<Softs::FSolverVec3> Positions = Solver->IsForceBasedSolver() ? Solver->GetParticleXsView(ParticleRangeId)
-					: TConstArrayView<Softs::FSolverVec3>(Solver->GetParticleXs());
-
-				const int32 Offset = Solver->IsForceBasedSolver() ? 0 : ParticleRangeId;
+				// Constraints are locally indexed for new solver
+				const TConstArrayView<Softs::FSolverVec3> Positions = Solver->IsLegacySolver() ? TConstArrayView<Softs::FSolverVec3>(Solver->GetParticleXs()) : Solver->GetParticleXsView(ParticleRangeId);
+				const int32 Offset = Solver->IsLegacySolver() ? ParticleRangeId : 0;
 
 				const FTriangleMesh& TriangleMesh = Cloth->GetTriangleMesh(Solver);
 
@@ -2557,7 +2549,7 @@ FLinearColor PseudoRandomColor(int32 NumColorRotations)
 
 				if (const Softs::FPBDCollisionSpringConstraints* const SelfCollisionConstraints = ClothConstraints.GetSelfCollisionConstraints().Get())
 				{
-					const int32 Offset = Solver->IsForceBasedSolver() ? 0 : ParticleRangeId;
+					const int32 Offset = Solver->IsLegacySolver() ? ParticleRangeId : 0;
 					const TConstArrayView<Softs::FSolverVec3> Positions = Cloth->GetParticlePositions(Solver);
 					const TConstArrayView<Softs::FSolverReal> InvMasses = Cloth->GetParticleInvMasses(Solver);
 					const TConstArrayView<int32>& WeightMap = Cloth->GetFaceIntMapByProperty(Solver, TEXT("SelfCollisionLayers"));
@@ -2607,7 +2599,7 @@ FLinearColor PseudoRandomColor(int32 NumColorRotations)
 
 				if (const Softs::FPBDCollisionSpringConstraints* const SelfCollisionConstraints = ClothConstraints.GetSelfCollisionConstraints().Get())
 				{
-					const int32 Offset = Solver->IsForceBasedSolver() ? 0 : ParticleRangeId;
+					const int32 Offset = Solver->IsLegacySolver() ? ParticleRangeId : 0;
 					const TConstArrayView<Softs::FSolverVec3> Positions = Cloth->GetParticlePositions(Solver);
 					const TConstArrayView<Softs::FSolverReal> InvMasses = Cloth->GetParticleInvMasses(Solver);
 
@@ -2669,8 +2661,8 @@ FLinearColor PseudoRandomColor(int32 NumColorRotations)
 			if (const Softs::FPBDTriangleMeshCollisions* const SelfCollisionInit = ClothConstraints.GetSelfCollisionInit().Get())
 			{
 				const FTriangleMesh& KinematicColliderMesh = SelfCollisionInit->GetCollidableSubMesh().GetKinematicColliderSubMesh();
-				// Elements are local indexed for force based solver
-				const int32 Offset = Solver->IsForceBasedSolver() ? 0 : ParticleRangeId;
+				// Elements are local indexed for new solver
+				const int32 Offset = Solver->IsLegacySolver() ? ParticleRangeId : 0;
 				const TConstArrayView<TVec3<int32>> Elements = KinematicColliderMesh.GetElements();
 				const TConstArrayView<Softs::FSolverVec3> Positions = Cloth->GetParticlePositions(Solver);
 				for (int32 ElementIndex = 0; ElementIndex < Elements.Num(); ++ElementIndex)
@@ -2712,8 +2704,8 @@ FLinearColor PseudoRandomColor(int32 NumColorRotations)
 			if (const Softs::FPBDTriangleMeshCollisions* const SelfCollisionInit = ClothConstraints.GetSelfCollisionInit().Get())
 			{
 				const FTriangleMesh& KinematicColliderMesh = SelfCollisionInit->GetCollidableSubMesh().GetKinematicColliderSubMesh();
-				// Elements are local indexed for force based solver
-				const int32 Offset = Solver->IsForceBasedSolver() ? 0 : ParticleRangeId;
+				// Elements are local indexed for new solver
+				const int32 Offset = Solver->IsLegacySolver() ? ParticleRangeId : 0;
 				const TConstArrayView<TVec3<int32>> Elements = KinematicColliderMesh.GetElements();
 				const TConstArrayView<Softs::FSolverVec3> Positions = Cloth->GetParticlePositions(Solver);
 				for (int32 ElementIndex = 0; ElementIndex < Elements.Num(); ++ElementIndex, VertexIndex += 3)
diff --git a/Engine/Plugins/ChaosCloth/Source/ChaosCloth/Private/ChaosCloth/ChaosClothingSimulationCloth.cpp b/Engine/Plugins/ChaosCloth/Source/ChaosCloth/Private/ChaosCloth/ChaosClothingSimulationCloth.cpp
--- a/Engine/Plugins/ChaosCloth/Source/ChaosCloth/Private/ChaosCloth/ChaosClothingSimulationCloth.cpp
+++ b/Engine/Plugins/ChaosCloth/Source/ChaosCloth/Private/ChaosCloth/ChaosClothingSimulationCloth.cpp
@@ -176,7 +176,7 @@ void FClothingSimulationCloth::FLODData::AddParticles(FClothingSimulationSolver*
 		return;
 	}
 
-	if (!Solver->IsForceBasedSolver())
+	if (Solver->IsLegacySolver())
 	{
 		const TArray<TVec3<int32>>& Elements = NoOffsetTriangleMesh.GetElements();
 		TArray<TVec3<int32>> OffsetElements;
@@ -193,7 +193,7 @@ void FClothingSimulationCloth::FLODData::AddParticles(FClothingSimulationSolver*
 		SolverDatum.OffsetTriangleMesh.GetPointToTriangleMap(); // Builds map for later use by GetPointNormals(), and the velocity fields
 	}
 
-	const FTriangleMesh& TriangleMesh = Solver->IsForceBasedSolver() ? NoOffsetTriangleMesh : SolverDatum.OffsetTriangleMesh;
+	const FTriangleMesh& TriangleMesh = Solver->IsLegacySolver() ? SolverDatum.OffsetTriangleMesh : NoOffsetTriangleMesh;
 
 	PRAGMA_DISABLE_DEPRECATION_WARNINGS  // TODO: CHAOS_IS_CLOTHINGSIMULATIONMESH_ABSTRACT
 	// Update source mesh for this LOD, this is required prior to reset the start pose
@@ -277,7 +277,7 @@ void FClothingSimulationCloth::FLODData::Add(FClothingSimulationSolver* Solver,
 	const Softs::FSolverReal MeshScale = Cloth->Mesh->GetScale();
 	PRAGMA_ENABLE_DEPRECATION_WARNINGS
 
-	const FTriangleMesh& TriangleMesh = Solver->IsForceBasedSolver() ? NoOffsetTriangleMesh : SolverDatum.OffsetTriangleMesh;
+	const FTriangleMesh& TriangleMesh = Solver->IsLegacySolver() ? SolverDatum.OffsetTriangleMesh : NoOffsetTriangleMesh;
 
 	// Retrieve config properties
 	check(Cloth->Config);
@@ -287,9 +287,9 @@ void FClothingSimulationCloth::FLODData::Add(FClothingSimulationSolver* Solver,
 	TSharedPtr<Softs::FMultiResConstraints> FineLODMultiResConstraint;
 	const FTriangleMesh* CoarseLODTriangleMesh = nullptr;
 	int32 CoarseLODParticleRangeId = INDEX_NONE;
-	if (Solver->IsForceBasedSolver())
+	if (!Solver->IsLegacySolver())
 	{
-		// Multi-Res is only supported with ForceBasedSolver
+		// Multi-Res isn't supported by legacy solver
 		if (InLODIndex == 0)
 		{
 			// Only allow LOD0 to be a fine LOD for now.
@@ -885,7 +885,7 @@ int32 FClothingSimulationCloth::GetParticleRangeId(const FClothingSimulationSolv
 TVec3<FRealSingle> FClothingSimulationCloth::GetGravity(const FClothingSimulationSolver* Solver) const
 {
 	check(Solver);
-	if (Solver->IsForceBasedSolver())
+	if (!Solver->IsLegacySolver())
 	{
 		const int32 ParticleRangeId = GetParticleRangeId(Solver);
 		if (const Softs::FExternalForces* const ExternalForces = Solver->GetClothConstraints(ParticleRangeId).GetExternalForces().Get())
@@ -945,7 +945,7 @@ const FTriangleMesh& FClothingSimulationCloth::GetTriangleMesh(const FClothingSi
 	const int32 LODIndex = LODIndices.FindChecked(Solver);
 	static const FTriangleMesh EmptyTriangleMesh;
 
-	return LODData.IsValidIndex(LODIndex) ? (Solver->IsForceBasedSolver() ? LODData[LODIndex]->NoOffsetTriangleMesh : LODData[LODIndex]->SolverData.FindChecked(Solver).OffsetTriangleMesh) : EmptyTriangleMesh;
+	return LODData.IsValidIndex(LODIndex) ? (Solver->IsLegacySolver() ? LODData[LODIndex]->SolverData.FindChecked(Solver).OffsetTriangleMesh: LODData[LODIndex]->NoOffsetTriangleMesh) : EmptyTriangleMesh;
 }
 
 
@@ -1114,7 +1114,7 @@ PRAGMA_ENABLE_DEPRECATION_WARNINGS
 				Solver->GetParticleXs(ParticleRangeId),
 				Solver->GetParticleVs(ParticleRangeId));
 
-			if(!Solver->IsForceBasedSolver())
+			if(Solver->IsLegacySolver())
 			{
 				// Update the wind velocity field for the new LOD mesh
 				Solver->SetWindAndPressureGeometry(GroupId, GetTriangleMesh(Solver), ConfigProperties, LODData[LODIndex]->WeightMaps);
@@ -1182,7 +1182,7 @@ PRAGMA_ENABLE_DEPRECATION_WARNINGS
 			OutAngularVelocityScale,
 			FictitiousAngularScale,
 			OutMaxVelocityScale);
-		if (Solver->IsForceBasedSolver())
+		if (!Solver->IsLegacySolver())
 		{
 			Solver->SetProperties(ParticleRangeId, ConfigProperties, LODData[LODIndex]->WeightMaps);
 			if (CoarseLODIndex != INDEX_NONE)
diff --git a/Engine/Plugins/ChaosCloth/Source/ChaosCloth/Private/ChaosCloth/ChaosClothingSimulationCollider.cpp b/Engine/Plugins/ChaosCloth/Source/ChaosCloth/Private/ChaosCloth/ChaosClothingSimulationCollider.cpp
--- a/Engine/Plugins/ChaosCloth/Source/ChaosCloth/Private/ChaosCloth/ChaosClothingSimulationCollider.cpp
+++ b/Engine/Plugins/ChaosCloth/Source/ChaosCloth/Private/ChaosCloth/ChaosClothingSimulationCollider.cpp
@@ -344,7 +344,7 @@ void FClothingSimulationCollider::FLODData::Add(
 			// Add bone proxies first so they get updated first?
 			TArray<int32> SolverBoneIndices;
 			SolverBoneIndices.Reserve(SkinnedCollisionData.MappedSkinnedBones.Num());
-			const int32 GlobalBoneOffset = Solver->IsForceBasedSolver() ? 0 : CollisionRangeId;
+			const int32 GlobalBoneOffset = Solver->IsLegacySolver() ? CollisionRangeId : 0;
 			for (int32 MappedSubBoneIndex : SkinnedCollisionData.MappedSkinnedBones)
 			{
 				// Always initialize the collision particle transforms before setting any geometry as otherwise NaNs gets detected during the bounding box updates
diff --git a/Engine/Plugins/ChaosCloth/Source/ChaosCloth/Private/ChaosCloth/ChaosClothingSimulationSolver.cpp b/Engine/Plugins/ChaosCloth/Source/ChaosCloth/Private/ChaosCloth/ChaosClothingSimulationSolver.cpp
--- a/Engine/Plugins/ChaosCloth/Source/ChaosCloth/Private/ChaosCloth/ChaosClothingSimulationSolver.cpp
+++ b/Engine/Plugins/ChaosCloth/Source/ChaosCloth/Private/ChaosCloth/ChaosClothingSimulationSolver.cpp
@@ -160,7 +160,7 @@ TArrayView<T> GetCollisionParticleArrayView(TUniquePtr<Softs::FEvolution>& Evolu
 }
 }
 
-FClothingSimulationSolver::FClothingSimulationSolver(bool bForceBasedSolver, FClothingSimulationConfig* InConfig)
+FClothingSimulationSolver::FClothingSimulationSolver(FClothingSimulationConfig* InConfig, bool bUseLegacySolver)
 	: Evolution(nullptr), PBDEvolution(nullptr)
 	, OldLocalSpaceLocation(0.)
 	, LocalSpaceLocation(0.)
@@ -178,7 +178,7 @@ FClothingSimulationSolver::FClothingSimulationSolver(bool bForceBasedSolver, FCl
 {
 	SetConfig(InConfig); // This will generate a local default config if nullptr so we have something to use if there are no cloth assets..
 
-	if (bForceBasedSolver)
+	if (!bUseLegacySolver)
 	{
 		Evolution.Reset(
 			new Softs::FEvolution(
diff --git a/Engine/Plugins/ChaosCloth/Source/ChaosCloth/Public/ChaosCloth/ChaosClothingSimulationSolver.h b/Engine/Plugins/ChaosCloth/Source/ChaosCloth/Public/ChaosCloth/ChaosClothingSimulationSolver.h
--- a/Engine/Plugins/ChaosCloth/Source/ChaosCloth/Public/ChaosCloth/ChaosClothingSimulationSolver.h
+++ b/Engine/Plugins/ChaosCloth/Source/ChaosCloth/Public/ChaosCloth/ChaosClothingSimulationSolver.h
@@ -31,7 +31,7 @@ namespace Chaos
 	class FClothingSimulationSolver final : public FPhysicsSolverEvents
 	{
 	public:
-		CHAOSCLOTH_API FClothingSimulationSolver(bool bForceBasedSolver = false, FClothingSimulationConfig* InConfig = nullptr);
+		CHAOSCLOTH_API FClothingSimulationSolver(FClothingSimulationConfig* InConfig = nullptr, bool bUseLegacySolver = true);
 		CHAOSCLOTH_API ~FClothingSimulationSolver();
 		
 		FClothingSimulationSolver(const FClothingSimulationSolver&) = delete;
@@ -39,7 +39,7 @@ namespace Chaos
 		FClothingSimulationSolver& operator=(const FClothingSimulationSolver&) = delete;
 		FClothingSimulationSolver& operator=(FClothingSimulationSolver&&) = delete;
 
-		bool IsForceBasedSolver() const { return !!Evolution; }
+		bool IsLegacySolver() const { return !!PBDEvolution; }
 
 		// ---- Animatable property setters ----
 		CHAOSCLOTH_API void SetLocalSpaceLocation(const FVec3& InLocalSpaceLocation, bool bReset = false);
diff --git a/Engine/Plugins/ChaosClothAsset/Source/ChaosClothAssetEngine/Private/ChaosClothAsset/ClothSimulationProxy.cpp b/Engine/Plugins/ChaosClothAsset/Source/ChaosClothAssetEngine/Private/ChaosClothAsset/ClothSimulationProxy.cpp
--- a/Engine/Plugins/ChaosClothAsset/Source/ChaosClothAssetEngine/Private/ChaosClothAsset/ClothSimulationProxy.cpp
+++ b/Engine/Plugins/ChaosClothAsset/Source/ChaosClothAssetEngine/Private/ChaosClothAsset/ClothSimulationProxy.cpp
@@ -102,14 +102,9 @@ namespace UE::Chaos::ClothAsset
 		// Create solver config simulation thread object first. Need to know which solver type we're creating.
 		const int32 SolverConfigIndex = Configs.Emplace(MakeUnique<FClothingSimulationConfig>(ClothComponent.GetPropertyCollections()));  // TODO: Use a separate solver config for outfits
 
-		// Create a solver which can handle force-based solving if any of the LODs require it.
-		bool bForceBasedSolver = false;
-		Configs[SolverConfigIndex]->ForAllProperties([&bForceBasedSolver](Softs::FCollectionPropertyFacade& Property)
-		{
-			bForceBasedSolver = bForceBasedSolver || Property.GetValue<bool>(TEXT("EnableForceBasedSolver"), false);
-		});
-
-		Solver = MakeUnique<::Chaos::FClothingSimulationSolver>(bForceBasedSolver, Configs[SolverConfigIndex].Get());
+		// Use new SoftsEvolution, not PBDEvolution.
+		constexpr bool bUseLegacySolver = false;
+		Solver = MakeUnique<::Chaos::FClothingSimulationSolver>(Configs[SolverConfigIndex].Get(), bUseLegacySolver);
 		Visualization = MakeUnique<::Chaos::FClothVisualization>(Solver.Get());
 
 		// Need a valid context to initialize the mesh
diff --git a/Engine/Plugins/ChaosClothAssetEditor/Source/ChaosClothAssetDataflowNodes/Private/ChaosClothAsset/SimulationSolverConfigNode.cpp b/Engine/Plugins/ChaosClothAssetEditor/Source/ChaosClothAssetDataflowNodes/Private/ChaosClothAsset/SimulationSolverConfigNode.cpp
--- a/Engine/Plugins/ChaosClothAssetEditor/Source/ChaosClothAssetDataflowNodes/Private/ChaosClothAsset/SimulationSolverConfigNode.cpp
+++ b/Engine/Plugins/ChaosClothAssetEditor/Source/ChaosClothAssetDataflowNodes/Private/ChaosClothAsset/SimulationSolverConfigNode.cpp
@@ -19,11 +19,6 @@ void FChaosClothAssetSimulationSolverConfigNode::AddProperties(FPropertyHelper&
 	PropertyHelper.SetProperty(TEXT("DynamicSubstepDeltaTime"), DynamicSubstepDeltaTimeValue);
 	PropertyHelper.SetPropertyBool(this, &bEnableNumSelfCollisionSubsteps);
 	PropertyHelper.SetProperty(this, &NumSelfCollisionSubsteps);
-	PropertyHelper.SetPropertyBool(this, &bEnableForceBasedSolver, {}, ECollectionPropertyFlags::Intrinsic);
-	PropertyHelper.SetProperty(this, &NumNewtonIterations);
-	PropertyHelper.SetProperty(this, &MaxNumCGIterations);
-	PropertyHelper.SetProperty(this, &CGResidualTolerance);
-	PropertyHelper.SetPropertyBool(this, &bDoQuasistatics);
 	
 	PropertyHelper.SetSolverProperty(FName(TEXT("NumSubsteps")), NumSubstepsImported,
 		[](UE::Chaos::ClothAsset::FCollectionClothFacade& ClothFacade)-> int32
@@ -35,16 +30,8 @@ void FChaosClothAssetSimulationSolverConfigNode::AddProperties(FPropertyHelper&
 void FChaosClothAssetSimulationSolverConfigNode::Serialize(FArchive& Ar)
 {
 	Super::Serialize(Ar);
-	Ar.UsingCustomVersion(FFortniteValkyrieBranchObjectVersion::GUID);
 	if (Ar.IsLoading())
 	{
-		if (Ar.CustomVer(FFortniteValkyrieBranchObjectVersion::GUID) < FFortniteValkyrieBranchObjectVersion::ChaosClothAssetWeightedMassAndGravity)
-		{
-			if (!bEnableForceBasedSolver)
-			{
-				NumNewtonIterations = 0;
-			}
-		}
 #if WITH_EDITORONLY_DATA
 		if (NumSubsteps_DEPRECATED != UE::Chaos::ClothAsset::FDefaultSolver::SubSteps)
 		{
diff --git a/Engine/Plugins/ChaosClothAssetEditor/Source/ChaosClothAssetDataflowNodes/Private/ChaosClothAsset/TerminalNode.cpp b/Engine/Plugins/ChaosClothAssetEditor/Source/ChaosClothAssetDataflowNodes/Private/ChaosClothAsset/TerminalNode.cpp
--- a/Engine/Plugins/ChaosClothAssetEditor/Source/ChaosClothAssetDataflowNodes/Private/ChaosClothAsset/TerminalNode.cpp
+++ b/Engine/Plugins/ChaosClothAssetEditor/Source/ChaosClothAssetDataflowNodes/Private/ChaosClothAsset/TerminalNode.cpp
@@ -75,10 +75,6 @@ namespace UE::Chaos::ClothAsset::Private
 					return false;
 				}
 			}
-			if (Property0.GetValue<bool>(TEXT("EnableForceBasedSolver"), false) != Property1.GetValue<bool>(TEXT("EnableForceBasedSolver"), false))
-			{
-				return false;
-			}
 		}
 		return true;
 	}
diff --git a/Engine/Plugins/ChaosClothAssetEditor/Source/ChaosClothAssetDataflowNodes/Public/ChaosClothAsset/SimulationDampingConfigNode.h b/Engine/Plugins/ChaosClothAssetEditor/Source/ChaosClothAssetDataflowNodes/Public/ChaosClothAsset/SimulationDampingConfigNode.h
--- a/Engine/Plugins/ChaosClothAssetEditor/Source/ChaosClothAssetDataflowNodes/Public/ChaosClothAsset/SimulationDampingConfigNode.h
+++ b/Engine/Plugins/ChaosClothAssetEditor/Source/ChaosClothAssetDataflowNodes/Public/ChaosClothAsset/SimulationDampingConfigNode.h
@@ -16,7 +16,6 @@ public:
 	/**
 	 * The amount of global damping applied to the cloth velocities, also known as point damping.
 	 * Point damping improves simulation stability, but can also cause an overall slow-down effect and therefore is best left to very small percentage amounts.
-	 * Weight map only works with UseForceBasedSolver (on SolverConfig). Otherwise, Low value is always used.
 	 */
 	UPROPERTY(EditAnywhere, Category = "Damping Properties", Meta = (UIMin = "0", UIMax = "1", ClampMin = "0", ClampMax = "1"))
 	FChaosClothAssetWeightedValue DampingCoefficientWeighted = {true, 0.01f, 0.01f, TEXT("DampingCoefficient")};
diff --git a/Engine/Plugins/ChaosClothAssetEditor/Source/ChaosClothAssetDataflowNodes/Public/ChaosClothAsset/SimulationGravityConfigNode.h b/Engine/Plugins/ChaosClothAssetEditor/Source/ChaosClothAssetDataflowNodes/Public/ChaosClothAsset/SimulationGravityConfigNode.h
--- a/Engine/Plugins/ChaosClothAssetEditor/Source/ChaosClothAssetDataflowNodes/Public/ChaosClothAsset/SimulationGravityConfigNode.h
+++ b/Engine/Plugins/ChaosClothAssetEditor/Source/ChaosClothAssetDataflowNodes/Public/ChaosClothAsset/SimulationGravityConfigNode.h
@@ -17,9 +17,7 @@ public:
 	UPROPERTY(EditAnywhere, Category = "Gravity Properties", Meta = (InlineEditConditionToggle))
 	bool bUseGravityOverride = false;
 
-	/** Scale factor applied to the world gravity and also to the clothing simulation interactor gravity. Does not affect the gravity if set using the override below. 
-	* Map only has an effect if using "Enable Force Based Solver" on solver config. Otherwise, only the Low value is used
-	*/
+	/** Scale factor applied to the world gravity and also to the clothing simulation interactor gravity. Does not affect the gravity if set using the override below. */
 	UPROPERTY(EditAnywhere, Category = "Gravity Properties", DisplayName = "Gravity Scale", Meta = (UIMin = "0", UIMax = "10", EditCondition = "!bUseGravityOverride"))
 	FChaosClothAssetWeightedValue GravityScaleWeighted = {true, 1.f, 1.f, TEXT("GravityScale")};
 
diff --git a/Engine/Plugins/ChaosClothAssetEditor/Source/ChaosClothAssetDataflowNodes/Public/ChaosClothAsset/SimulationSolverConfigNode.h b/Engine/Plugins/ChaosClothAssetEditor/Source/ChaosClothAssetDataflowNodes/Public/ChaosClothAsset/SimulationSolverConfigNode.h
--- a/Engine/Plugins/ChaosClothAssetEditor/Source/ChaosClothAssetDataflowNodes/Public/ChaosClothAsset/SimulationSolverConfigNode.h
+++ b/Engine/Plugins/ChaosClothAssetEditor/Source/ChaosClothAssetDataflowNodes/Public/ChaosClothAsset/SimulationSolverConfigNode.h
@@ -61,36 +61,6 @@ public:
 	UPROPERTY(EditAnywhere, Category = Simulation, meta = (UIMin = "1", UIMax = "10", ClampMin = "1", ClampMax = "100", EditCondition = "bEnableNumSelfCollisionSubsteps"))
 	int32 NumSelfCollisionSubsteps = 1;
 
-	/**
-	 * Enable the higher accuracy force-based solver (experimental).
-	 */
-	UPROPERTY(EditAnywhere, Category = Experimental)
-	bool bEnableForceBasedSolver = false;
-	
-	/**
-	 * Number of Newton iterations for force-based solver. Prototype only--very few constraints support this.
-	 */
-	UPROPERTY(EditAnywhere, Category = Experimental, meta = (UIMin = "1", UIMax = "10", ClampMin = "0", ClampMax = "100", EditCondition = "bEnableForceBasedSolver"))
-	int32 NumNewtonIterations = 0;
-
-	/**
-	 * Max number of CG Iterations per linear solve
-	 */
-	UPROPERTY(EditAnywhere, Category = Experimental, meta = (UIMin = "1", UIMax = "100", ClampMin = "1", EditCondition = "bEnableForceBasedSolver"))
-	int32 MaxNumCGIterations = 50;
-
-	/**
-	 * CG Tolerance
-	 */
-	UPROPERTY(EditAnywhere, Category = Experimental, meta = (ClampMin = "0", EditCondition = "bEnableForceBasedSolver"))
-	float CGResidualTolerance = 1e-4;
-
-	/**
-	 * Solve quasistaticly (no inertia) 
-	 */
-	UPROPERTY(EditAnywhere, Category = Experimental, meta = (EditCondition = "bEnableForceBasedSolver"))
-	bool bDoQuasistatics = false;
-
 	FChaosClothAssetSimulationSolverConfigNode(const Dataflow::FNodeParameters& InParam, FGuid InGuid = FGuid::NewGuid());
 
 	virtual void Serialize(FArchive& Ar) override;
@@ -99,9 +69,23 @@ private:
 
 	// Deprecated properties
 #if WITH_EDITORONLY_DATA
-
 	UPROPERTY()
 	int32 NumSubsteps_DEPRECATED = UE::Chaos::ClothAsset::FDefaultSolver::SubSteps;
+
+	UPROPERTY()
+	bool bEnableForceBasedSolver_DEPRECATED = false;
+
+	UPROPERTY()
+	int32 NumNewtonIterations_DEPRECATED = 0;
+
+	UPROPERTY()
+	int32 MaxNumCGIterations_DEPRECATED = 50;
+
+	UPROPERTY()
+	float CGResidualTolerance_DEPRECATED = 1e-4;
+
+	UPROPERTY()
+	bool bDoQuasistatics_DEPRECATED = false;
 #endif
 	virtual void AddProperties(FPropertyHelper& PropertyHelper) const override;
 };
