Chaos Cloth Asset: Fix crash when using BendingElement constraints

diff --git a/Engine/Plugins/ChaosCloth/Source/ChaosCloth/Private/ChaosCloth/ChaosClothConstraints.cpp b/Engine/Plugins/ChaosCloth/Source/ChaosCloth/Private/ChaosCloth/ChaosClothConstraints.cpp
--- a/Engine/Plugins/ChaosCloth/Source/ChaosCloth/Private/ChaosCloth/ChaosClothConstraints.cpp
+++ b/Engine/Plugins/ChaosCloth/Source/ChaosCloth/Private/ChaosCloth/ChaosClothConstraints.cpp
@@ -1635,7 +1635,11 @@ void FClothConstraints::CreateForceBasedRules(const TSharedPtr<Softs::FMultiResC
 	}
 	if (BendingElementConstraints)
 	{
-		RuleCreator.AddPBDParallelInitRule(BendingElementConstraints.Get());
+		// These are PBD, not XPBD constraints, but the difference is that XPBD constraints have an Init method and PBD (usually) do not.
+		// BendingElementConstraints do have an Init function, so this will add a call to that.
+		// TODO: clean up the names of these methods to be more explicit.
+		constexpr bool bInitParticles = true;
+		RuleCreator.AddXPBDParallelInitRule<bInitParticles>(BendingElementConstraints.Get());
 		constexpr bool bPostCollisions = false;
 		RuleCreator.AddPerIterationPBDConstraintRule_Apply<bPostCollisions>(BendingElementConstraints.Get());
 	}
