Chaos Cloth - Added Buckling properties to the Chaos Cloth Interactor.

diff --git a/Engine/Plugins/ChaosCloth/Source/ChaosCloth/Private/ChaosCloth/ChaosClothingSimulationConfig.cpp b/Engine/Plugins/ChaosCloth/Source/ChaosCloth/Private/ChaosCloth/ChaosClothingSimulationConfig.cpp
--- a/Engine/Plugins/ChaosCloth/Source/ChaosCloth/Private/ChaosCloth/ChaosClothingSimulationConfig.cpp
+++ b/Engine/Plugins/ChaosCloth/Source/ChaosCloth/Private/ChaosCloth/ChaosClothingSimulationConfig.cpp
@@ -93,7 +93,7 @@ namespace Chaos
 					Property->SetWeightedValue(BendingElementStiffnessIndex, ClothConfig->BendingStiffnessWeighted.Low, ClothConfig->BendingStiffnessWeighted.High);
 					Property->SetStringValue(BendingElementStiffnessIndex, TEXT("BendingStiffness"));
 
-					Property->AddValue(TEXT("BucklingRatio"), ClothConfig->BucklingRatio, NonAnimatablePropertyFlags);
+					Property->AddValue(TEXT("BucklingRatio"), ClothConfig->BucklingRatio, AnimatablePropertyFlags);
 
 					if (ClothConfig->BucklingStiffnessWeighted.Low > 0.f || ClothConfig->BucklingStiffnessWeighted.High > 0.f)
 					{
diff --git a/Engine/Plugins/ChaosCloth/Source/ChaosCloth/Private/ChaosCloth/ChaosClothingSimulationInteractor.cpp b/Engine/Plugins/ChaosCloth/Source/ChaosCloth/Private/ChaosCloth/ChaosClothingSimulationInteractor.cpp
--- a/Engine/Plugins/ChaosCloth/Source/ChaosCloth/Private/ChaosCloth/ChaosClothingSimulationInteractor.cpp
+++ b/Engine/Plugins/ChaosCloth/Source/ChaosCloth/Private/ChaosCloth/ChaosClothingSimulationInteractor.cpp
@@ -60,6 +60,15 @@ void UChaosClothingInteractor::SetMaterial(FVector2D EdgeStiffness, FVector2D Be
 	}));
 }
 
+void UChaosClothingInteractor::SetMaterialBuckling(FVector2D BucklingRatio, FVector2D BucklingStiffness)
+{
+	ConfigCommands.Add(FChaosClothingInteractorConfigCommand::CreateLambda([BucklingRatio, BucklingStiffness](Chaos::FClothingSimulationConfig* Config, int32 LODIndex)
+	{
+		Config->GetProperties(LODIndex).SetValue(TEXT("BucklingRatio"), (float)BucklingRatio[0]);  // TODO: Make BuckingRatio weighted
+		Config->GetProperties(LODIndex).SetWeightedFloatValue(TEXT("BucklingStiffness"), FVector2f(BucklingStiffness));
+	}));
+}
+
 void UChaosClothingInteractor::SetLongRangeAttachmentLinear(float TetherStiffnessLinear, float TetherScale)
 {
 	// Deprecated
diff --git a/Engine/Plugins/ChaosCloth/Source/ChaosCloth/Public/ChaosCloth/ChaosClothingSimulationInteractor.h b/Engine/Plugins/ChaosCloth/Source/ChaosCloth/Public/ChaosCloth/ChaosClothingSimulationInteractor.h
--- a/Engine/Plugins/ChaosCloth/Source/ChaosCloth/Public/ChaosCloth/ChaosClothingSimulationInteractor.h
+++ b/Engine/Plugins/ChaosCloth/Source/ChaosCloth/Public/ChaosCloth/ChaosClothingSimulationInteractor.h
@@ -32,6 +32,9 @@ public:
 	UFUNCTION(BlueprintCallable, Category = ClothingSimulation, Meta = (Keywords = "Chaos Edge Bending Area Stiffness"))
 	CHAOSCLOTH_API void SetMaterial(FVector2D EdgeStiffness = FVector2D(1.f, 1.f), FVector2D BendingStiffness = FVector2D(1.f, 1.f), FVector2D AreaStiffness = FVector2D(1.f, 1.f));
 
+	UFUNCTION(BlueprintCallable, Category = ClothingSimulation, Meta = (Keywords = "Chaos Bending Element Buckling Stiffness Ratio"))
+	CHAOSCLOTH_API void SetMaterialBuckling(FVector2D BucklingRatio = FVector2D(0.f, 0.f), FVector2D BucklingStiffness = FVector2D(1.f, 1.f));
+
 	UFUNCTION(BlueprintCallable, Category = ClothingSimulation, Meta = (Keywords = "Chaos Tether Stiffness Scale"))
 	CHAOSCLOTH_API void SetLongRangeAttachmentLinear(float TetherStiffness = 1.f, float TetherScale = 1.f);
 
