Chaos Cloth: make fictitious angular scale not apply any delta when teleporting.

diff --git a/Engine/Plugins/ChaosCloth/Source/ChaosCloth/Private/ChaosCloth/ChaosClothingSimulationCloth.cpp b/Engine/Plugins/ChaosCloth/Source/ChaosCloth/Private/ChaosCloth/ChaosClothingSimulationCloth.cpp
--- a/Engine/Plugins/ChaosCloth/Source/ChaosCloth/Private/ChaosCloth/ChaosClothingSimulationCloth.cpp
+++ b/Engine/Plugins/ChaosCloth/Source/ChaosCloth/Private/ChaosCloth/ChaosClothingSimulationCloth.cpp
@@ -1136,12 +1136,14 @@ PRAGMA_ENABLE_DEPRECATION_WARNINGS
 		FVec3f OutLinearVelocityScale;
 		FRealSingle OutAngularVelocityScale;
 		FRealSingle OutMaxVelocityScale;
+		bool bDisableFictitiousForces = false;
 		if (bNeedsReset)
 		{
 			// Make sure not to do any pre-sim transform just after a reset
 			OutLinearVelocityScale = FVec3f(1.f);
 			OutAngularVelocityScale = 1.f;
 			OutMaxVelocityScale = 1.f;
+			bDisableFictitiousForces = true; // It doesn't actually matter what value we set here since AngularVelocityScale == 1 means fictitious forces will be 0.
 
 			// Reset to start pose
 			LODData[LODIndex]->ResetStartPose(Solver);
@@ -1162,6 +1164,7 @@ PRAGMA_ENABLE_DEPRECATION_WARNINGS
 			OutLinearVelocityScale = FVec3f(0.f);
 			OutAngularVelocityScale = 0.f;
 			OutMaxVelocityScale = 1.f;
+			bDisableFictitiousForces = true; // Disable fictitious forces. Otherwise they will be applied since AngularVelocityScale < 1.
 			UE_LOG(LogChaosCloth, VeryVerbose, TEXT("Cloth in group Id %d Needs teleport."), GroupId);
 		}
 		else
@@ -1181,7 +1184,8 @@ PRAGMA_ENABLE_DEPRECATION_WARNINGS
 			OutLinearVelocityScale,
 			OutAngularVelocityScale,
 			FictitiousAngularScale,
-			OutMaxVelocityScale);
+			OutMaxVelocityScale, 
+			bDisableFictitiousForces);
 		if (!Solver->IsLegacySolver())
 		{
 			Solver->SetProperties(ParticleRangeId, ConfigProperties, LODData[LODIndex]->WeightMaps);
diff --git a/Engine/Plugins/ChaosCloth/Source/ChaosCloth/Private/ChaosCloth/ChaosClothingSimulationSolver.cpp b/Engine/Plugins/ChaosCloth/Source/ChaosCloth/Private/ChaosCloth/ChaosClothingSimulationSolver.cpp
--- a/Engine/Plugins/ChaosCloth/Source/ChaosCloth/Private/ChaosCloth/ChaosClothingSimulationSolver.cpp
+++ b/Engine/Plugins/ChaosCloth/Source/ChaosCloth/Private/ChaosCloth/ChaosClothingSimulationSolver.cpp
@@ -1078,7 +1078,8 @@ void FClothingSimulationSolver::SetReferenceVelocityScale(
 	const FRigidTransform3& ReferenceSpaceTransform,
 	const TVec3<FRealSingle>& LinearVelocityScale,
 	FRealSingle AngularVelocityScale, FRealSingle FictitiousAngularScale,
-	FRealSingle MaxVelocityScale
+	FRealSingle MaxVelocityScale, 
+	bool bDisableFictitiousForces
 )
 {
 	FRigidTransform3 OldRootBoneLocalTransform = OldReferenceSpaceTransform;
@@ -1127,7 +1128,7 @@ void FClothingSimulationSolver::SetReferenceVelocityScale(
 		Softs::FSolverRotation3(PreSimulationTransform.GetRotation()));
 
 	// Fictitious angular scale only applied for PBDEvolution. It's applied by ExternalForces for Evolution.
-	const FReal AppliedFictitiousAngularScale = PBDEvolution ? FMath::Min((FReal)2., (FReal)FictitiousAngularScale) : (FReal)1.;
+	const FReal AppliedFictitiousAngularScale = bDisableFictitiousForces ? (FReal)0.f : (PBDEvolution ? FMath::Min((FReal)2., (FReal)FictitiousAngularScale) : (FReal)1.);
 
 	// Save the reference bone relative angular velocity for calculating the fictitious forces
 	const FVec3 FictitiousAngularDisplacement = ReferenceSpaceTransform.TransformVector(Axis * PartialDeltaAngle)
diff --git a/Engine/Plugins/ChaosCloth/Source/ChaosCloth/Public/ChaosCloth/ChaosClothingSimulationSolver.h b/Engine/Plugins/ChaosCloth/Source/ChaosCloth/Public/ChaosCloth/ChaosClothingSimulationSolver.h
--- a/Engine/Plugins/ChaosCloth/Source/ChaosCloth/Public/ChaosCloth/ChaosClothingSimulationSolver.h
+++ b/Engine/Plugins/ChaosCloth/Source/ChaosCloth/Public/ChaosCloth/ChaosClothingSimulationSolver.h
@@ -145,7 +145,8 @@ namespace Chaos
 			const TVec3<FRealSingle>& LinearVelocityScale,
 			FRealSingle AngularVelocityScale,
 			FRealSingle FictitiousAngularScale,
-			FRealSingle MaxVelocityScale = 1.f);
+			FRealSingle MaxVelocityScale = 1.f,
+			bool bDisableFictitiousForces = false);
 
 		/** PBDSolver version */
 		CHAOSCLOTH_API void SetProperties(
