Cloth: remove warning about "Invalid pattern data. Num pattern indices ... ".

Every skeletal mesh-based cloth will generate this warning, but in this case, the behavior is expected.

It is no longer possible to author data using the Chaos Cloth Asset (the other user of this code) in a way to generate this warning.

diff --git a/Engine/Plugins/ChaosCloth/Source/ChaosCloth/Private/ChaosCloth/ChaosClothingPatternData.cpp b/Engine/Plugins/ChaosCloth/Source/ChaosCloth/Private/ChaosCloth/ChaosClothingPatternData.cpp
--- a/Engine/Plugins/ChaosCloth/Source/ChaosCloth/Private/ChaosCloth/ChaosClothingPatternData.cpp
+++ b/Engine/Plugins/ChaosCloth/Source/ChaosCloth/Private/ChaosCloth/ChaosClothingPatternData.cpp
@@ -54,8 +54,6 @@ void FClothingPatternData::GenerateDerivedPatternData(
 {
 	if (Indices.Num() != PatternIndices.Num())
 	{
-		UE_LOG(LogChaosCloth, Warning, TEXT("Invalid pattern data. Num pattern indices %d != Num welded indices %d. Panel-based cloth is disabled for this LOD."), PatternIndices.Num(), Indices.Num());
-
 		// Reset all of the pattern data and return
 		Reset();
 		return;
