Chaos Cloth - Removed deprecated cloth code, including the non abstract mesh interface.

#rnx

diff --git a/Engine/Plugins/ChaosCloth/Source/ChaosCloth/Private/ChaosCloth/ChaosClothConstraints.cpp b/Engine/Plugins/ChaosCloth/Source/ChaosCloth/Private/ChaosCloth/ChaosClothConstraints.cpp
--- a/Engine/Plugins/ChaosCloth/Source/ChaosCloth/Private/ChaosCloth/ChaosClothConstraints.cpp
+++ b/Engine/Plugins/ChaosCloth/Source/ChaosCloth/Private/ChaosCloth/ChaosClothConstraints.cpp
@@ -444,35 +444,6 @@ void FClothConstraints::Enable(bool bEnable)
 	}
 }
 
-void FClothConstraints::AddRules(
-	const Softs::FCollectionPropertyConstFacade& ConfigProperties,
-	const FTriangleMesh& TriangleMesh,
-	const TArray<TConstArrayView<FRealSingle>>& WeightMapArray,
-	const TArray<TConstArrayView<TTuple<int32, int32, FRealSingle>>>& Tethers,
-	Softs::FSolverReal MeshScale,
-	bool bEnabled)
-{
-	// Build new weight map container from the old legacy weight map enum
-	TMap<FString, TConstArrayView<FRealSingle>> WeightMaps;
-
-	const UEnum* const ChaosWeightMapTargetEnum = StaticEnum<EChaosWeightMapTarget>();
-	const int32 NumWeightMaps = (int32)ChaosWeightMapTargetEnum->GetMaxEnumValue() + 1;
-
-	WeightMaps.Reserve(NumWeightMaps);
-
-	for (int32 EnumIndex = 0; EnumIndex < ChaosWeightMapTargetEnum->NumEnums(); ++EnumIndex)
-	{
-		const int32 TargetIndex = (int32)ChaosWeightMapTargetEnum->GetValueByIndex(EnumIndex);
-		const FString WeightMapName = ChaosWeightMapTargetEnum->GetNameByIndex(EnumIndex).ToString();
-
-		WeightMaps.Add(WeightMapName, WeightMapArray[TargetIndex]);
-	}
-
-	// Call new AddRules function
-	AddRules(ConfigProperties, TriangleMesh, nullptr, WeightMaps, TMap<FString, const TSet<int32>*>(), TMap<FString, const TSet<int32>*>(),
-		TMap<FString, TConstArrayView<int32>>(), Tethers, MeshScale, bEnabled);
-}
-
 void FClothConstraints::AddRules(
 	const Softs::FCollectionPropertyConstFacade& ConfigProperties,
 	const FTriangleMesh& TriangleMesh,
@@ -2323,7 +2294,7 @@ void FClothConstraints::Update(
 	}
 }
 
-// Deprecated
+// Deprecated 5.4
 void FClothConstraints::Update(
 	const Softs::FCollectionPropertyConstFacade& ConfigProperties,
 	const TMap<FString, TConstArrayView<FRealSingle>>& WeightMaps,
@@ -2333,13 +2304,4 @@ void FClothConstraints::Update(
 	Update(ConfigProperties, WeightMaps, TMap<FString, const TSet<int32>*>(), TMap<FString, const TSet<int32>*>(), TMap<FString, TConstArrayView<int32>>(), MeshScale, MaxDistancesScale);
 }
 
-// Deprecated
-void FClothConstraints::Update(
-	const Softs::FCollectionPropertyConstFacade& ConfigProperties,
-	Softs::FSolverReal MeshScale,
-	Softs::FSolverReal MaxDistancesScale)
-{
-	Update(ConfigProperties, TMap<FString, TConstArrayView<FRealSingle>>(), TMap<FString, const TSet<int32>*>(), TMap<FString, const TSet<int32>*>(), TMap<FString, TConstArrayView<int32>>(), MeshScale, MaxDistancesScale);
-}
-
 }  // End namespace Chaos
diff --git a/Engine/Plugins/ChaosCloth/Source/ChaosCloth/Private/ChaosCloth/ChaosClothingSimulation.cpp b/Engine/Plugins/ChaosCloth/Source/ChaosCloth/Private/ChaosCloth/ChaosClothingSimulation.cpp
--- a/Engine/Plugins/ChaosCloth/Source/ChaosCloth/Private/ChaosCloth/ChaosClothingSimulation.cpp
+++ b/Engine/Plugins/ChaosCloth/Source/ChaosCloth/Private/ChaosCloth/ChaosClothingSimulation.cpp
@@ -377,9 +377,7 @@ void FClothingSimulation::CreateActor(USkeletalMeshComponent* InOwnerComponent,
 	// Create cloth runtime simulation object
 	const int32 ClothIndex = Cloths.Emplace(MakeUnique<FClothingSimulationCloth>(
 		Configs[ClothConfigIndex].Get(),
-PRAGMA_DISABLE_DEPRECATION_WARNINGS  // TODO: CHAOS_IS_CLOTHINGSIMULATIONMESH_ABSTRACT
 		Meshes[MeshIndex].Get(),
-PRAGMA_ENABLE_DEPRECATION_WARNINGS
 		TArray<FClothingSimulationCollider*>({ Colliders[ColliderIndex].Get() }),
 		InSimDataIndex));
 
@@ -535,8 +533,7 @@ void FClothingSimulation::Simulate(IClothingSimulationContext* InContext)
 		}
 
 		// Step the simulation
-		PRAGMA_DISABLE_DEPRECATION_WARNINGS // Supporting deprecated CachedPositions instead of new CacheData
-		if(Solver->GetEnableSolver() || (!Context->CacheData.HasData() && Context->CachedPositions.Num() == 0))
+		if (Solver->GetEnableSolver() || !Context->CacheData.HasData())
 		{
 			Solver->Update(SmoothedDeltaTime);
 
@@ -547,17 +544,8 @@ void FClothingSimulation::Simulate(IClothingSimulationContext* InContext)
 		}
 		else
 		{
-			if (Context->CacheData.HasData())
-			{
-				Solver->UpdateFromCache(Context->CacheData);
-			}
-			else
-			{
-				check(Context->CachedPositions.Num());
-				Solver->UpdateFromCache(Context->CachedPositions, Context->CachedVelocities);
-			}
+			Solver->UpdateFromCache(Context->CacheData);
 		}
-		PRAGMA_ENABLE_DEPRECATION_WARNINGS
 
 		// Keep the actual used number of iterations for the stats
 		NumIterations = Solver->GetNumUsedIterations();
@@ -640,9 +628,7 @@ void FClothingSimulation::GetSimulationData(
 		}
 
 		// If the LOD has changed while the simulation is suspended, the cloth still needs to be updated with the correct LOD data
-PRAGMA_DISABLE_DEPRECATION_WARNINGS  // TODO: CHAOS_IS_CLOTHINGSIMULATIONMESH_ABSTRACT
 		const int32 LODIndex = Cloth->GetMesh()->GetLODIndex();
-PRAGMA_ENABLE_DEPRECATION_WARNINGS
 		if (LODIndex != Cloth->GetLODIndex(Solver.Get()))
 		{
 			Solver->Update(FSolverReal(0.));  // Update for LOD switching, but do not simulate
@@ -722,9 +708,7 @@ PRAGMA_ENABLE_DEPRECATION_WARNINGS
 		}
 
 		// Set the current LOD these data apply to, so that the correct deformer mappings can be applied
-PRAGMA_DISABLE_DEPRECATION_WARNINGS  // TODO: CHAOS_IS_CLOTHINGSIMULATIONMESH_ABSTRACT
 		Data.LODIndex = Cloth->GetMesh()->GetOwnerLODIndex(LODIndex);  // The owner component LOD index can be different to the cloth mesh LOD index
-PRAGMA_ENABLE_DEPRECATION_WARNINGS
 	}
 }
 
@@ -821,9 +805,7 @@ void FClothingSimulation::RefreshClothConfig(const IClothingSimulationContext* I
 		// Recreate cloth runtime simulation object
 		Cloth = MakeUnique<FClothingSimulationCloth>(
 			Config,
-			PRAGMA_DISABLE_DEPRECATION_WARNINGS  // TODO: CHAOS_IS_CLOTHINGSIMULATIONMESH_ABSTRACT
 			Mesh,
-			PRAGMA_ENABLE_DEPRECATION_WARNINGS
 			MoveTemp(ClothColliders),
 			GroupId);
 
diff --git a/Engine/Plugins/ChaosCloth/Source/ChaosCloth/Private/ChaosCloth/ChaosClothingSimulationCloth.cpp b/Engine/Plugins/ChaosCloth/Source/ChaosCloth/Private/ChaosCloth/ChaosClothingSimulationCloth.cpp
--- a/Engine/Plugins/ChaosCloth/Source/ChaosCloth/Private/ChaosCloth/ChaosClothingSimulationCloth.cpp
+++ b/Engine/Plugins/ChaosCloth/Source/ChaosCloth/Private/ChaosCloth/ChaosClothingSimulationCloth.cpp
@@ -195,10 +195,8 @@ void FClothingSimulationCloth::FLODData::AddParticles(FClothingSimulationSolver*
 
 	const FTriangleMesh& TriangleMesh = Solver->IsLegacySolver() ? SolverDatum.OffsetTriangleMesh : NoOffsetTriangleMesh;
 
-	PRAGMA_DISABLE_DEPRECATION_WARNINGS  // TODO: CHAOS_IS_CLOTHINGSIMULATIONMESH_ABSTRACT
 	// Update source mesh for this LOD, this is required prior to reset the start pose
 	Cloth->Mesh->Update(Solver, INDEX_NONE, InLODIndex, 0, SolverDatum.ParticleRangeId);
-	PRAGMA_ENABLE_DEPRECATION_WARNINGS
 
 	// Reset the particles start pose before setting up mass and constraints
 	ResetStartPose(Solver);
@@ -272,10 +270,8 @@ void FClothingSimulationCloth::FLODData::Add(FClothingSimulationSolver* Solver,
 	check(SolverDatum.ParticleRangeId != INDEX_NONE);
 	const int32 ParticleRangeId = SolverDatum.ParticleRangeId;
 
-	PRAGMA_DISABLE_DEPRECATION_WARNINGS  // TODO: CHAOS_IS_CLOTHINGSIMULATIONMESH_ABSTRACT
 	// Retrieve the component's scale
 	const Softs::FSolverReal MeshScale = Cloth->Mesh->GetScale();
-	PRAGMA_ENABLE_DEPRECATION_WARNINGS
 
 	const FTriangleMesh& TriangleMesh = Solver->IsLegacySolver() ? SolverDatum.OffsetTriangleMesh : NoOffsetTriangleMesh;
 
@@ -364,9 +360,7 @@ void FClothingSimulationCloth::FLODData::Update(FClothingSimulationSolver* Solve
 	FClothConstraints& ClothConstraints = Solver->GetClothConstraints(ParticleRangeId);
 
 	check(Cloth->Config);
-PRAGMA_DISABLE_DEPRECATION_WARNINGS  // TODO: CHAOS_IS_CLOTHINGSIMULATIONMESH_ABSTRACT
 	const Softs::FSolverReal MeshScale = Cloth->Mesh->GetScale();
-PRAGMA_ENABLE_DEPRECATION_WARNINGS
 	const Softs::FSolverReal MaxDistancesScale = (Softs::FSolverReal)Cloth->MaxDistancesMultiplier;
 	ClothConstraints.Update(Cloth->Config->GetProperties(SolverDatum.LODIndex), WeightMaps, VertexSets, FaceSets, FaceIntMaps, MeshScale, MaxDistancesScale);
 }
@@ -409,9 +403,7 @@ void FClothingSimulationCloth::FLODData::UpdateNormals(FClothingSimulationSolver
 
 FClothingSimulationCloth::FClothingSimulationCloth(
 	FClothingSimulationConfig* InConfig,
-PRAGMA_DISABLE_DEPRECATION_WARNINGS  // TODO: CHAOS_IS_CLOTHINGSIMULATIONMESH_ABSTRACT
 	FClothingSimulationMesh* InMesh,
-PRAGMA_ENABLE_DEPRECATION_WARNINGS
 	TArray<FClothingSimulationCollider*>&& InColliders,
 	uint32 InGroupId)
 	: GroupId(InGroupId)
@@ -421,232 +413,6 @@ PRAGMA_ENABLE_DEPRECATION_WARNINGS
 	SetColliders(MoveTemp(InColliders));
 }
 
-FClothingSimulationCloth::FClothingSimulationCloth(
-PRAGMA_DISABLE_DEPRECATION_WARNINGS  // TODO: CHAOS_IS_CLOTHINGSIMULATIONMESH_ABSTRACT
-	FClothingSimulationMesh* InMesh,
-PRAGMA_ENABLE_DEPRECATION_WARNINGS
-	TArray<FClothingSimulationCollider*>&& InColliders,
-	uint32 InGroupId,
-	EMassMode InMassMode,
-	FRealSingle InMassValue,
-	FRealSingle InMinPerParticleMass,
-	const TVec2<FRealSingle>& InEdgeStiffness,
-	const TVec2<FRealSingle>& InBendingStiffness,
-	FRealSingle InBucklingRatio,
-	const TVec2<FRealSingle>& InBucklingStiffness,
-	bool bInUseBendingElements,
-	const TVec2<FRealSingle>& InAreaStiffness,
-	FRealSingle /*InVolumeStiffness*/,  // Deprecated
-	bool /*bInUseThinShellVolumeConstraints*/,  // Deprecated
-	const TVec2<FRealSingle>& InTetherStiffness,
-	const TVec2<FRealSingle>& InTetherScale,
-PRAGMA_DISABLE_DEPRECATION_WARNINGS
-	ETetherMode InTetherMode,
-PRAGMA_ENABLE_DEPRECATION_WARNINGS
-	FRealSingle InMaxDistancesMultiplier,
-	const TVec2<FRealSingle>& InAnimDriveStiffness,
-	const TVec2<FRealSingle>& InAnimDriveDamping,
-	FRealSingle /*InShapeTargetStiffness*/,  // Deprecated
-	bool bInUseXPBDEdgeConstraints,
-	bool bInUseXPBDBendingConstraints,
-	bool /*bInUseXPBDAreaConstraints*/,  // Deprecated
-	FRealSingle InGravityScale,
-	bool bInUseGravityOverride,
-	const TVec3<FRealSingle>& InGravityOverride,
-	const TVec3<FRealSingle>& InLinearVelocityScale,
-	FRealSingle InAngularVelocityScale,
-	FRealSingle InFictitiousAngularScale,
-	const TVec2<FRealSingle>& InDrag,
-	const TVec2<FRealSingle>& InLift,
-	bool bInUsePointBasedWindModel,
-	const TVec2<FRealSingle>& InPressure,
-	FRealSingle InDampingCoefficient,
-	FRealSingle InLocalDampingCoefficient,
-	FRealSingle InCollisionThickness,
-	FRealSingle InFrictionCoefficient,
-	bool bInUseCCD,
-	bool bInUseSelfCollisions,
-	FRealSingle InSelfCollisionThickness,
-	FRealSingle InSelfCollisionFrictionCoefficient,
-	bool bInUseSelfIntersections,
-	bool bInUseLegacyBackstop,
-	bool bInUseLODIndexOverride, 
-	int32 InLODIndexOverride,
-	const TVec2<FRealSingle>& InEdgeDampingRatio,
-	const TVec2<FRealSingle>& InBendingDampingRatio)
-	: GroupId(InGroupId)
-	, PropertyCollection(MakeShared<FManagedArrayCollection>())
-	, bUseLODIndexOverride(bInUseLODIndexOverride)
-	, LODIndexOverride(InLODIndexOverride)
-{
-	// Turn parameters into a config
-	Softs::FCollectionPropertyMutableFacade Properties(PropertyCollection);
-	Properties.DefineSchema();
-
-	constexpr bool bEnable = true;
-	constexpr bool bAnimatable = true;
-	
-	// Mass
-	{
-		Properties.AddValue(TEXT("MassMode"), (int32)InMassMode);
-		Properties.AddValue(TEXT("MassValue"), InMassValue);
-		Properties.AddValue(TEXT("MinPerParticleMass"), InMinPerParticleMass);
-	}
-	
-	// Edge constraint
-	if (InEdgeStiffness[0] > 0.f || InEdgeStiffness[1] > 0.f)
-	{
-		const int32 EdgeSpringStiffnessIndex = Properties.AddProperty(TEXT("EdgeSpringStiffness"), bEnable, bAnimatable);
-		Properties.SetWeightedValue(EdgeSpringStiffnessIndex, InEdgeStiffness[0], InEdgeStiffness[1]);
-		Properties.SetStringValue(EdgeSpringStiffnessIndex, TEXT("EdgeStiffness"));
-	}
-	
-	// Bending constraint
-	if (InBendingStiffness[0] > 0.f || InBendingStiffness[1] > 0.f ||
-		(bInUseBendingElements && (InBucklingStiffness[0] > 0.f || InBucklingStiffness[1] > 0.f)))
-	{
-		if (bInUseBendingElements)
-		{
-			const int32 BendingElementStiffnessIndex = Properties.AddProperty(TEXT("BendingElementStiffness"), bEnable, bAnimatable);
-			Properties.SetWeightedValue(BendingElementStiffnessIndex, InBendingStiffness[0], InBendingStiffness[1]);
-			Properties.SetStringValue(BendingElementStiffnessIndex, TEXT("BendingStiffness"));
-
-			Properties.AddValue(TEXT("BucklingRatio"), InBucklingRatio);
-
-			if (InBucklingStiffness[0] > 0.f || InBucklingStiffness[1] > 0.f)
-			{
-				const int32 BucklingStiffnessIndex = Properties.AddProperty(TEXT("BucklingStiffness"), bEnable, bAnimatable);
-				Properties.SetWeightedValue(BucklingStiffnessIndex, InBucklingStiffness[0], InBucklingStiffness[1]);
-				Properties.SetStringValue(BucklingStiffnessIndex, TEXT("BucklingStiffness"));
-			}
-		}
-		else  // Not using bending elements
-		{
-			const int32 BendingSpringStiffnessIndex = Properties.AddProperty(TEXT("BendingSpringStiffness"), bEnable, bAnimatable);
-			Properties.SetWeightedValue(BendingSpringStiffnessIndex, InBendingStiffness[0], InBendingStiffness[1]);
-			Properties.SetStringValue(BendingSpringStiffnessIndex, TEXT("BendingStiffness"));
-		}
-	}
-
-	// Area constraint
-	if (InAreaStiffness[0] > 0.f || InAreaStiffness[1] > 0.f)
-	{
-		const int32 AreaSpringStiffnessIndex = Properties.AddProperty(TEXT("AreaSpringStiffness"), bEnable, bAnimatable);
-		Properties.SetWeightedValue(AreaSpringStiffnessIndex, InAreaStiffness[0], InAreaStiffness[1]);
-		Properties.SetStringValue(AreaSpringStiffnessIndex, TEXT("AreaStiffness"));
-	}
-
-	// Long range attachment
-	if (InTetherStiffness[0] > 0.f || InTetherStiffness[1] > 0.f)
-	{
-		PRAGMA_DISABLE_DEPRECATION_WARNINGS
-		const bool bUseGeodesicTethers = InTetherMode == ETetherMode::Geodesic;
-		PRAGMA_ENABLE_DEPRECATION_WARNINGS
-		Properties.AddValue(TEXT("UseGeodesicTethers"), bUseGeodesicTethers);
-
-		const int32 TetherStiffnessIndex = Properties.AddProperty(TEXT("TetherStiffness"), bEnable , bAnimatable);
-		Properties.SetWeightedValue(TetherStiffnessIndex, InTetherStiffness[0], InTetherStiffness[1]);
-		Properties.SetStringValue(TetherStiffnessIndex, TEXT("TetherStiffness"));
-	
-		const int32 TetherScaleIndex = Properties.AddProperty(TEXT("TetherScale"), bEnable, bAnimatable);
-		Properties.SetWeightedValue(TetherScaleIndex, InTetherScale[0], InTetherScale[1]);
-		Properties.SetStringValue(TetherScaleIndex, TEXT("TetherScale"));
-	}
-
-	// AnimDrive
-	if (InAnimDriveStiffness[0] > 0.f || InAnimDriveStiffness[1] > 0.f)
-	{
-		const int32 AnimDriveStiffnessIndex = Properties.AddProperty(TEXT("AnimDriveStiffness"), bEnable, bAnimatable);
-		Properties.SetWeightedValue(AnimDriveStiffnessIndex, InAnimDriveStiffness[0], InAnimDriveStiffness[1]);
-		Properties.SetStringValue(AnimDriveStiffnessIndex, TEXT("AnimDriveStiffness"));
-
-		const int32 AnimDriveDampingIndex = Properties.AddProperty(TEXT("AnimDriveDamping"), bEnable, bAnimatable);
-		Properties.SetWeightedValue(AnimDriveDampingIndex, InAnimDriveDamping[0], InAnimDriveDamping[1]);
-		Properties.SetStringValue(AnimDriveDampingIndex, TEXT("AnimDriveDamping"));
-	}
-
-	// Gravity
-	{
-		Properties.AddValue(TEXT("GravityScale"), InGravityScale, bEnable, bAnimatable);
-		Properties.AddValue(TEXT("UseGravityOverride"), bInUseGravityOverride, bEnable, bAnimatable);
-		Properties.AddValue(TEXT("GravityOverride"), FVector3f(InGravityOverride), bEnable, bAnimatable);
-	}
-
-	// Velocity scale
-	{
-		Properties.AddValue(TEXT("LinearVelocityScale"), FVector3f(InLinearVelocityScale), bEnable, bAnimatable);
-		Properties.AddValue(TEXT("AngularVelocityScale"), InAngularVelocityScale, bEnable, bAnimatable);
-		Properties.AddValue(TEXT("FictitiousAngularScale"), InFictitiousAngularScale, bEnable, bAnimatable);
-	}
-
-	// Aerodynamics
-	Properties.AddValue(TEXT("UsePointBasedWindModel"), bInUsePointBasedWindModel);
-	if (!bInUsePointBasedWindModel && (InDrag[0] > 0.f || InDrag[1] > 0.f || InLift[0] > 0.f || InLift[1] > 0.f))
-	{
-		const int32 DragIndex = Properties.AddProperty(TEXT("Drag"), bEnable, bAnimatable);
-		Properties.SetWeightedValue(DragIndex, InDrag[0], InDrag[1]);
-		Properties.SetStringValue(DragIndex, TEXT("Drag"));
-
-		const int32 LiftIndex = Properties.AddProperty(TEXT("Lift"), bEnable, bAnimatable);
-		Properties.SetWeightedValue(LiftIndex, InLift[0], InLift[1]);
-		Properties.SetStringValue(LiftIndex, TEXT("Lift"));
-
-		Properties.AddValue(TEXT("FluidDensity"), ClothingSimulationClothDefault::AirDensity, bEnable, bAnimatable);
-
-		Properties.AddValue(TEXT("WindVelocity"), FVector3f(0.f), bEnable, bAnimatable);  // Wind velocity must exist to be animatable
-	}
-
-	// Pressure
-	if (InPressure[0] != 0.f || InPressure[1] != 0.f)
-	{
-		const int32 PressureIndex = Properties.AddProperty(TEXT("Pressure"), bEnable, bAnimatable);
-		Properties.SetWeightedValue(PressureIndex, InPressure[0], InPressure[1]);
-		Properties.SetStringValue(PressureIndex, TEXT("Pressure"));
-	}
-
-	// Damping
-	Properties.AddValue(TEXT("DampingCoefficient"), InDampingCoefficient, bEnable, bAnimatable);
-	Properties.AddValue(TEXT("LocalDampingCoefficient"), InLocalDampingCoefficient, bEnable, bAnimatable);
-
-	// Collision
-	Properties.AddValue(TEXT("CollisionThickness"), InCollisionThickness, bEnable, bAnimatable);
-	Properties.AddValue(TEXT("FrictionCoefficient"), InFrictionCoefficient, bEnable, bAnimatable);
-	Properties.AddValue(TEXT("UseCCD"), bInUseCCD, bEnable, bAnimatable);
-	Properties.AddValue(TEXT("UseSelfCollisions"), bInUseSelfCollisions);
-	Properties.AddValue(TEXT("SelfCollisionThickness"), InSelfCollisionThickness);
-	Properties.AddValue(TEXT("SelfCollisionFrictionCoefficient"), InSelfCollisionFrictionCoefficient);
-	Properties.AddValue(TEXT("UseSelfIntersections"), bInUseSelfIntersections);
-	Properties.AddValue(TEXT("SelfCollideAgainstAllKinematicVertices"), true); // Match legacy behavior
-
-	// Max distance
-	{
-		const int32 MaxDistanceIndex = Properties.AddProperty(TEXT("MaxDistance"));
-		Properties.SetWeightedValue(MaxDistanceIndex, 0.f, 1.f);  // Backward compatibility with legacy mask must use a unit range since the multiplier is in the mask
-		Properties.SetStringValue(MaxDistanceIndex, TEXT("MaxDistance"));
-	}
-
-	// Backstop
-	{
-		const int32 BackstopScaleIndex = Properties.AddProperty(TEXT("BackstopDistance"));
-		Properties.SetWeightedValue(BackstopScaleIndex, 0.f, 1.f);  // Backward compatibility with legacy mask must use a unit range since the multiplier is in the mask
-		Properties.SetStringValue(BackstopScaleIndex, TEXT("BackstopDistance"));
-
-		const int32 BackstopRadiusIndex = Properties.AddProperty(TEXT("BackstopRadius"));
-		Properties.SetWeightedValue(BackstopRadiusIndex, 0.f, 1.f);  // Backward compatibility with legacy mask must use a unit range since the multiplier is in the mask
-		Properties.SetStringValue(BackstopRadiusIndex, TEXT("BackstopRadius"));
-
-		Properties.AddValue(TEXT("UseLegacyBackstop"), bInUseLegacyBackstop);
-	}
-
-	PRAGMA_DISABLE_DEPRECATION_WARNINGS
-	Config = new FClothingSimulationConfig(PropertyCollection);
-	PRAGMA_ENABLE_DEPRECATION_WARNINGS
-
-	// Set mesh and colliders
-	SetMesh(InMesh);
-	SetColliders(MoveTemp(InColliders));
-}
-
 FClothingSimulationCloth::~FClothingSimulationCloth()
 {
 	// If the PropertyCollection is owned by this object, so does the current config object
@@ -656,7 +422,6 @@ FClothingSimulationCloth::~FClothingSimulationCloth()
 	}
 }
 
-PRAGMA_DISABLE_DEPRECATION_WARNINGS  // TODO: CHAOS_IS_CLOTHINGSIMULATIONMESH_ABSTRACT
 void FClothingSimulationCloth::SetMesh(FClothingSimulationMesh* InMesh)
 {
 	Mesh = InMesh;
@@ -711,7 +476,6 @@ void FClothingSimulationCloth::SetMesh(FClothingSimulationMesh* InMesh)
 		Solver->RefreshCloth(this);
 	}
 }
-PRAGMA_ENABLE_DEPRECATION_WARNINGS
 
 void FClothingSimulationCloth::SetConfig(FClothingSimulationConfig* InConfig)
 {
@@ -730,9 +494,7 @@ void FClothingSimulationCloth::SetConfig(FClothingSimulationConfig* InConfig)
 	{
 		// Create a default empty config object for coherence
 		PropertyCollection = MakeShared<FManagedArrayCollection>();
-		PRAGMA_DISABLE_DEPRECATION_WARNINGS
-		Config = new FClothingSimulationConfig(PropertyCollection); 
-		PRAGMA_ENABLE_DEPRECATION_WARNINGS
+		Config = new FClothingSimulationConfig({ PropertyCollection });
 	}
 }
 
@@ -948,14 +710,6 @@ const FTriangleMesh& FClothingSimulationCloth::GetTriangleMesh(const FClothingSi
 	return LODData.IsValidIndex(LODIndex) ? (Solver->IsLegacySolver() ? LODData[LODIndex]->SolverData.FindChecked(Solver).OffsetTriangleMesh: LODData[LODIndex]->NoOffsetTriangleMesh) : EmptyTriangleMesh;
 }
 
-
-// Deprecated for 5.3
-const TArray<TConstArrayView<FRealSingle>>& FClothingSimulationCloth::GetWeightMaps(const FClothingSimulationSolver* Solver) const
-{
-	static const TArray<TConstArrayView<FRealSingle>> EmptyWeightMaps;
-	return EmptyWeightMaps;  // Can't return a reference anymore, so returning an empty weightmap array is probably the safest for existing code
-}
-
 TConstArrayView<FRealSingle> FClothingSimulationCloth::GetWeightMapByName(const FClothingSimulationSolver* Solver, const FString& Name) const
 {
 	const int32 LODIndex = LODIndices.FindChecked(Solver);
@@ -1003,9 +757,7 @@ const TArray<TConstArrayView<TTuple<int32, int32, float>>>& FClothingSimulationC
 
 int32 FClothingSimulationCloth::GetReferenceBoneIndex() const
 {
-PRAGMA_DISABLE_DEPRECATION_WARNINGS  // TODO: CHAOS_IS_CLOTHINGSIMULATIONMESH_ABSTRACT
 	return Mesh ? Mesh->GetReferenceBoneIndex() : INDEX_NONE;
-PRAGMA_ENABLE_DEPRECATION_WARNINGS
 }
 
 void FClothingSimulationCloth::PreUpdate(FClothingSimulationSolver* Solver)
@@ -1043,13 +795,11 @@ void FClothingSimulationCloth::Update(FClothingSimulationSolver* Solver)
 	int32& LODIndex = LODIndices.FindChecked(Solver);  // Must be added to solver first
 
 	const int32 PrevLODIndex = LODIndex;
-PRAGMA_DISABLE_DEPRECATION_WARNINGS  // TODO: CHAOS_IS_CLOTHINGSIMULATIONMESH_ABSTRACT
 	LODIndex = bUseLODIndexOverride && LODData.IsValidIndex(LODIndexOverride) ? LODIndexOverride : Mesh->GetLODIndex();
 
 	// Update reference space transform from the mesh's reference bone transform  TODO: Add override in the style of LODIndexOverride
 	const FRigidTransform3 OldReferenceSpaceTransform = ReferenceSpaceTransform;
 	ReferenceSpaceTransform = Mesh->GetReferenceBoneTransform();
-PRAGMA_ENABLE_DEPRECATION_WARNINGS
 	ReferenceSpaceTransform.SetScale3D(FVec3(1.f));
 
 	// Update Cloth Colliders
@@ -1070,9 +820,7 @@ PRAGMA_ENABLE_DEPRECATION_WARNINGS
 		return;
 	}
 
-PRAGMA_DISABLE_DEPRECATION_WARNINGS  // TODO: CHAOS_IS_CLOTHINGSIMULATIONMESH_ABSTRACT
 	Mesh->Update(Solver, PrevLODIndex, LODIndex, PrevParticleRangeId, ParticleRangeId);
-PRAGMA_ENABLE_DEPRECATION_WARNINGS
 
 	// Retrieve config
 	check(Config);
@@ -1083,10 +831,8 @@ PRAGMA_ENABLE_DEPRECATION_WARNINGS
 	if (CoarseLODIndex != INDEX_NONE)
 	{
 		LODData[CoarseLODIndex]->Enable(Solver, true);
-PRAGMA_DISABLE_DEPRECATION_WARNINGS 
 		// TODO: interpolate/ reset when LOD switching to enable multires
 		Mesh->Update(Solver, CoarseLODIndex, CoarseLODIndex, CoarseParticleRangeId, CoarseParticleRangeId);
-PRAGMA_ENABLE_DEPRECATION_WARNINGS
 	}
 
 	// LOD Switching
@@ -1114,9 +860,7 @@ PRAGMA_ENABLE_DEPRECATION_WARNINGS
 			NumActiveDynamicParticles = LODData[LODIndex]->NumDynamicParticles;
 
 			// Wrap new LOD based on previous LOD if possible (can only do 1 level LOD at a time, and if previous LOD exists)
-PRAGMA_DISABLE_DEPRECATION_WARNINGS  // TODO: CHAOS_IS_CLOTHINGSIMULATIONMESH_ABSTRACT
 			bNeedsReset = bNeedsReset || !Mesh->WrapDeformLOD(
-PRAGMA_ENABLE_DEPRECATION_WARNINGS
 				PrevLODIndex,
 				LODIndex,
 				Solver->GetNormals(PrevParticleRangeId),
@@ -1227,9 +971,7 @@ PRAGMA_ENABLE_DEPRECATION_WARNINGS
 			Solver->SetWindVelocity(GroupId, WindVelocity + Solver->GetWindVelocity());
 
 			// Update general solver properties
-			PRAGMA_DISABLE_DEPRECATION_WARNINGS  // TODO: CHAOS_IS_CLOTHINGSIMULATIONMESH_ABSTRACT
 			const Softs::FSolverReal MeshScale = Mesh->GetScale();
-			PRAGMA_ENABLE_DEPRECATION_WARNINGS
 
 			const FRealSingle DampingCoefficient = ConfigProperties.GetValue<float>(TEXT("DampingCoefficient"), ClothingSimulationClothDefault::DampingCoefficient);
 			const FRealSingle LocalDampingCoefficient = ConfigProperties.GetValue<float>(TEXT("LocalDampingCoefficient"));
@@ -1335,77 +1077,4 @@ TConstArrayView<Softs::FSolverReal> FClothingSimulationCloth::GetParticleInvMass
 	check(GetParticleRangeId(Solver, LODIndex) != INDEX_NONE);
 	return TConstArrayView<Softs::FSolverReal>(Solver->GetParticleInvMasses(GetParticleRangeId(Solver, LODIndex)), GetNumParticles(LODIndex));
 }
-
-void FClothingSimulationCloth::SetMaterialProperties(const TVec2<FRealSingle>& InEdgeStiffness, const TVec2<FRealSingle>& InBendingStiffness, const TVec2<FRealSingle>& InAreaStiffness)
-{
-	Config->GetProperties().SetWeightedFloatValue(TEXT("EdgeSpringStiffness"), FVector2f(InEdgeStiffness));
-	Config->GetProperties().SetWeightedFloatValue(TEXT("BendingSpringStiffness"), FVector2f(InBendingStiffness));
-	Config->GetProperties().SetWeightedFloatValue(TEXT("AreaSpringStiffness"), FVector2f(InAreaStiffness));
-}
-
-void FClothingSimulationCloth::SetLongRangeAttachmentProperties(const TVec2<FRealSingle>& InTetherStiffness, const TVec2<FRealSingle>& InTetherScale)
-{
-	Config->GetProperties().SetWeightedFloatValue(TEXT("TetherStiffness"), FVector2f(InTetherStiffness));
-	Config->GetProperties().SetWeightedFloatValue(TEXT("TetherScale"), FVector2f(InTetherScale));
-}
-
-void FClothingSimulationCloth::SetCollisionProperties(FRealSingle InCollisionThickness, FRealSingle InFrictionCoefficient, bool bInUseCCD, FRealSingle InSelfCollisionThickness)
-{
-	Config->GetProperties().SetValue(TEXT("CollisionThickness"), (float)InCollisionThickness);
-	Config->GetProperties().SetValue(TEXT("FrictionCoefficient"), (float)InFrictionCoefficient);
-	Config->GetProperties().SetValue(TEXT("UseCCD"), bInUseCCD);
-	Config->GetProperties().SetValue(TEXT("SelfCollisionThickness"), (float)InSelfCollisionThickness);
-}
-
-void FClothingSimulationCloth::SetBackstopProperties(bool bInEnableBackstop)
-{
-	Config->GetProperties().SetEnabled(TEXT("BackstopRadius"), bInEnableBackstop);  // BackstopRadius controls whether the backstop is enabled or not
-}
-
-void FClothingSimulationCloth::SetDampingProperties(FRealSingle InDampingCoefficient, FRealSingle InLocalDampingCoefficient)
-{
-	Config->GetProperties().SetValue(TEXT("DampingCoefficient"), (float)InDampingCoefficient);
-	Config->GetProperties().SetValue(TEXT("LocalDampingCoefficient"), (float)InLocalDampingCoefficient);
-}
-
-void FClothingSimulationCloth::SetAerodynamicsProperties(const TVec2<FRealSingle>& InDrag, const TVec2<FRealSingle>& InLift, FRealSingle InAirDensity, const FVec3& InWindVelocity)
-{
-	constexpr float WorldScale = 100.f;  // Interactor values are setup in engine scale, kg/cm^3 for air density and cm/s for wind velocity, but the properties are stored in kg/m^3 and m/s in the UI.
-	Config->GetProperties().SetWeightedFloatValue(TEXT("Drag"), FVector2f(InDrag));
-	Config->GetProperties().SetWeightedFloatValue(TEXT("Lift"), FVector2f(InLift));
-	Config->GetProperties().SetValue(TEXT("FluidDensity"), (float)InAirDensity * FMath::Cube(WorldScale));
-	Config->GetProperties().SetValue(TEXT("WindVelocity"), FVector3f(InWindVelocity) / WorldScale);
-}
-
-void FClothingSimulationCloth::SetPressureProperties(const TVec2<FRealSingle>& InPressure)
-{
-	Config->GetProperties().SetWeightedFloatValue(TEXT("Pressure"), FVector2f(InPressure));
-}
-
-void FClothingSimulationCloth::SetGravityProperties(FRealSingle InGravityScale, bool bInUseGravityOverride, const FVec3& InGravityOverride)
-{
-	Config->GetProperties().SetValue(TEXT("GravityScale"), (float)InGravityScale);
-	Config->GetProperties().SetValue(TEXT("UseGravityOverride"), bInUseGravityOverride);
-	Config->GetProperties().SetValue(TEXT("GravityOverride"), FVector3f(InGravityOverride));
-}
-
-void FClothingSimulationCloth::SetAnimDriveProperties(const TVec2<FRealSingle>& InAnimDriveStiffness, const TVec2<FRealSingle>& InAnimDriveDamping)
-{
-	Config->GetProperties().SetWeightedFloatValue(TEXT("AnimDriveStiffness"), FVector2f(InAnimDriveStiffness));
-	Config->GetProperties().SetWeightedFloatValue(TEXT("AnimDriveDamping"), FVector2f(InAnimDriveDamping));
-}
-
-void FClothingSimulationCloth::GetAnimDriveProperties(TVec2<FRealSingle>& OutAnimDriveStiffness, TVec2<FRealSingle>& OutAnimDriveDamping)
-{
-	OutAnimDriveStiffness = TVec2<FRealSingle>(Config->GetProperties().GetWeightedFloatValue(TEXT("AnimDriveStiffness")));
-	OutAnimDriveDamping = TVec2<FRealSingle>(Config->GetProperties().GetWeightedFloatValue(TEXT("AnimDriveDamping")));
-}
-
-void FClothingSimulationCloth::SetVelocityScaleProperties(const FVec3& InLinearVelocityScale, FRealSingle InAngularVelocityScale, FRealSingle InFictitiousAngularScale)
-{
-	Config->GetProperties().SetValue(TEXT("LinearVelocityScale"), FVector3f(InLinearVelocityScale));
-	Config->GetProperties().SetValue(TEXT("AngularVelocityScale"), (float)InAngularVelocityScale);
-	Config->GetProperties().SetValue(TEXT("FictitiousAngularScale"), (float)InFictitiousAngularScale);
-}
-
 }  // End namespace Chaos
diff --git a/Engine/Plugins/ChaosCloth/Source/ChaosCloth/Private/ChaosCloth/ChaosClothingSimulationCollider.cpp b/Engine/Plugins/ChaosCloth/Source/ChaosCloth/Private/ChaosCloth/ChaosClothingSimulationCollider.cpp
--- a/Engine/Plugins/ChaosCloth/Source/ChaosCloth/Private/ChaosCloth/ChaosClothingSimulationCollider.cpp
+++ b/Engine/Plugins/ChaosCloth/Source/ChaosCloth/Private/ChaosCloth/ChaosClothingSimulationCollider.cpp
@@ -450,23 +450,6 @@ FClothingSimulationCollider::FClothingSimulationCollider(const UPhysicsAsset* In
 	}
 }
 
-FClothingSimulationCollider::FClothingSimulationCollider(
-	const UClothingAssetCommon* InAsset,
-	const USkeletalMeshComponent* /*InSkeletalMeshComponent*/,
-	bool /*bInUseLODIndexOverride*/,
-	int32 /*InLODIndexOverride*/)
-	: PhysicsAsset(InAsset->PhysicsAsset)
-	, ReferenceSkeleton(&CastChecked<USkeletalMesh>(InAsset->GetOuter())->GetRefSkeleton())
-{
-	// Prepare LOD array
-	const int32 NumLODs = InAsset ? InAsset->LodData.Num() : 0;
-	const int32 NumLODData = (int32)ECollisionDataType::LODs + NumLODs;
-	LODData.Reserve(NumLODData);
-	for (int32 Index = 0; Index < NumLODData; ++Index)
-	{
-		LODData.Add(MakeUnique<FLODData>());
-	}
-}
 
 FClothingSimulationCollider::~FClothingSimulationCollider()
 {
diff --git a/Engine/Plugins/ChaosCloth/Source/ChaosCloth/Private/ChaosCloth/ChaosClothingSimulationConfig.cpp b/Engine/Plugins/ChaosCloth/Source/ChaosCloth/Private/ChaosCloth/ChaosClothingSimulationConfig.cpp
--- a/Engine/Plugins/ChaosCloth/Source/ChaosCloth/Private/ChaosCloth/ChaosClothingSimulationConfig.cpp
+++ b/Engine/Plugins/ChaosCloth/Source/ChaosCloth/Private/ChaosCloth/ChaosClothingSimulationConfig.cpp
@@ -11,13 +11,6 @@ namespace Chaos
 	{
 	}
 
-	FClothingSimulationConfig::FClothingSimulationConfig(const TSharedPtr<const FManagedArrayCollection>& InPropertyCollection)
-	{
-		PRAGMA_DISABLE_DEPRECATION_WARNINGS
-		Initialize(InPropertyCollection);
-		PRAGMA_ENABLE_DEPRECATION_WARNINGS
-	}
-
 	FClothingSimulationConfig::FClothingSimulationConfig(const TArray<TSharedPtr<const FManagedArrayCollection>>& InPropertyCollections)
 	{
 		Initialize(InPropertyCollections);
diff --git a/Engine/Plugins/ChaosCloth/Source/ChaosCloth/Private/ChaosCloth/ChaosClothingSimulationMesh.cpp b/Engine/Plugins/ChaosCloth/Source/ChaosCloth/Private/ChaosCloth/ChaosClothingSimulationMesh.cpp
--- a/Engine/Plugins/ChaosCloth/Source/ChaosCloth/Private/ChaosCloth/ChaosClothingSimulationMesh.cpp
+++ b/Engine/Plugins/ChaosCloth/Source/ChaosCloth/Private/ChaosCloth/ChaosClothingSimulationMesh.cpp
@@ -4,14 +4,14 @@
 #include "ChaosCloth/ChaosWeightMapTarget.h"
 #include "ChaosCloth/ChaosClothPrivate.h"
 #include "ClothingSimulation.h"
-#if !defined(CHAOS_IS_CLOTHINGSIMULATIONMESH_ABSTRACT) || !CHAOS_IS_CLOTHINGSIMULATIONMESH_ABSTRACT
+#if UE_ENABLE_INCLUDE_ORDER_DEPRECATED_IN_5_5
 #include "Components/SkeletalMeshComponent.h"
 #include "ClothingAsset.h"
-#else
-#include "Components/SkinnedMeshComponent.h"
 #endif
+#include "Components/SkinnedMeshComponent.h"
 #include "Containers/ArrayView.h"
 #include "Async/ParallelFor.h"
+#include "SkeletalMeshTypes.h"  // For FMeshToMeshVertData
 #if INTEL_ISPC
 #include "ChaosClothingSimulationMesh.ispc.generated.h"
 #endif
@@ -34,287 +34,15 @@ DECLARE_CYCLE_STAT(TEXT("Chaos Cloth Wrap Deform Cloth LOD"), STAT_ChaosClothWra
 namespace Chaos
 {
 
-#if !defined(CHAOS_IS_CLOTHINGSIMULATIONMESH_ABSTRACT) || !CHAOS_IS_CLOTHINGSIMULATIONMESH_ABSTRACT
-
-PRAGMA_DISABLE_DEPRECATION_WARNINGS  // TODO: CHAOS_IS_CLOTHINGSIMULATIONMESH_ABSTRACT
-
 FClothingSimulationMesh::FClothingSimulationMesh(const FString& InDebugName)
-	: Asset(nullptr)
-	, SkeletalMeshComponent(nullptr)
-#if !UE_BUILD_SHIPPING
-	, DebugName(InDebugName)
-#endif
-{
-}
-
-FClothingSimulationMesh::FClothingSimulationMesh(const UClothingAssetCommon* InAsset, const USkeletalMeshComponent* InSkeletalMeshComponent)
-	: Asset(InAsset)
-	, SkeletalMeshComponent(InSkeletalMeshComponent)
 #if !UE_BUILD_SHIPPING
-	, DebugName(InSkeletalMeshComponent->GetOwner() ?
-		FString::Format(TEXT("{0}|{1}"), { InSkeletalMeshComponent->GetOwner()->GetName(), InSkeletalMeshComponent->GetName() } ) :
-		InSkeletalMeshComponent->GetName())
+	: DebugName(InDebugName)
 #endif
 {
 }
 
 FClothingSimulationMesh::~FClothingSimulationMesh() = default;
 
-int32 FClothingSimulationMesh::GetNumLODs() const
-{
-	return Asset ? Asset->LodData.Num() : 0;
-}
-
-int32 FClothingSimulationMesh::GetLODIndex() const
-{
-	int32 LODIndex = INDEX_NONE;
-
-	if (Asset && SkeletalMeshComponent)
-	{
-		const int32 MeshLODIndex = SkeletalMeshComponent->GetPredictedLODLevel();
-		if (Asset->LodMap.IsValidIndex(MeshLODIndex))
-		{
-			const int32 ClothLODIndex = Asset->LodMap[MeshLODIndex];
-			if (Asset->LodData.IsValidIndex(ClothLODIndex))
-			{
-				LODIndex = ClothLODIndex;
-			}
-		}
-	}
-	return LODIndex;
-}
-
-bool FClothingSimulationMesh::IsValidLODIndex(int32 LODIndex) const
-{
-	return Asset && Asset->LodData.IsValidIndex(LODIndex);
-}
-
-int32 FClothingSimulationMesh::GetOwnerLODIndex(int32 LODIndex) const
-{
-	const int32 OwnerLODIndex = Asset ? Asset->LodMap.Find(LODIndex) : INDEX_NONE;
-	return OwnerLODIndex != INDEX_NONE ? OwnerLODIndex : 0;  // Safer to return the default LOD 0 than INDEX_NONE in this case
-}
-
-int32 FClothingSimulationMesh::GetNumPoints(int32 LODIndex) const
-{
-	return IsValidLODIndex(LODIndex) ? Asset->LodData[LODIndex].PhysicalMeshData.Vertices.Num() : 0;
-}
-
-int32 FClothingSimulationMesh::GetNumPatternPoints(int32 LODIndex) const
-{
-	return 0;
-}
-
-TConstArrayView<FVector3f> FClothingSimulationMesh::GetPositions(int32 LODIndex) const
-{
-	if (IsValidLODIndex(LODIndex))
-	{
-		const FClothLODDataCommon& ClothLODData = Asset->LodData[LODIndex];
-		const FClothPhysicalMeshData& ClothPhysicalMeshData = ClothLODData.PhysicalMeshData;
-		return TConstArrayView<FVector3f>(ClothPhysicalMeshData.Vertices);
-	}
-	return TConstArrayView<FVector3f>();
-}
-
-TConstArrayView<FVector2f> FClothingSimulationMesh::GetPatternPositions(int32 LODIndex) const
-{
-	return TConstArrayView<FVector2f>();
-}
-
-TConstArrayView<FVector3f> FClothingSimulationMesh::GetNormals(int32 LODIndex) const
-{
-	if (IsValidLODIndex(LODIndex))
-	{
-		const FClothLODDataCommon& ClothLODData = Asset->LodData[LODIndex];
-		const FClothPhysicalMeshData& ClothPhysicalMeshData = ClothLODData.PhysicalMeshData;
-		return TConstArrayView<FVector3f>(ClothPhysicalMeshData.Normals);
-	}
-	return TConstArrayView<FVector3f>();
-}
-
-TConstArrayView<uint32> FClothingSimulationMesh::GetIndices(int32 LODIndex) const
-{
-	return IsValidLODIndex(LODIndex) ?
-		TConstArrayView<uint32>(Asset->LodData[LODIndex].PhysicalMeshData.Indices) :
-		TConstArrayView<uint32>();
-}
-
-TConstArrayView<uint32> FClothingSimulationMesh::GetPatternIndices(int32 LODIndex) const
-{
-	return TConstArrayView<uint32>();
-}
-
-TConstArrayView<uint32> FClothingSimulationMesh::GetPatternToWeldedIndices(int32 LODIndex) const
-{
-	return TConstArrayView<uint32>();
-}
-
-TArray<FName> FClothingSimulationMesh::GetWeightMapNames(int32 LODIndex) const
-{
-	TArray<FName> WeightMapNames;
-	if (IsValidLODIndex(LODIndex))
-	{
-		// This must match the order of GetWeightMaps
-		const FClothLODDataCommon& ClothLODData = Asset->LodData[LODIndex];
-		const FClothPhysicalMeshData& ClothPhysicalMeshData = ClothLODData.PhysicalMeshData;
-		WeightMapNames.Reserve(ClothPhysicalMeshData.WeightMaps.Num());
-
-		const UEnum* const ChaosWeightMapTargetEnum = StaticEnum<EChaosWeightMapTarget>();
-		for (const TPair<uint32, FPointWeightMap>& TargetIDAndMap : ClothPhysicalMeshData.WeightMaps)
-		{
-			WeightMapNames.Add(FName(ChaosWeightMapTargetEnum->GetNameStringByValue(TargetIDAndMap.Get<0>())));
-		}
-	}
-
-	return WeightMapNames;
-}
-
-TMap<FString, int32> FClothingSimulationMesh::GetWeightMapIndices(int32 LODIndex) const
-{
-	TMap<FString, int32> WeightMapIndices;
-	const TArray<FName> WeightMapNames = GetWeightMapNames(LODIndex);
-	WeightMapIndices.Reserve(WeightMapNames.Num());
-	for (int32 WeightMapIndex = 0; WeightMapIndex < WeightMapNames.Num(); ++WeightMapIndex)
-	{
-		const FName& WeightMapName = WeightMapNames[WeightMapIndex];
-		WeightMapIndices.Emplace(WeightMapName.ToString(), WeightMapIndex);
-	}
-	return WeightMapIndices;
-}
-
-TArray<TConstArrayView<FRealSingle>> FClothingSimulationMesh::GetWeightMaps(int32 LODIndex) const
-{
-	TArray<TConstArrayView<FRealSingle>> WeightMaps;
-	if (IsValidLODIndex(LODIndex))
-	{
-		// This must match the order of GetWeightMapNames
-		const FClothLODDataCommon& ClothLODData = Asset->LodData[LODIndex];
-		const FClothPhysicalMeshData& ClothPhysicalMeshData = ClothLODData.PhysicalMeshData;
-
-		WeightMaps.Reserve(ClothPhysicalMeshData.WeightMaps.Num());
-		for (const TPair<uint32, FPointWeightMap>& TargetIDAndMap : ClothPhysicalMeshData.WeightMaps)
-		{
-			WeightMaps.Add(TargetIDAndMap.Get<1>().Values);
-		}
-	}
-	return WeightMaps;
-}
-
-TMap<FString, const TSet<int32>*> FClothingSimulationMesh::GetVertexSets(int32 LODIndex) const
-{
-	// Not supported
-	return TMap<FString, const TSet<int32>*>();
-}
-
-TMap<FString, const TSet<int32>*> FClothingSimulationMesh::GetFaceSets(int32 LODIndex) const
-{
-	// Not supported
-	return TMap<FString, const TSet<int32>*>();
-}
-
-TMap<FString, TConstArrayView<int32>> FClothingSimulationMesh::GetFaceIntMaps(int32 LODIndex) const
-{
-	// Not supported
-	return TMap<FString, TConstArrayView<int32>>();
-}
-
-TArray<TConstArrayView<TTuple<int32, int32, float>>> FClothingSimulationMesh::GetTethers(int32 LODIndex, bool bUseGeodesicTethers) const
-{
-	TArray<TConstArrayView<TTuple<int32, int32, float>>> Tethers;
-	if (IsValidLODIndex(LODIndex))
-	{
-		const FClothLODDataCommon& ClothLODData = Asset->LodData[LODIndex];
-		const FClothPhysicalMeshData& ClothPhysicalMeshData = ClothLODData.PhysicalMeshData;
-		const FClothTetherData& ClothTetherData = bUseGeodesicTethers ? ClothPhysicalMeshData.GeodesicTethers : ClothPhysicalMeshData.EuclideanTethers;
-		
-		const int32 NumTetherBatches = ClothTetherData.Tethers.Num();
-		Tethers.Reserve(NumTetherBatches);
-		for (int32 Index = 0; Index < NumTetherBatches; ++Index)
-		{
-			Tethers.Emplace(TConstArrayView<TTuple<int32, int32, float>>(ClothTetherData.Tethers[Index]));
-		}
-	}
-	return Tethers;
-}
-
-int32 FClothingSimulationMesh::GetReferenceBoneIndex() const
-{
-	return Asset ? Asset->ReferenceBoneIndex : INDEX_NONE;
-}
-
-FTransform FClothingSimulationMesh::GetReferenceBoneTransform() const
-{
-	if (const FClothingSimulationContextCommon* const Context = GetContext())
-	{
-		const int32 ReferenceBoneIndex = GetReferenceBoneIndex();
-		const TArray<FTransform>& BoneTransforms = Context->BoneTransforms;
-
-		return BoneTransforms.IsValidIndex(ReferenceBoneIndex) ?
-			BoneTransforms[ReferenceBoneIndex] * Context->ComponentToWorld :
-			Context->ComponentToWorld;
-	}
-	return FTransform::Identity;
-}
-
-const TArray<FTransform>& FClothingSimulationMesh::GetBoneTransforms() const
-{
-	return GetContext()->BoneTransforms;
-}
-
-const FTransform& FClothingSimulationMesh::GetComponentToWorldTransform() const
-{
-	return GetContext() ? GetContext()->ComponentToWorld : FTransform::Identity;
-}
-
-const TArray<FMatrix44f>& FClothingSimulationMesh::GetRefToLocalMatrices() const
-{
-	static TArray<FMatrix44f> EmptyArray;
-	return GetContext() ? GetContext()->RefToLocals : EmptyArray;
-}
-
-TConstArrayView<int32> FClothingSimulationMesh::GetBoneMap() const
-{
-	return Asset ? Asset->UsedBoneIndices : TConstArrayView<int32>();
-}
-
-TConstArrayView<FClothVertBoneData> FClothingSimulationMesh::GetBoneData(int32 LODIndex) const
-{
-	if (IsValidLODIndex(LODIndex))
-	{
-		const FClothLODDataCommon& ClothLODData = Asset->LodData[LODIndex];
-		const FClothPhysicalMeshData& ClothPhysicalMeshData = ClothLODData.PhysicalMeshData;
-		UE_CLOG(ClothPhysicalMeshData.MaxBoneWeights > 12, LogChaosCloth, Warning, TEXT("The cloth physics mesh skinning code can't cope with more than 12 bone influences."));
-		return ClothPhysicalMeshData.BoneData;
-	}
-	static TArray<FClothVertBoneData> EmptyArray;
-	return EmptyArray;
-}
-
-TConstArrayView<FMeshToMeshVertData> FClothingSimulationMesh::GetTransitionUpSkinData(int32 LODIndex) const
-{
-	return IsValidLODIndex(LODIndex) ? TConstArrayView<FMeshToMeshVertData>(Asset->LodData[LODIndex].TransitionUpSkinData) : TConstArrayView<FMeshToMeshVertData>();
-}
-
-TConstArrayView<FMeshToMeshVertData> FClothingSimulationMesh::GetTransitionDownSkinData(int32 LODIndex) const
-{
-	return IsValidLODIndex(LODIndex) ? TConstArrayView<FMeshToMeshVertData>(Asset->LodData[LODIndex].TransitionDownSkinData) : TConstArrayView<FMeshToMeshVertData>();
-}
-
-const FClothingSimulationContextCommon* FClothingSimulationMesh::GetContext() const
-{
-	return GetSkeletalMeshComponent() ?
-		static_cast<const FClothingSimulationContextCommon*>(GetSkeletalMeshComponent()->GetClothingSimulationContext()) :
-		nullptr;
-}
-
-PRAGMA_ENABLE_DEPRECATION_WARNINGS
-
-#else
-FClothingSimulationMesh::FClothingSimulationMesh() = default;
-
-FClothingSimulationMesh::~FClothingSimulationMesh() = default;
-#endif
-
 Softs::FSolverReal FClothingSimulationMesh::GetScale() const
 {
 	return GetComponentToWorldTransform().GetScale3D().GetMax();
diff --git a/Engine/Plugins/ChaosCloth/Source/ChaosCloth/Private/ChaosCloth/ChaosClothingSimulationSkeletalMesh.cpp b/Engine/Plugins/ChaosCloth/Source/ChaosCloth/Private/ChaosCloth/ChaosClothingSimulationSkeletalMesh.cpp
--- a/Engine/Plugins/ChaosCloth/Source/ChaosCloth/Private/ChaosCloth/ChaosClothingSimulationSkeletalMesh.cpp
+++ b/Engine/Plugins/ChaosCloth/Source/ChaosCloth/Private/ChaosCloth/ChaosClothingSimulationSkeletalMesh.cpp
@@ -9,11 +9,9 @@
 namespace Chaos
 {
 	FClothingSimulationSkeletalMesh::FClothingSimulationSkeletalMesh(const UClothingAssetCommon* InAsset, const USkeletalMeshComponent* InSkeletalMeshComponent)
-PRAGMA_DISABLE_DEPRECATION_WARNINGS  // TODO: CHAOS_IS_CLOTHINGSIMULATIONMESH_ABSTRACT
 		: FClothingSimulationMesh(InSkeletalMeshComponent->GetOwner() ?
 				FString::Format(TEXT("{0}|{1}"), { InSkeletalMeshComponent->GetOwner()->GetName(), InSkeletalMeshComponent->GetName() }) :
 				InSkeletalMeshComponent->GetName())
-PRAGMA_ENABLE_DEPRECATION_WARNINGS
 		, Asset(InAsset)
 		, SkeletalMeshComponent(InSkeletalMeshComponent)
 	{
@@ -70,7 +68,7 @@ PRAGMA_ENABLE_DEPRECATION_WARNINGS
 		{
 			const FClothLODDataCommon& ClothLODData = Asset->LodData[LODIndex];
 			const FClothPhysicalMeshData& ClothPhysicalMeshData = ClothLODData.PhysicalMeshData;
-			return ClothPhysicalMeshData.Vertices;
+			return TConstArrayView<FVector3f>(ClothPhysicalMeshData.Vertices);
 		}
 		return TConstArrayView<FVector3f>();
 	}
@@ -86,7 +84,7 @@ PRAGMA_ENABLE_DEPRECATION_WARNINGS
 		{
 			const FClothLODDataCommon& ClothLODData = Asset->LodData[LODIndex];
 			const FClothPhysicalMeshData& ClothPhysicalMeshData = ClothLODData.PhysicalMeshData;
-			return ClothPhysicalMeshData.Normals;
+			return TConstArrayView<FVector3f>(ClothPhysicalMeshData.Normals);
 		}
 		return TConstArrayView<FVector3f>();
 	}
diff --git a/Engine/Plugins/ChaosCloth/Source/ChaosCloth/Private/ChaosCloth/ChaosClothingSimulationSolver.cpp b/Engine/Plugins/ChaosCloth/Source/ChaosCloth/Private/ChaosCloth/ChaosClothingSimulationSolver.cpp
--- a/Engine/Plugins/ChaosCloth/Source/ChaosCloth/Private/ChaosCloth/ChaosClothingSimulationSolver.cpp
+++ b/Engine/Plugins/ChaosCloth/Source/ChaosCloth/Private/ChaosCloth/ChaosClothingSimulationSolver.cpp
@@ -520,9 +520,7 @@ void FClothingSimulationSolver::SetConfig(FClothingSimulationConfig* InConfig)
 	{
 		// Create a default empty config object for coherence
 		PropertyCollection = MakeShared<FManagedArrayCollection>();
-		PRAGMA_DISABLE_DEPRECATION_WARNINGS
-		Config = new FClothingSimulationConfig(PropertyCollection);
-		PRAGMA_ENABLE_DEPRECATION_WARNINGS
+		Config = new FClothingSimulationConfig({ PropertyCollection });
 	}
 }
 
@@ -1296,14 +1294,6 @@ void FClothingSimulationSolver::SetWindVelocity(const TVec3<FRealSingle>& InWind
 	LegacyWindAdaption = InLegacyWindAdaption;
 }
 
-void FClothingSimulationSolver::SetNumIterations(int32 InNumIterations)
-{
-	if (ensure(Config))
-	{
-		Config->GetProperties(SolverLOD).SetValue(TEXT("NumIterations"), InNumIterations);
-	}
-}
-
 int32 FClothingSimulationSolver::GetNumIterations() const
 {
 	return Config ?
@@ -1311,14 +1301,6 @@ int32 FClothingSimulationSolver::GetNumIterations() const
 		ClothingSimulationSolverDefault::NumIterations;
 }
 
-void FClothingSimulationSolver::SetMaxNumIterations(int32 InMaxNumIterations)
-{
-	if (ensure(Config))
-	{
-		Config->GetProperties(SolverLOD).SetValue(TEXT("MaxNumIterations"), InMaxNumIterations);
-	}
-}
-
 int32 FClothingSimulationSolver::GetMaxNumIterations() const
 {
 	return Config ?
@@ -1326,14 +1308,6 @@ int32 FClothingSimulationSolver::GetMaxNumIterations() const
 		ClothingSimulationSolverDefault::MaxNumIterations;
 }
 
-void FClothingSimulationSolver::SetNumSubsteps(int32 InNumSubsteps)
-{
-	if (ensure(Config))
-	{
-		Config->GetProperties(SolverLOD).SetValue(TEXT("NumSubsteps"), InNumSubsteps);
-	}
-}
-
 int32 FClothingSimulationSolver::GetNumSubsteps() const
 {
 	return Config ?
@@ -2068,26 +2042,6 @@ void FClothingSimulationSolver::UpdateFromCache(const FClothingSimulationCacheDa
 
 }
 
-void FClothingSimulationSolver::UpdateFromCache(const TArray<FVector>& CachedPositions, const TArray<FVector>& CachedVelocities) 
-{
-	Chaos::Softs::FSolverParticles& SolverParticles = Evolution ? Evolution->GetParticles() : PBDEvolution->GetParticles();
-	const int32 NumParticles = GetNumParticles();
-	const bool bHasVelocity = CachedVelocities.Num() > 0;
-	if(CachedPositions.Num() == NumParticles)
-	{
-		for(int32 ParticleIndex = 0; ParticleIndex < NumParticles; ++ParticleIndex)
-		{
-			SolverParticles.X(ParticleIndex) = CachedPositions[ParticleIndex];
-			SolverParticles.V(ParticleIndex) = bHasVelocity ? CachedVelocities[ParticleIndex] : FVector::ZeroVector;
-		}
-	}
-	PhysicsParallelFor(Cloths.Num(), [this](int32 ClothIndex)
-	{
-		FClothingSimulationCloth* const Cloth = Cloths[ClothIndex];
-		Cloth->PostUpdate(this);
-	}, /*bForceSingleThreaded =*/ !bClothSolverParallelClothPostUpdate);
-}
-
 int32 FClothingSimulationSolver::GetNumUsedIterations() const
 {
 	return Evolution ? Evolution->GetNumUsedIterations() : PBDEvolution->GetIterations();
diff --git a/Engine/Plugins/ChaosCloth/Source/ChaosCloth/Public/ChaosCloth/ChaosClothConstraints.h b/Engine/Plugins/ChaosCloth/Source/ChaosCloth/Public/ChaosCloth/ChaosClothConstraints.h
--- a/Engine/Plugins/ChaosCloth/Source/ChaosCloth/Public/ChaosCloth/ChaosClothConstraints.h
+++ b/Engine/Plugins/ChaosCloth/Source/ChaosCloth/Public/ChaosCloth/ChaosClothConstraints.h
@@ -18,11 +18,6 @@ namespace Chaos
 	class FClothConstraints final
 	{
 	public:
-		UE_DEPRECATED(5.3, "ETetherMode has been replaced with bUseGeodesicTethers.")
-PRAGMA_DISABLE_DEPRECATION_WARNINGS
-		typedef Softs::FPBDLongRangeConstraints::EMode ETetherMode;
-PRAGMA_ENABLE_DEPRECATION_WARNINGS
-
 		FClothConstraints();
 		~FClothConstraints();
 
@@ -86,15 +81,6 @@ PRAGMA_ENABLE_DEPRECATION_WARNINGS
 			Softs::FSolverReal MeshScale,
 			bool bEnabled);
 
-		UE_DEPRECATED(5.3, "Use AddRules() with WeightMaps, VertexSets, FaceSets, FaceIntMap, and optional PatternData instead.")
-		void AddRules(
-			const Softs::FCollectionPropertyConstFacade& ConfigProperties,
-			const FTriangleMesh& TriangleMesh,
-			const TArray<TConstArrayView<FRealSingle>>& WeightMapArray,
-			const TArray<TConstArrayView<TTuple<int32, int32, FRealSingle>>>& Tethers,
-			Softs::FSolverReal MeshScale,
-			bool bEnabled);
-
 		void Update(
 			const Softs::FCollectionPropertyConstFacade& ConfigProperties,
 			const TMap<FString, TConstArrayView<FRealSingle>>& WeightMaps,
@@ -111,12 +97,6 @@ PRAGMA_ENABLE_DEPRECATION_WARNINGS
 			Softs::FSolverReal MeshScale,
 			Softs::FSolverReal MaxDistancesScale = (Softs::FSolverReal)1.);
 
-		UE_DEPRECATED(5.3, "Use Update() with WeightMaps, VertexSets, FaceSets, and FaceIntMaps instead.")
-		void Update(
-			const Softs::FCollectionPropertyConstFacade& ConfigProperties,
-			Softs::FSolverReal MeshScale,
-			Softs::FSolverReal MaxDistancesScale = (Softs::FSolverReal)1.);
-
 		// NOTE: this only does something if using the PBDSolver. Force-based solver constraints are activated automatically
 		// when activating a particle range.
 		void Enable(bool bEnable);
diff --git a/Engine/Plugins/ChaosCloth/Source/ChaosCloth/Public/ChaosCloth/ChaosClothingSimulation.h b/Engine/Plugins/ChaosCloth/Source/ChaosCloth/Public/ChaosCloth/ChaosClothingSimulation.h
--- a/Engine/Plugins/ChaosCloth/Source/ChaosCloth/Public/ChaosCloth/ChaosClothingSimulation.h
+++ b/Engine/Plugins/ChaosCloth/Source/ChaosCloth/Public/ChaosCloth/ChaosClothingSimulation.h
@@ -108,8 +108,6 @@ namespace Chaos
 		void DebugDrawBendingConstraint(FPrimitiveDrawInterface* PDI = nullptr) const { Visualization.DrawBendingConstraint(PDI); }
 		void DebugDrawLongRangeConstraint(FPrimitiveDrawInterface* PDI = nullptr) const { Visualization.DrawLongRangeConstraint(PDI); }
 		void DebugDrawWindAndPressureForces(FPrimitiveDrawInterface* PDI = nullptr) const { Visualization.DrawWindAndPressureForces(PDI); }
-		UE_DEPRECATED(5.1, "DebugDrawWindForces has been renamed DebugDrawWindAndPressureForces.")
-		void DebugDrawWindForces(FPrimitiveDrawInterface* PDI = nullptr) const { Visualization.DrawWindAndPressureForces(PDI); }
 		void DebugDrawLocalSpace(FPrimitiveDrawInterface* PDI = nullptr) const { Visualization.DrawLocalSpace(PDI); }
 		void DebugDrawSelfCollision(FPrimitiveDrawInterface* PDI = nullptr) const { Visualization.DrawSelfCollision(PDI); }
 		void DebugDrawSelfIntersection(FPrimitiveDrawInterface* PDI = nullptr) const { Visualization.DrawSelfIntersection(PDI); }
@@ -129,9 +127,7 @@ namespace Chaos
 
 		// Simulation objects
 		TUniquePtr<FClothingSimulationSolver> Solver;  // Default solver
-PRAGMA_DISABLE_DEPRECATION_WARNINGS  // TODO: CHAOS_IS_CLOTHINGSIMULATIONMESH_ABSTRACT
 		TArray<TUniquePtr<FClothingSimulationMesh>> Meshes;
-PRAGMA_ENABLE_DEPRECATION_WARNINGS
 		TArray<TUniquePtr<FClothingSimulationCloth>> Cloths;
 		TArray<TUniquePtr<FClothingSimulationCollider>> Colliders;
 		TArray<TUniquePtr<FClothingSimulationConfig>> Configs;
diff --git a/Engine/Plugins/ChaosCloth/Source/ChaosCloth/Public/ChaosCloth/ChaosClothingSimulationCloth.h b/Engine/Plugins/ChaosCloth/Source/ChaosCloth/Public/ChaosCloth/ChaosClothingSimulationCloth.h
--- a/Engine/Plugins/ChaosCloth/Source/ChaosCloth/Public/ChaosCloth/ChaosClothingSimulationCloth.h
+++ b/Engine/Plugins/ChaosCloth/Source/ChaosCloth/Public/ChaosCloth/ChaosClothingSimulationCloth.h
@@ -28,73 +28,12 @@ namespace Chaos
 			Density
 		};
 
-		UE_DEPRECATED(5.3, "ETetherMode has been replaced with bUseGeodesicTethers.")
-PRAGMA_DISABLE_DEPRECATION_WARNINGS
-		typedef FClothConstraints::ETetherMode ETetherMode;
-CHAOSCLOTH_API PRAGMA_ENABLE_DEPRECATION_WARNINGS
-
-		FClothingSimulationCloth(
+		CHAOSCLOTH_API FClothingSimulationCloth(
 			FClothingSimulationConfig* InConfig,
-PRAGMA_DISABLE_DEPRECATION_WARNINGS  // TODO: CHAOS_IS_CLOTHINGSIMULATIONMESH_ABSTRACT
 			FClothingSimulationMesh* InMesh,
-PRAGMA_ENABLE_DEPRECATION_WARNINGS
 			TArray<FClothingSimulationCollider*>&& InColliders,
 			uint32 InGroupId);
 
-		UE_DEPRECATED(5.2, "Use config based constructor instead.")
-		FClothingSimulationCloth(
-PRAGMA_DISABLE_DEPRECATION_WARNINGS  // TODO: CHAOS_IS_CLOTHINGSIMULATIONMESH_ABSTRACT
-			FClothingSimulationMesh* InMesh,
-PRAGMA_ENABLE_DEPRECATION_WARNINGS
-			TArray<FClothingSimulationCollider*>&& InColliders,
-			uint32 InGroupId,
-			EMassMode InMassMode,
-			FRealSingle InMassValue,
-			FRealSingle InMinPerParticleMass,
-			const TVec2<FRealSingle>& InEdgeStiffness,
-			const TVec2<FRealSingle>& InBendingStiffness,
-			FRealSingle InBucklingRatio,
-			const TVec2<FRealSingle>& InBucklingStiffness,
-			bool bInUseBendingElements,
-			const TVec2<FRealSingle>& InAreaStiffness,
-			FRealSingle InVolumeStiffness,
-			bool bInUseThinShellVolumeConstraints,
-			const TVec2<FRealSingle>& InTetherStiffness,
-			const TVec2<FRealSingle>& InTetherScale,
-PRAGMA_DISABLE_DEPRECATION_WARNINGS
-			ETetherMode InTetherMode,
-PRAGMA_ENABLE_DEPRECATION_WARNINGS
-			FRealSingle InMaxDistancesMultiplier,
-			const TVec2<FRealSingle>& InAnimDriveStiffness,
-			const TVec2<FRealSingle>& InAnimDriveDamping,
-			FRealSingle InShapeTargetStiffness,
-			bool bInUseXPBDEdgeConstraints,
-			bool bInUseXPBDBendingConstraints,
-			bool bInUseXPBDAreaConstraints,
-			FRealSingle InGravityScale,
-			bool bIsGravityOverridden,
-			const TVec3<FRealSingle>& InGravityOverride,
-			const TVec3<FRealSingle>& InLinearVelocityScale,
-			FRealSingle InAngularVelocityScale,
-			FRealSingle InFictitiousAngularScale,
-			const TVec2<FRealSingle>& InDrag,
-			const TVec2<FRealSingle>& InLift,
-			bool bInUsePointBasedWindModel,
-			const TVec2<FRealSingle>& InPressure,
-			FRealSingle InDampingCoefficient,
-			FRealSingle InLocalDampingCoefficient,
-			FRealSingle InCollisionThickness,
-			FRealSingle InFrictionCoefficient,
-			bool bInUseCCD,
-			bool bInUseSelfCollisions,
-			FRealSingle InSelfCollisionThickness,
-			FRealSingle InSelfCollisionFrictionCoefficient,
-			bool bInUseSelfIntersections,
-			bool bInUseLegacyBackstop,
-			bool bInUseLODIndexOverride,
-			int32 InLODIndexOverride,
-			const TVec2<FRealSingle>& EdgeDampingRatio = TVec2<FRealSingle>(0.f),
-			const TVec2<FRealSingle>& BendDampingRatio = TVec2<FRealSingle>(0.f));
 		CHAOSCLOTH_API ~FClothingSimulationCloth();
 
 		FClothingSimulationCloth(const FClothingSimulationCloth&) = delete;
@@ -110,38 +49,14 @@ PRAGMA_ENABLE_DEPRECATION_WARNINGS
 
 		// ---- Animatable property setters ----
 		void SetMaxDistancesMultiplier(FRealSingle InMaxDistancesMultiplier) { MaxDistancesMultiplier = InMaxDistancesMultiplier; }
-		UE_DEPRECATED(5.3, "Set properties directly through FClothingSimulationConfig")
-		CHAOSCLOTH_API void SetMaterialProperties(const TVec2<FRealSingle>& InEdgeStiffness, const TVec2<FRealSingle>& InBendingStiffness, const TVec2<FRealSingle>& InAreaStiffness);
-		UE_DEPRECATED(5.3, "Set properties directly through FClothingSimulationConfig")
-		CHAOSCLOTH_API void SetLongRangeAttachmentProperties(const TVec2<FRealSingle>& InTetherStiffness, const TVec2<FRealSingle>& InTetherScale);
-		UE_DEPRECATED(5.3, "Set properties directly through FClothingSimulationConfig")
-		CHAOSCLOTH_API void SetCollisionProperties(FRealSingle InCollisionThickness, FRealSingle InFrictionCoefficient, bool bInUseCCD, FRealSingle InSelfCollisionThickness);
-		UE_DEPRECATED(5.3, "Set properties directly through FClothingSimulationConfig")
-		CHAOSCLOTH_API void SetBackstopProperties(bool bInEnableBackstop);
-		UE_DEPRECATED(5.3, "Set properties directly through FClothingSimulationConfig")
-		CHAOSCLOTH_API void SetDampingProperties(FRealSingle InDampingCoefficient, FRealSingle InLocalDampingCoefficient = 0.f);
-		UE_DEPRECATED(5.3, "Set properties directly through FClothingSimulationConfig")
-		CHAOSCLOTH_API void SetAerodynamicsProperties(const TVec2<FRealSingle>& InDrag, const TVec2<FRealSingle>& InLift, FRealSingle InAirDensity, const FVec3& InWindVelocity);  // AirDensity is here in kg/cm^3 for legacy reason (kg/m^3 in UI)
-		UE_DEPRECATED(5.3, "Set properties directly through FClothingSimulationConfig")
-		CHAOSCLOTH_API void SetPressureProperties(const TVec2<FRealSingle>& InPressure);
-		UE_DEPRECATED(5.3, "Set properties directly through FClothingSimulationConfig")
-		CHAOSCLOTH_API void SetGravityProperties(FRealSingle InGravityScale, bool bInUseGravityOverride, const FVec3& InGravityOverride);
-		UE_DEPRECATED(5.3, "Set properties directly through FClothingSimulationConfig")
-		CHAOSCLOTH_API void SetAnimDriveProperties(const TVec2<FRealSingle>& InAnimDriveStiffness, const TVec2<FRealSingle>& InAnimDriveDamping);
-		UE_DEPRECATED(5.3, "Get properties directly through FClothingSimulationConfig")
-		CHAOSCLOTH_API void GetAnimDriveProperties(TVec2<FRealSingle>& OutAnimDriveStiffness, TVec2<FRealSingle>& OutAnimDriveDamping);
-		UE_DEPRECATED(5.3, "Set properties directly through FClothingSimulationConfig")
-		CHAOSCLOTH_API void SetVelocityScaleProperties(const FVec3& InLinearVelocityScale, FRealSingle InAngularVelocityScale, FRealSingle InFictitiousAngularScale);
 
 		void Reset() { bNeedsReset = true; }
 		void Teleport() { bNeedsTeleport = true; }
 		// ---- End of the animatable property setters ----
 
 		// ---- Node property getters/setters
-		PRAGMA_DISABLE_DEPRECATION_WARNINGS  // TODO: CHAOS_IS_CLOTHINGSIMULATIONMESH_ABSTRACT
 		FClothingSimulationMesh* GetMesh() const { return Mesh; }
 		CHAOSCLOTH_API void SetMesh(FClothingSimulationMesh* InMesh);
-		PRAGMA_ENABLE_DEPRECATION_WARNINGS
 
 		FClothingSimulationConfig* GetConfig() const { return Config; }
 		CHAOSCLOTH_API void SetConfig(FClothingSimulationConfig* InConfig);
@@ -189,8 +104,6 @@ PRAGMA_ENABLE_DEPRECATION_WARNINGS
 		CHAOSCLOTH_API TConstArrayView<int32> GetFaceIntMapByName(const FClothingSimulationSolver* Solver, const FString& Name) const;
 		// Return the face int map of the specified property name if it exists and is available on the current LOD, or an empty array view otherwise.
 		CHAOSCLOTH_API TConstArrayView<int32> GetFaceIntMapByProperty(const FClothingSimulationSolver* Solver, const FString& Property) const;
-		UE_DEPRECATED(5.3, "Returns an empty array from 5.3. Update your code with GetWeightMapByName and GetWeightMapByProperty to return the current LOD weight map instead.")
-		CHAOSCLOTH_API const TArray<TConstArrayView<FRealSingle>>& GetWeightMaps(const FClothingSimulationSolver* Solver) const;
 		// Return the current LOD tethers.
 		CHAOSCLOTH_API const TArray<TConstArrayView<TTuple<int32, int32, float>>>& GetTethers(const FClothingSimulationSolver* Solver) const;
 		// Return the reference bone index for this cloth.
@@ -218,9 +131,7 @@ PRAGMA_ENABLE_DEPRECATION_WARNINGS
 		struct FLODData;
 
 		// Cloth parameters
-PRAGMA_DISABLE_DEPRECATION_WARNINGS  // TODO: CHAOS_IS_CLOTHINGSIMULATIONMESH_ABSTRACT
 		FClothingSimulationMesh* Mesh = nullptr;
-PRAGMA_ENABLE_DEPRECATION_WARNINGS
 		FClothingSimulationConfig* Config = nullptr;
 		TArray<FClothingSimulationCollider*> Colliders;
 		uint32 GroupId = 0;
diff --git a/Engine/Plugins/ChaosCloth/Source/ChaosCloth/Public/ChaosCloth/ChaosClothingSimulationCollider.h b/Engine/Plugins/ChaosCloth/Source/ChaosCloth/Public/ChaosCloth/ChaosClothingSimulationCollider.h
--- a/Engine/Plugins/ChaosCloth/Source/ChaosCloth/Public/ChaosCloth/ChaosClothingSimulationCollider.h
+++ b/Engine/Plugins/ChaosCloth/Source/ChaosCloth/Public/ChaosCloth/ChaosClothingSimulationCollider.h
@@ -46,14 +46,6 @@ namespace Chaos
 		};
 
 		CHAOSCLOTH_API FClothingSimulationCollider(const UPhysicsAsset* InPhysicsAsset, const FReferenceSkeleton* InReferenceSkeleton);
-
-		UE_DEPRECATED(5.2, "Use FClothingSimulationCollider(const UPhysicsAsset* InPhysicsAsset, const FReferenceSkeleton* InReferenceSkeleton) instead.")
-		CHAOSCLOTH_API FClothingSimulationCollider(
-			const UClothingAssetCommon* InAsset,  // Cloth asset for collision data, can be nullptr
-			const USkeletalMeshComponent* InSkeletalMeshComponent,  // For asset LODs management, can be nullptr
-			bool bInUseLODIndexOverride,
-			int32 InLODIndexOverride);
-
 		CHAOSCLOTH_API ~FClothingSimulationCollider();
 
 		FClothingSimulationCollider(const FClothingSimulationCollider&) = delete;
diff --git a/Engine/Plugins/ChaosCloth/Source/ChaosCloth/Public/ChaosCloth/ChaosClothingSimulationConfig.h b/Engine/Plugins/ChaosCloth/Source/ChaosCloth/Public/ChaosCloth/ChaosClothingSimulationConfig.h
--- a/Engine/Plugins/ChaosCloth/Source/ChaosCloth/Public/ChaosCloth/ChaosClothingSimulationConfig.h
+++ b/Engine/Plugins/ChaosCloth/Source/ChaosCloth/Public/ChaosCloth/ChaosClothingSimulationConfig.h
@@ -22,8 +22,6 @@ namespace Chaos
 	{
 	public:
 		CHAOSCLOTH_API FClothingSimulationConfig();
-		UE_DEPRECATED(5.3, "Use TArray version of this constructor.")
-		CHAOSCLOTH_API FClothingSimulationConfig(const TSharedPtr<const FManagedArrayCollection>& InPropertyCollection);
 		CHAOSCLOTH_API FClothingSimulationConfig(const TArray<TSharedPtr<const FManagedArrayCollection>>& InPropertyCollections);
 
 		CHAOSCLOTH_API ~FClothingSimulationConfig();
@@ -41,14 +39,6 @@ namespace Chaos
 		 */
 		CHAOSCLOTH_API void Initialize(const UChaosClothConfig* ClothConfig, const UChaosClothSharedSimConfig* ClothSharedConfig, bool bUseLegacyConfig = false);
 
-		/** Initialize config from a property collection. */
-		UE_DEPRECATED(5.3, "Use TArray version of Initialize.")
-		CHAOSCLOTH_API void Initialize(const TSharedPtr<const FManagedArrayCollection>& InPropertyCollection)
-		{
-			Initialize(TArray<TSharedPtr<const FManagedArrayCollection>>({ InPropertyCollection }));
-			bIsLegacySingleLOD = true;
-		}
-
 		/** Initialize config from an array of property collections (one per LOD). */
 		CHAOSCLOTH_API void Initialize(const TArray<TSharedPtr<const FManagedArrayCollection>>& InPropertyCollections);
 
diff --git a/Engine/Plugins/ChaosCloth/Source/ChaosCloth/Public/ChaosCloth/ChaosClothingSimulationMesh.h b/Engine/Plugins/ChaosCloth/Source/ChaosCloth/Public/ChaosCloth/ChaosClothingSimulationMesh.h
--- a/Engine/Plugins/ChaosCloth/Source/ChaosCloth/Public/ChaosCloth/ChaosClothingSimulationMesh.h
+++ b/Engine/Plugins/ChaosCloth/Source/ChaosCloth/Public/ChaosCloth/ChaosClothingSimulationMesh.h
@@ -6,12 +6,10 @@
 #include "Containers/ContainersFwd.h"
 #include "ClothVertBoneData.h"
 
-//#define CHAOS_IS_CLOTHINGSIMULATIONMESH_ABSTRACT 1  // For testing only. TODO: Once FClothingSimulationMesh is abstract from 5.4, remove all lines where this is referenced
-
 struct FMeshToMeshVertData;
 struct FClothVertBoneData;
 
-#if !defined(CHAOS_IS_CLOTHINGSIMULATIONMESH_ABSTRACT) || !CHAOS_IS_CLOTHINGSIMULATIONMESH_ABSTRACT
+#if UE_ENABLE_INCLUDE_ORDER_DEPRECATED_IN_5_5
 class USkeletalMeshComponent;
 class UClothingAssetCommon;
 class FClothingSimulationContextCommon;
@@ -22,12 +20,7 @@ namespace Chaos
 	class FClothingSimulationSolver;
 
 	// Mesh simulation node
-#if !defined(CHAOS_IS_CLOTHINGSIMULATIONMESH_ABSTRACT) || !CHAOS_IS_CLOTHINGSIMULATIONMESH_ABSTRACT
-	class UE_DEPRECATED(5.2, "This class is becoming an abstract interface, do not use directly. Use FClothingSimulationSkeletalMesh instead.")
-		CHAOSCLOTH_API FClothingSimulationMesh
-#else
 	class CHAOSCLOTH_API FClothingSimulationMesh
-#endif
 	{
 	public:
 		explicit FClothingSimulationMesh(const FString& InDebugName);
@@ -45,47 +38,6 @@ namespace Chaos
 		const FString& GetDebugName() const { return FText::GetEmpty().ToString(); }
 #endif
 
-#if !defined(CHAOS_IS_CLOTHINGSIMULATIONMESH_ABSTRACT) || !CHAOS_IS_CLOTHINGSIMULATIONMESH_ABSTRACT
-		UE_DEPRECATED(5.2, "Use the FClothingSimulationSkeletalMesh constructor instead.")
-		FClothingSimulationMesh(const UClothingAssetCommon* InAsset, const USkeletalMeshComponent* InSkeletalMeshComponent);
-		UE_DEPRECATED(5.2, "Use the FClothingSimulationSkeletalMesh::GetAsset() instead.")
-		virtual const UClothingAssetCommon* GetAsset() const { return Asset; }
-		UE_DEPRECATED(5.2, "Use the FClothingSimulationSkeletalMesh::GetSkeletalMeshComponent() instead.")
-		virtual const USkeletalMeshComponent* GetSkeletalMeshComponent() const { return SkeletalMeshComponent; }
-
-		virtual int32 GetNumLODs() const;
-		virtual int32 GetLODIndex() const;
-		virtual int32 GetOwnerLODIndex(int32 LODIndex) const;
-		virtual bool IsValidLODIndex(int32 LODIndex) const;
-		virtual int32 GetNumPoints(int32 LODIndex) const;
-		virtual int32 GetNumPatternPoints(int32 LODIndex) const;
-		virtual TConstArrayView<FVector3f> GetPositions(int32 LODIndex) const;
-		virtual TConstArrayView<FVector2f> GetPatternPositions(int32 LODIndex) const;
-		virtual TConstArrayView<FVector3f> GetNormals(int32 LODIndex) const;
-		virtual TConstArrayView<uint32> GetIndices(int32 LODIndex) const;
-		virtual TConstArrayView<uint32> GetPatternIndices(int32 LODIndex) const;
-		virtual TConstArrayView<uint32> GetPatternToWeldedIndices(int32 LODIndex) const;
-		virtual TArray<FName> GetWeightMapNames(int32 LODIndex) const;
-		UE_DEPRECATED(5.3, "Use LODIndex version.")
-		virtual TArray<FName> GetWeightMapNames() const { return GetWeightMapNames(0); }
-		virtual TMap<FString, int32> GetWeightMapIndices(int32 LODIndex) const;
-		UE_DEPRECATED(5.3, "Use LODIndex version.")
-		virtual TMap<FString, int32> GetWeightMapIndices() const { return GetWeightMapIndices(0); }
-		virtual TArray<TConstArrayView<FRealSingle>> GetWeightMaps(int32 LODIndex) const;
-		virtual TMap<FString, const TSet<int32>*> GetVertexSets(int32 LODIndex) const;
-		virtual TMap<FString, const TSet<int32>*> GetFaceSets(int32 LODIndex) const;
-		virtual TMap<FString, TConstArrayView<int32>> GetFaceIntMaps(int32 LODIndex) const;
-		virtual TArray<TConstArrayView<TTuple<int32, int32, float>>> GetTethers(int32 LODIndex, bool bUseGeodesicTethers) const;
-		virtual int32 GetReferenceBoneIndex() const;
-		virtual FTransform GetReferenceBoneTransform() const;
-		virtual const TArray<FTransform>& GetBoneTransforms() const;
-		virtual const FTransform& GetComponentToWorldTransform() const;
-		virtual const TArray<FMatrix44f>& GetRefToLocalMatrices() const;
-		virtual TConstArrayView<int32> GetBoneMap() const;
-		virtual TConstArrayView<FClothVertBoneData> GetBoneData(int32 LODIndex) const;
-		virtual TConstArrayView<FMeshToMeshVertData> GetTransitionUpSkinData(int32 LODIndex) const;
-		virtual TConstArrayView<FMeshToMeshVertData> GetTransitionDownSkinData(int32 LODIndex) const;
-#else
 		/* Return the number of LODs on this mesh. */
 		virtual int32 GetNumLODs() const = 0;
 
@@ -124,22 +76,18 @@ namespace Chaos
 
 		/* Return all weight maps associated with this mesh returned in the same order as GetWeightMaps. */
 		virtual TArray<FName> GetWeightMapNames(int32 LODIndex) const = 0;
-		UE_DEPRECATED(5.3, "Use LODIndex version.")
-		virtual TArray<FName> GetWeightMapNames() const = 0;
 
 		/* Return a map of all weight map names associated with this mesh to the index in the array returned by GetWeightMaps. */
 		virtual TMap<FString, int32> GetWeightMapIndices(int32 LODIndex) const = 0;
-		UE_DEPRECATED(5.3, "Use LODIndex version.")
-		virtual TMap<FString, int32> GetWeightMapIndices() const = 0;
 
 		/* Return the specified LOD's weight map. */
 		virtual TArray<TConstArrayView<FRealSingle>> GetWeightMaps(int32 LODIndex) const = 0;
 
 		/* Return the specified LOD's vertex sets. */
-		virtual const TMap<FString, TSet<int32>*> GetVertexSets(int32 LODIndex) const = 0;
+		virtual TMap<FString, const TSet<int32>*> GetVertexSets(int32 LODIndex) const = 0;
 
 		/* Return the specified LOD's face sets. */
-		virtual const TMap<FString, TSet<int32>*> GetFaceSets(int32 LODIndex) const = 0;
+		virtual TMap<FString, const TSet<int32>*> GetFaceSets(int32 LODIndex) const = 0;
 
 		/* Return the specified LOD's face int maps. */
 		virtual TMap<FString, TConstArrayView<int32>> GetFaceIntMaps(int32 LODIndex) const = 0;
@@ -154,7 +102,7 @@ namespace Chaos
 		virtual FTransform GetReferenceBoneTransform() const = 0;
 
 		/* Return the bone transforms as required when updating the collider pose. */
-		const TArray<FTransform>& GetBoneTransforms() const = 0;
+		virtual const TArray<FTransform>& GetBoneTransforms() const = 0;
 
 		/* Return the transform of the bone treated as the root of the simulation space. */
 		virtual const FTransform& GetComponentToWorldTransform() const = 0;
@@ -173,7 +121,7 @@ namespace Chaos
 
 		/* Return the transition down data (PrevLODIndex > LODIndex), for matching shapes during LOD changes. */
 		virtual TConstArrayView<FMeshToMeshVertData> GetTransitionDownSkinData(int32 LODIndex) const = 0;
-#endif
+
 		/** Return this mesh uniform scale as the maximum of the three axis scale value. */
 		virtual Softs::FSolverReal GetScale() const;
 
@@ -212,13 +160,6 @@ namespace Chaos
 			Softs::FSolverVec3* OutPositions,
 			Softs::FSolverVec3* OutNormals) const;
 
-#if !defined(CHAOS_IS_CLOTHINGSIMULATIONMESH_ABSTRACT) || !CHAOS_IS_CLOTHINGSIMULATIONMESH_ABSTRACT
-		const FClothingSimulationContextCommon* GetContext() const;
-
-		const UClothingAssetCommon* Asset;
-		const USkeletalMeshComponent* SkeletalMeshComponent;
-#endif
-
 #if !UE_BUILD_SHIPPING
 		/** Debug name of the source component. */
 		FString DebugName;
diff --git a/Engine/Plugins/ChaosCloth/Source/ChaosCloth/Public/ChaosCloth/ChaosClothingSimulationSkeletalMesh.h b/Engine/Plugins/ChaosCloth/Source/ChaosCloth/Public/ChaosCloth/ChaosClothingSimulationSkeletalMesh.h
--- a/Engine/Plugins/ChaosCloth/Source/ChaosCloth/Public/ChaosCloth/ChaosClothingSimulationSkeletalMesh.h
+++ b/Engine/Plugins/ChaosCloth/Source/ChaosCloth/Public/ChaosCloth/ChaosClothingSimulationSkeletalMesh.h
@@ -8,30 +8,20 @@ class UClothingAssetCommon;
 
 namespace Chaos
 {
-PRAGMA_DISABLE_DEPRECATION_WARNINGS  // TODO: CHAOS_IS_CLOTHINGSIMULATIONMESH_ABSTRACT
 	class FClothingSimulationSkeletalMesh final : public FClothingSimulationMesh
-PRAGMA_ENABLE_DEPRECATION_WARNINGS
 	{
 	public:
 		FClothingSimulationSkeletalMesh(const UClothingAssetCommon* InAsset, const USkeletalMeshComponent* InSkeletalMeshComponent);
-PRAGMA_DISABLE_DEPRECATION_WARNINGS  // TODO: CHAOS_IS_CLOTHINGSIMULATIONMESH_ABSTRACT
 		virtual ~FClothingSimulationSkeletalMesh() override = default;
-PRAGMA_ENABLE_DEPRECATION_WARNINGS
 
 		FClothingSimulationSkeletalMesh(const FClothingSimulationSkeletalMesh&) = delete;
 		FClothingSimulationSkeletalMesh(FClothingSimulationSkeletalMesh&&) = delete;
 		FClothingSimulationSkeletalMesh& operator=(const FClothingSimulationSkeletalMesh&) = delete;
 		FClothingSimulationSkeletalMesh& operator=(FClothingSimulationSkeletalMesh&&) = delete;
 
-#if !defined(CHAOS_IS_CLOTHINGSIMULATIONMESH_ABSTRACT) || !CHAOS_IS_CLOTHINGSIMULATIONMESH_ABSTRACT
-		virtual const UClothingAssetCommon* GetAsset() const override { return Asset; }
-PRAGMA_DISABLE_DEPRECATION_WARNINGS  // TODO: CHAOS_IS_CLOTHINGSIMULATIONMESH_ABSTRACT
-		virtual const USkeletalMeshComponent* GetSkeletalMeshComponent() const override { return SkeletalMeshComponent; }
-PRAGMA_ENABLE_DEPRECATION_WARNINGS
-#else
 		const UClothingAssetCommon* GetAsset() const { return Asset; }
-		const USkeletalMeshComponent* GetSkeletalMeshComponent() const { return Cast<USkeletalMeshComponent>(GetSkinnedMeshComponent()); }
-#endif
+		const USkeletalMeshComponent* GetSkeletalMeshComponent() const { return SkeletalMeshComponent; }
+
 		//~ Begin FClothingSimulationMesh Interface
 		virtual int32 GetNumLODs() const override;
 		virtual int32 GetLODIndex() const override;
@@ -46,11 +36,7 @@ PRAGMA_ENABLE_DEPRECATION_WARNINGS
 		virtual TConstArrayView<uint32> GetPatternIndices(int32 LODIndex) const override;
 		virtual TConstArrayView<uint32> GetPatternToWeldedIndices(int32 LODIndex) const override;
 		virtual TArray<FName> GetWeightMapNames(int32 LODIndex) const override;
-		UE_DEPRECATED(5.3, "Use LODIndex version.")
-		virtual TArray<FName> GetWeightMapNames() const override { return GetWeightMapNames(0); }
 		virtual TMap<FString, int32> GetWeightMapIndices(int32 LODIndex) const override;
-		UE_DEPRECATED(5.3, "Use LODIndex version.")
-		virtual TMap<FString, int32> GetWeightMapIndices() const override { return GetWeightMapIndices(0); }
 		virtual TArray<TConstArrayView<FRealSingle>> GetWeightMaps(int32 LODIndex) const override;
 		virtual TMap<FString, const TSet<int32>*> GetVertexSets(int32 LODIndex) const override;
 		virtual TMap<FString, const TSet<int32>*> GetFaceSets(int32 LODIndex) const override;
diff --git a/Engine/Plugins/ChaosCloth/Source/ChaosCloth/Public/ChaosCloth/ChaosClothingSimulationSolver.h b/Engine/Plugins/ChaosCloth/Source/ChaosCloth/Public/ChaosCloth/ChaosClothingSimulationSolver.h
--- a/Engine/Plugins/ChaosCloth/Source/ChaosCloth/Public/ChaosCloth/ChaosClothingSimulationSolver.h
+++ b/Engine/Plugins/ChaosCloth/Source/ChaosCloth/Public/ChaosCloth/ChaosClothingSimulationSolver.h
@@ -58,14 +58,8 @@ namespace Chaos
 		CHAOSCLOTH_API void SetWindVelocity(const TVec3<FRealSingle>& InWindVelocity, FRealSingle InLegacyWindAdaption = (FRealSingle)0.);
 		const TVec3<FRealSingle>& GetWindVelocity() const { return WindVelocity; }
 
-		UE_DEPRECATED(5.3, "Set properties directly through FClothingSimulationConfig")
-		CHAOSCLOTH_API void SetNumIterations(int32 InNumIterations);
 		CHAOSCLOTH_API int32 GetNumIterations() const;
-		UE_DEPRECATED(5.3, "Set properties directly through FClothingSimulationConfig")
-		CHAOSCLOTH_API void SetMaxNumIterations(int32 InMaxNumIterations);
 		CHAOSCLOTH_API int32 GetMaxNumIterations() const;
-		UE_DEPRECATED(5.3, "Set properties directly through FClothingSimulationConfig")
-		CHAOSCLOTH_API void SetNumSubsteps(int32 InNumSubsteps);
 		CHAOSCLOTH_API int32 GetNumSubsteps() const;
 
 		void SetEnableSolver(bool InbEnableSolver) { bEnableSolver = InbEnableSolver; }
@@ -103,8 +97,6 @@ namespace Chaos
 
 		/** Set the cached positions onto the particles */
 		CHAOSCLOTH_API void UpdateFromCache(const FClothingSimulationCacheData& CacheData);
-		UE_DEPRECATED(5.3, "Use UpdateFromCache(CacheData) instead")
-		CHAOSCLOTH_API void UpdateFromCache(const TArray<FVector>& CachedPositions, const TArray<FVector>& CachedVelocities);
 
 		// Return the actual of number of iterations used by the Evolution solver after the update (different from the number of iterations, depends on frame rate)
 		CHAOSCLOTH_API int32 GetNumUsedIterations() const;
diff --git a/Engine/Plugins/ChaosClothAsset/Source/ChaosClothAssetEngine/Private/ChaosClothAsset/ClothAsset.cpp b/Engine/Plugins/ChaosClothAsset/Source/ChaosClothAssetEngine/Private/ChaosClothAsset/ClothAsset.cpp
--- a/Engine/Plugins/ChaosClothAsset/Source/ChaosClothAssetEngine/Private/ChaosClothAsset/ClothAsset.cpp
+++ b/Engine/Plugins/ChaosClothAsset/Source/ChaosClothAssetEngine/Private/ChaosClothAsset/ClothAsset.cpp
@@ -583,17 +583,6 @@ void UChaosClothAsset::CalculateBounds()
 	Bounds = FBoxSphereBounds(BoundingBox);
 }
 
-void UChaosClothAsset::UpdateSkeleton(bool bRebuildClothSimulationModel)
-{
-	CalculateInvRefMatrices();
-
-	if (bRebuildClothSimulationModel)
-	{
-		// Rebuild simulation model  // TODO: How does this work with skinning
-		BuildClothSimulationModel();
-	}
-}
-
 void UChaosClothAsset::Build(TArray<FChaosClothAssetLodTransitionDataCache>* InOutTransitionCache)
 {
 	using namespace UE::Chaos::ClothAsset;
@@ -1030,17 +1019,6 @@ void UChaosClothAsset::ReregisterComponents()
 	}
 }
 
-void UChaosClothAsset::BindSimMeshToRootBone()
-{
-	using namespace UE::Chaos::ClothAsset;
-	check(GetClothCollections().Num());
-
-	for (TSharedRef<FManagedArrayCollection>& ClothCollection : GetClothCollections())
-	{
-		FClothGeometryTools::BindMeshToRootBone(ClothCollection, true, false);
-	}
-}
-
 #if WITH_EDITORONLY_DATA
 
 void UChaosClothAsset::SetPreviewSceneSkeletalMesh(USkeletalMesh* Mesh)
diff --git a/Engine/Plugins/ChaosClothAsset/Source/ChaosClothAssetEngine/Private/ChaosClothAsset/ClothSimulationMesh.cpp b/Engine/Plugins/ChaosClothAsset/Source/ChaosClothAssetEngine/Private/ChaosClothAsset/ClothSimulationMesh.cpp
--- a/Engine/Plugins/ChaosClothAsset/Source/ChaosClothAssetEngine/Private/ChaosClothAsset/ClothSimulationMesh.cpp
+++ b/Engine/Plugins/ChaosClothAsset/Source/ChaosClothAssetEngine/Private/ChaosClothAsset/ClothSimulationMesh.cpp
@@ -9,9 +9,7 @@
 namespace UE::Chaos::ClothAsset
 {
 	FClothSimulationMesh::FClothSimulationMesh(const FChaosClothSimulationModel& InClothSimulationModel, const FClothSimulationContext& InClothSimulationContext, const FString& DebugName)
-PRAGMA_DISABLE_DEPRECATION_WARNINGS  // TODO: CHAOS_IS_CLOTHINGSIMULATIONMESH_ABSTRACT
 		: ::Chaos::FClothingSimulationMesh(DebugName)
-PRAGMA_ENABLE_DEPRECATION_WARNINGS
 		, ClothSimulationModel(InClothSimulationModel)
 		, ClothSimulationContext(InClothSimulationContext)
 	{
diff --git a/Engine/Plugins/ChaosClothAsset/Source/ChaosClothAssetEngine/Private/ChaosClothAsset/ClothSimulationMesh.h b/Engine/Plugins/ChaosClothAsset/Source/ChaosClothAssetEngine/Private/ChaosClothAsset/ClothSimulationMesh.h
--- a/Engine/Plugins/ChaosClothAsset/Source/ChaosClothAssetEngine/Private/ChaosClothAsset/ClothSimulationMesh.h
+++ b/Engine/Plugins/ChaosClothAsset/Source/ChaosClothAssetEngine/Private/ChaosClothAsset/ClothSimulationMesh.h
@@ -10,15 +10,11 @@ namespace UE::Chaos::ClothAsset
 {
 	struct FClothSimulationContext;
 
-PRAGMA_DISABLE_DEPRECATION_WARNINGS  // TODO: CHAOS_IS_CLOTHINGSIMULATIONMESH_ABSTRACT
 	class FClothSimulationMesh : public ::Chaos::FClothingSimulationMesh
-PRAGMA_ENABLE_DEPRECATION_WARNINGS
 	{
 	public:
 		FClothSimulationMesh(const FChaosClothSimulationModel& InClothSimulationModel, const FClothSimulationContext& InClothSimulationContext, const FString& DebugName);
-PRAGMA_DISABLE_DEPRECATION_WARNINGS  // TODO: CHAOS_IS_CLOTHINGSIMULATIONMESH_ABSTRACT
 		virtual ~FClothSimulationMesh() override = default;
-PRAGMA_ENABLE_DEPRECATION_WARNINGS
 
 		FClothSimulationMesh(const FClothSimulationMesh&) = delete;
 		FClothSimulationMesh(FClothSimulationMesh&&) = delete;
@@ -39,11 +35,7 @@ PRAGMA_ENABLE_DEPRECATION_WARNINGS
 		virtual TConstArrayView<uint32> GetPatternIndices(int32 LODIndex) const override;
 		virtual TConstArrayView<uint32> GetPatternToWeldedIndices(int32 LODIndex) const override;
 		virtual TArray<FName> GetWeightMapNames(int32 LODIndex) const override;
-		UE_DEPRECATED(5.3, "Use LODIndex version.")
-		virtual TArray<FName> GetWeightMapNames() const override { return GetWeightMapNames(0); }
 		virtual TMap<FString, int32> GetWeightMapIndices(int32 LODIndex) const override;
-		UE_DEPRECATED(5.3, "Use LODIndex version.")
-		virtual TMap<FString, int32> GetWeightMapIndices() const override { return GetWeightMapIndices(0); }
 		virtual TArray<TConstArrayView<::Chaos::FRealSingle>> GetWeightMaps(int32 LODIndex) const override;
 		virtual TMap<FString, const TSet<int32>*> GetVertexSets(int32 LODIndex) const override;
 		virtual TMap<FString, const TSet<int32>*> GetFaceSets(int32 LODIndex) const override;
diff --git a/Engine/Plugins/ChaosClothAsset/Source/ChaosClothAssetEngine/Private/ChaosClothAsset/ClothSimulationProxy.cpp b/Engine/Plugins/ChaosClothAsset/Source/ChaosClothAssetEngine/Private/ChaosClothAsset/ClothSimulationProxy.cpp
--- a/Engine/Plugins/ChaosClothAsset/Source/ChaosClothAssetEngine/Private/ChaosClothAsset/ClothSimulationProxy.cpp
+++ b/Engine/Plugins/ChaosClothAsset/Source/ChaosClothAssetEngine/Private/ChaosClothAsset/ClothSimulationProxy.cpp
@@ -141,10 +141,8 @@ namespace UE::Chaos::ClothAsset
 		const int32 MeshIndex = Meshes.Emplace(MakeUnique<FClothSimulationMesh>(*ClothSimulationModel, *ClothSimulationContext, DebugName));
 
 		// Create collider simulation thread object
-PRAGMA_DISABLE_DEPRECATION_WARNINGS  // TODO: CHAOS_IS_CLOTHINGSIMULATIONMESH_ABSTRACT
 		const FReferenceSkeleton* const ReferenceSkeleton = &ClothAsset->GetRefSkeleton();
 		const int32 ColliderIndex = Colliders.Emplace(MakeUnique<FClothingSimulationCollider>(ClothComponent.GetPhysicsAsset(), ReferenceSkeleton));
-PRAGMA_ENABLE_DEPRECATION_WARNINGS
 		Colliders[ColliderIndex]->SetCollisionData(&CollisionSourcesProxy->GetCollisionData());
 
 		// Create cloth config simulation thread object
@@ -155,9 +153,7 @@ PRAGMA_ENABLE_DEPRECATION_WARNINGS
 
 		const int32 ClothIndex = Cloths.Emplace(MakeUnique<FClothingSimulationCloth>(
 			Configs[ClothConfigIndex].Get(),
-PRAGMA_DISABLE_DEPRECATION_WARNINGS  // TODO: CHAOS_IS_CLOTHINGSIMULATIONMESH_ABSTRACT
 			Meshes[MeshIndex].Get(),
-PRAGMA_ENABLE_DEPRECATION_WARNINGS
 			ColliderIndex != INDEX_NONE ? TArray<FClothingSimulationCollider*>({ Colliders[ColliderIndex].Get() }) : TArray<FClothingSimulationCollider*>(),
 			GroupId));
 
@@ -416,9 +412,7 @@ PRAGMA_ENABLE_DEPRECATION_WARNINGS
 			}
 
 			// If the LOD has changed while the simulation is suspended, the cloth still needs to be updated with the correct LOD data
-			PRAGMA_DISABLE_DEPRECATION_WARNINGS  // TODO: CHAOS_IS_CLOTHINGSIMULATIONMESH_ABSTRACT
 			const int32 LODIndex = Cloth->GetMesh()->GetLODIndex();
-			PRAGMA_ENABLE_DEPRECATION_WARNINGS
 			if (LODIndex != Cloth->GetLODIndex(Solver.Get()))
 			{
 				if (!ClothComponent.IsSimulationEnabled())
@@ -488,9 +482,7 @@ PRAGMA_ENABLE_DEPRECATION_WARNINGS
 
 			// If the LOD has changed while the simulation is suspended, the cloth still needs to be updated with the correct LOD data
 			// This should be handled by calling UpdateClothLODs when not ticking/simulating.
-			PRAGMA_DISABLE_DEPRECATION_WARNINGS  // TODO: CHAOS_IS_CLOTHINGSIMULATIONMESH_ABSTRACT
 			const int32 LODIndex = Cloth->GetMesh()->GetLODIndex();
-			PRAGMA_ENABLE_DEPRECATION_WARNINGS
 			ensure(LODIndex == Cloth->GetLODIndex(Solver.Get()));
 
 			if (Cloth->GetParticleRangeId(Solver.Get()) == INDEX_NONE || Cloth->GetLODIndex(Solver.Get()) == INDEX_NONE)
@@ -555,9 +547,7 @@ PRAGMA_ENABLE_DEPRECATION_WARNINGS
 			}
 
 			// Set the current LOD these data apply to, so that the correct deformer mappings can be applied
-PRAGMA_DISABLE_DEPRECATION_WARNINGS  // TODO: CHAOS_IS_CLOTHINGSIMULATIONMESH_ABSTRACT
 			Data.LODIndex = Cloth->GetMesh()->GetOwnerLODIndex(LODIndex);  // The owner component LOD index can be different to the cloth mesh LOD index
-PRAGMA_ENABLE_DEPRECATION_WARNINGS
 		}
 	}
 
diff --git a/Engine/Plugins/ChaosClothAsset/Source/ChaosClothAssetEngine/Public/ChaosClothAsset/ClothAsset.h b/Engine/Plugins/ChaosClothAsset/Source/ChaosClothAssetEngine/Public/ChaosClothAsset/ClothAsset.h
--- a/Engine/Plugins/ChaosClothAsset/Source/ChaosClothAssetEngine/Public/ChaosClothAsset/ClothAsset.h
+++ b/Engine/Plugins/ChaosClothAsset/Source/ChaosClothAssetEngine/Public/ChaosClothAsset/ClothAsset.h
@@ -176,31 +176,30 @@ public:
 	 */
 	void SetReferenceSkeleton(const FReferenceSkeleton* ReferenceSkeleton, bool bRebuildModels = true, bool bRebindMeshes = true);
 
-	/** Set the bone hierachy to use for this cloth. */
-	UE_DEPRECATED(5.3, "Use SetReferenceSkeleton(const FReferenceSkeleton*, bool, bool) instead")
-	PRAGMA_DISABLE_DEPRECATION_WARNINGS
-	void SetReferenceSkeleton(const FReferenceSkeleton& InReferenceSkeleton, bool bRebuildClothSimulationModel = true) { GetRefSkeleton() = InReferenceSkeleton; UpdateSkeleton(bRebuildClothSimulationModel); }
-	PRAGMA_ENABLE_DEPRECATION_WARNINGS
-
-	/** Set the skinning weights for all of the sim vertices to be bound to the root node of the reference skeleton. */
-	UE_DEPRECATED(5.3, "Use FClothGeometryTools::BindMeshToRootBone or SetReferenceSkeleton(const FReferenceSkeleton*, bool, bool) instead.")
-	void BindSimMeshToRootBone();
-
 	const FPerQualityLevelInt& GetQualityLevelMinLod() const { return MinQualityLevelLOD; }
 	void SetQualityLevelMinLod(FPerQualityLevelInt InMinLod) { MinQualityLevelLOD = MoveTemp(InMinLod); }
 	static void OnLodStrippingQualityLevelChanged(IConsoleVariable* Variable);
 
-	//
-	// Dataflow
-	//
-	UE_DEPRECATED(5.3, "Do not use. Will be made private in 5.4")
-	UPROPERTY(EditAnywhere, Category = "Dataflow")
-	TObjectPtr<UDataflow> DataflowAsset;
+	/**
+	 * Set the Dataflow graph asset for this cloth asset.
+	 * @param InDataflow The new dataflow asset.
+	 */
+	void SetDataflow(UDataflow* InDataflow) { DataflowAsset = InDataflow; }
 
-	UE_DEPRECATED(5.3, "Do not use. Will be made private in 5.4")
-	UPROPERTY(EditAnywhere, Category = "Dataflow")
-	FString DataflowTerminal = "ClothAssetTerminal";
+	/** Return the Dataflow graph asset associated to this cloth asset if any. */
+	UDataflow* GetDataflow() { return DataflowAsset; }
 
+	/** Return the Dataflow graph asset associated to this cloth asset if any, const version. */
+	const UDataflow* GetDataflow() const { return DataflowAsset; }
+
+	/**
+	 * Set the name of the Dataflow terminal node for this cloth asset.
+	 * @param InDataflowTerminal The new name of the Dataflow terminal node.
+	 */
+	void SetDataflowTerminal(const FString& InDataflowTerminal) { DataflowTerminal = InDataflowTerminal; }
+
+	/** Return the Dataflow graph asset associated to this cloth asset if any. */
+	const FString& GetDataflowTerminal() const { return DataflowTerminal; }
 
 #if WITH_EDITORONLY_DATA
 
@@ -213,6 +212,14 @@ public:
 #endif
 
 private:
+	/** Dataflow asset. */
+	UPROPERTY(EditAnywhere, Category = "Dataflow")
+	TObjectPtr<UDataflow> DataflowAsset;
+
+	/** Dataflow Asset terminal node. */
+	UPROPERTY(EditAnywhere, Category = "Dataflow")
+	FString DataflowTerminal = "ClothAssetTerminal";
+
 	//~ Begin USkinnedAsset interface
 	/** Initial step for the Post Load process - Can't be done in parallel. */
 	virtual void BeginPostLoadInternal(FSkinnedAssetPostLoadContext& Context) override;
@@ -236,10 +243,6 @@ private:
 	/** Re-calculate the bounds for this asset. */
 	void CalculateBounds();
 
-	/** Update the bone informations after a change of skeleton. */
-	UE_DEPRECATED(5.3, "Use Build() instead")
-	void UpdateSkeleton(bool bRebuildClothSimulationModel = true);
-
 #if WITH_EDITOR
 	/** Build the SkeletalMeshLODModel for this asset. */
 	void BuildMeshModel();
diff --git a/Engine/Plugins/ChaosClothAsset/Source/ChaosClothAssetEngine/Public/ChaosClothAsset/ClothSimulationProxy.h b/Engine/Plugins/ChaosClothAsset/Source/ChaosClothAssetEngine/Public/ChaosClothAsset/ClothSimulationProxy.h
--- a/Engine/Plugins/ChaosClothAsset/Source/ChaosClothAssetEngine/Public/ChaosClothAsset/ClothSimulationProxy.h
+++ b/Engine/Plugins/ChaosClothAsset/Source/ChaosClothAssetEngine/Public/ChaosClothAsset/ClothSimulationProxy.h
@@ -105,9 +105,7 @@ namespace UE::Chaos::ClothAsset
 
 		// Simulation objects
 		TUniquePtr<::Chaos::FClothingSimulationSolver> Solver;
-PRAGMA_DISABLE_DEPRECATION_WARNINGS  // TODO: CHAOS_IS_CLOTHINGSIMULATIONMESH_ABSTRACT
 		TArray<TUniquePtr<::Chaos::FClothingSimulationMesh>> Meshes;
-PRAGMA_ENABLE_DEPRECATION_WARNINGS
 		TArray<TUniquePtr<::Chaos::FClothingSimulationCloth>> Cloths;
 		TArray<TUniquePtr<::Chaos::FClothingSimulationConfig>> Configs;
 		TArray<TUniquePtr<::Chaos::FClothingSimulationCollider>> Colliders;
diff --git a/Engine/Plugins/ChaosClothAssetEditor/Source/ChaosClothAssetDataflowNodes/Private/ChaosClothAsset/AddWeightMapNode.cpp b/Engine/Plugins/ChaosClothAssetEditor/Source/ChaosClothAssetDataflowNodes/Private/ChaosClothAsset/AddWeightMapNode.cpp
--- a/Engine/Plugins/ChaosClothAssetEditor/Source/ChaosClothAssetDataflowNodes/Private/ChaosClothAsset/AddWeightMapNode.cpp
+++ b/Engine/Plugins/ChaosClothAssetEditor/Source/ChaosClothAssetDataflowNodes/Private/ChaosClothAsset/AddWeightMapNode.cpp
@@ -167,9 +167,7 @@ void FChaosClothAssetAddWeightMapNode::SetAssetValue(TObjectPtr<UObject> Asset,
 
 	if (UChaosClothAsset* const ClothAsset = Cast<UChaosClothAsset>(Asset.Get()))
 	{
-PRAGMA_DISABLE_DEPRECATION_WARNINGS  // TODO: Remove after adding the getter function (currently shelved!)  
-		if (UDataflow* const DataflowAsset = ClothAsset->DataflowAsset)
-PRAGMA_ENABLE_DEPRECATION_WARNINGS
+		if (UDataflow* const DataflowAsset = ClothAsset->GetDataflow())
 		{
 			const TSharedPtr<Dataflow::FGraph, ESPMode::ThreadSafe> Dataflow = DataflowAsset->GetDataflow();
 			if (const TSharedPtr<FDataflowNode> BaseNode = Dataflow->FindBaseNode(this->GetGuid()))  // This is basically a safe const_cast
diff --git a/Engine/Plugins/ChaosClothAssetEditor/Source/ChaosClothAssetDataflowNodes/Private/ChaosClothAsset/SelectionNode.cpp b/Engine/Plugins/ChaosClothAssetEditor/Source/ChaosClothAssetDataflowNodes/Private/ChaosClothAsset/SelectionNode.cpp
--- a/Engine/Plugins/ChaosClothAssetEditor/Source/ChaosClothAssetDataflowNodes/Private/ChaosClothAsset/SelectionNode.cpp
+++ b/Engine/Plugins/ChaosClothAssetEditor/Source/ChaosClothAssetDataflowNodes/Private/ChaosClothAsset/SelectionNode.cpp
@@ -250,9 +250,7 @@ void FChaosClothAssetSelectionNode::SetAssetValue(TObjectPtr<UObject> Asset, Dat
 
 	if (UChaosClothAsset* const ClothAsset = Cast<UChaosClothAsset>(Asset.Get()))
 	{
-		PRAGMA_DISABLE_DEPRECATION_WARNINGS  // TODO: Remove after adding the getter function (currently shelved!)  
-		if (UDataflow* const DataflowAsset = ClothAsset->DataflowAsset)
-		PRAGMA_ENABLE_DEPRECATION_WARNINGS
+		if (UDataflow* const DataflowAsset = ClothAsset->GetDataflow())
 		{
 			const TSharedPtr<Dataflow::FGraph, ESPMode::ThreadSafe> Dataflow = DataflowAsset->GetDataflow();
 			if (const TSharedPtr<FDataflowNode> BaseNode = Dataflow->FindBaseNode(this->GetGuid()))  // This is basically a safe const_cast
diff --git a/Engine/Plugins/ChaosClothAssetEditor/Source/ChaosClothAssetDataflowNodes/Private/ChaosClothAsset/USDImportNode.cpp b/Engine/Plugins/ChaosClothAssetEditor/Source/ChaosClothAssetDataflowNodes/Private/ChaosClothAsset/USDImportNode.cpp
--- a/Engine/Plugins/ChaosClothAssetEditor/Source/ChaosClothAssetDataflowNodes/Private/ChaosClothAsset/USDImportNode.cpp
+++ b/Engine/Plugins/ChaosClothAssetEditor/Source/ChaosClothAssetDataflowNodes/Private/ChaosClothAsset/USDImportNode.cpp
@@ -331,9 +331,7 @@ void FChaosClothAssetUSDImportNode::SetAssetValue(TObjectPtr<UObject> Asset, Dat
 
 	if (UChaosClothAsset* const ClothAsset = Cast<UChaosClothAsset>(Asset.Get()))
 	{
-PRAGMA_DISABLE_DEPRECATION_WARNINGS  // TODO: Don't use public property, and have Getter/Setter API instead
-		if (UDataflow* const DataflowAsset = ClothAsset->DataflowAsset)
-PRAGMA_ENABLE_DEPRECATION_WARNINGS
+		if (UDataflow* const DataflowAsset = ClothAsset->GetDataflow())
 		{
 			const TSharedPtr<Dataflow::FGraph, ESPMode::ThreadSafe> Dataflow = DataflowAsset->GetDataflow();
 			if (const TSharedPtr<FDataflowNode> BaseNode = Dataflow->FindBaseNode(this->GetGuid()))  // This is basically a safe const_cast
diff --git a/Engine/Plugins/ChaosClothAssetEditor/Source/ChaosClothAssetDataflowNodes/Private/ChaosClothAsset/WeightMapNode.cpp b/Engine/Plugins/ChaosClothAssetEditor/Source/ChaosClothAssetDataflowNodes/Private/ChaosClothAsset/WeightMapNode.cpp
--- a/Engine/Plugins/ChaosClothAssetEditor/Source/ChaosClothAssetDataflowNodes/Private/ChaosClothAsset/WeightMapNode.cpp
+++ b/Engine/Plugins/ChaosClothAssetEditor/Source/ChaosClothAssetDataflowNodes/Private/ChaosClothAsset/WeightMapNode.cpp
@@ -59,9 +59,7 @@ void FChaosClothAssetWeightMapNode::SetAssetValue(TObjectPtr<UObject> Asset, Dat
 
 	if (UChaosClothAsset* const ClothAsset = Cast<UChaosClothAsset>(Asset.Get()))
 	{
-PRAGMA_DISABLE_DEPRECATION_WARNINGS  // TODO: Remove after adding the getter function (currently shelved!)  
-		if (UDataflow* const DataflowAsset = ClothAsset->DataflowAsset)
-PRAGMA_ENABLE_DEPRECATION_WARNINGS
+		if (UDataflow* const DataflowAsset = ClothAsset->GetDataflow())
 		{
 			const TSharedPtr<Dataflow::FGraph, ESPMode::ThreadSafe> Dataflow = DataflowAsset->GetDataflow();
 			if (const TSharedPtr<FDataflowNode> BaseNode = Dataflow->FindBaseNode(this->GetGuid()))  // This is basically a safe const_cast
diff --git a/Engine/Plugins/ChaosClothAssetEditor/Source/ChaosClothAssetEditor/Private/ChaosClothAsset/AssetDefinition_ClothAsset.cpp b/Engine/Plugins/ChaosClothAssetEditor/Source/ChaosClothAssetEditor/Private/ChaosClothAsset/AssetDefinition_ClothAsset.cpp
--- a/Engine/Plugins/ChaosClothAssetEditor/Source/ChaosClothAssetEditor/Private/ChaosClothAsset/AssetDefinition_ClothAsset.cpp
+++ b/Engine/Plugins/ChaosClothAssetEditor/Source/ChaosClothAssetEditor/Private/ChaosClothAsset/AssetDefinition_ClothAsset.cpp
@@ -179,15 +179,13 @@ EAssetCommandResult UAssetDefinition_ClothAsset::OpenAssets(const FAssetOpenArgs
 
 		// Validate the asset
 		UChaosClothAsset* const ClothAsset = CastChecked<UChaosClothAsset>(ClothObjects[0]);
-PRAGMA_DISABLE_DEPRECATION_WARNINGS  // TODO: Don't use public property, and have Getter/Setter API instead
-		if (!ClothAsset->DataflowAsset)
+		if (!ClothAsset->GetDataflow())
 		{
 			if (UDataflow* const NewDataflowAsset = Cast<UDataflow>(ClothAssetDefinitionHelpers::NewOrOpenDataflowAsset(ClothAsset)))
 			{
-				ClothAsset->DataflowAsset = NewDataflowAsset;
+				ClothAsset->SetDataflow(NewDataflowAsset);
 			}
 		}
-PRAGMA_ENABLE_DEPRECATION_WARNINGS
 
 		TArray<TObjectPtr<UObject>> Objects;
 		for (UChaosClothAsset* const ClothObject : ClothObjects)
diff --git a/Engine/Plugins/ChaosClothAssetEditor/Source/ChaosClothAssetEditor/Private/ChaosClothAsset/ClothEditorToolkit.cpp b/Engine/Plugins/ChaosClothAssetEditor/Source/ChaosClothAssetEditor/Private/ChaosClothAsset/ClothEditorToolkit.cpp
--- a/Engine/Plugins/ChaosClothAssetEditor/Source/ChaosClothAssetEditor/Private/ChaosClothAsset/ClothEditorToolkit.cpp
+++ b/Engine/Plugins/ChaosClothAssetEditor/Source/ChaosClothAssetEditor/Private/ChaosClothAsset/ClothEditorToolkit.cpp
@@ -236,17 +236,13 @@ TSharedPtr<Dataflow::FEngineContext> FChaosClothAssetEditorToolkit::GetDataflowC
 
 const UDataflow* FChaosClothAssetEditorToolkit::GetDataflow() const
 {
-	PRAGMA_DISABLE_DEPRECATION_WARNINGS  // TODO: Don't use public property, and have Getter/Setter API instead
-	return GetAsset()->DataflowAsset;
-	PRAGMA_ENABLE_DEPRECATION_WARNINGS
+	return GetAsset()->GetDataflow();
 }
 
 
 UDataflow* FChaosClothAssetEditorToolkit::GetDataflow()
 {
-	PRAGMA_DISABLE_DEPRECATION_WARNINGS  // TODO: Don't use public property, and have Getter/Setter API instead
-	return GetAsset()->DataflowAsset;
-	PRAGMA_ENABLE_DEPRECATION_WARNINGS
+	return GetAsset()->GetDataflow();
 }
 
 //~ Begin FTickableEditorObject overrides
@@ -584,9 +580,7 @@ void FChaosClothAssetEditorToolkit::OnAssetsSavedAs(const TArray<UObject*>& Save
 
 	if (NewClothAsset && NewDataflowAsset)
 	{
-PRAGMA_DISABLE_DEPRECATION_WARNINGS  // TODO: Don't use public property, and have Getter/Setter API instead
-		NewClothAsset->DataflowAsset = NewDataflowAsset;
-PRAGMA_ENABLE_DEPRECATION_WARNINGS
+		NewClothAsset->SetDataflow(NewDataflowAsset);
 
 		// Now save the new Cloth asset again since we've updated its Property
 		const TArray<UPackage*> PackagesToSave{NewClothAsset->GetOutermost()};
diff --git a/Engine/Source/Runtime/ClothingSystemRuntimeCommon/Private/ClothLODData.cpp b/Engine/Source/Runtime/ClothingSystemRuntimeCommon/Private/ClothLODData.cpp
--- a/Engine/Source/Runtime/ClothingSystemRuntimeCommon/Private/ClothLODData.cpp
+++ b/Engine/Source/Runtime/ClothingSystemRuntimeCommon/Private/ClothLODData.cpp
@@ -38,8 +38,9 @@ bool FClothLODDataCommon::Serialize(FArchive& Ar)
 
 #if WITH_EDITORONLY_DATA
 	const int32 ClothingCustomVersion = Ar.CustomVer(FClothingAssetCustomVersion::GUID);
-	if (ClothingCustomVersion < FClothingAssetCustomVersion::MovePropertiesToCommonBaseClasses)
+	if (Ar.IsLoading() && ClothingCustomVersion < FClothingAssetCustomVersion::MovePropertiesToCommonBaseClasses)
 	{
+PRAGMA_DISABLE_DEPRECATION_WARNINGS
 		// Migrate maps
 		PhysicalMeshData.GetWeightMap(EWeightMapTargetCommon::MaxDistance).Values = MoveTemp(PhysicalMeshData.MaxDistances_DEPRECATED);
 		PhysicalMeshData.GetWeightMap(EWeightMapTargetCommon::BackstopDistance).Values = MoveTemp(PhysicalMeshData.BackstopDistances_DEPRECATED);
@@ -55,11 +56,10 @@ bool FClothLODDataCommon::Serialize(FArchive& Ar)
 		ParameterMasks_DEPRECATED.Empty();
 
 		// Remove deprecated Apex collisions from the LOD data
-PRAGMA_DISABLE_DEPRECATION_WARNINGS
-		CollisionData.Spheres.Empty();
-		CollisionData.SphereConnections.Empty();
-		CollisionData.Convexes.Empty();
-		CollisionData.Boxes.Empty();
+		CollisionData_DEPRECATED.Spheres.Empty();
+		CollisionData_DEPRECATED.SphereConnections.Empty();
+		CollisionData_DEPRECATED.Convexes.Empty();
+		CollisionData_DEPRECATED.Boxes.Empty();
 PRAGMA_ENABLE_DEPRECATION_WARNINGS
 	}
 #endif // WITH_EDITORONLY_DATA
diff --git a/Engine/Source/Runtime/ClothingSystemRuntimeCommon/Private/ClothPhysicalMeshData.cpp b/Engine/Source/Runtime/ClothingSystemRuntimeCommon/Private/ClothPhysicalMeshData.cpp
--- a/Engine/Source/Runtime/ClothingSystemRuntimeCommon/Private/ClothPhysicalMeshData.cpp
+++ b/Engine/Source/Runtime/ClothingSystemRuntimeCommon/Private/ClothPhysicalMeshData.cpp
@@ -93,20 +93,6 @@ void FClothPhysicalMeshData::ClearWeightMaps()
 	AddWeightMap(EWeightMapTargetCommon::AnimDriveStiffness);
 }
 
-void FClothPhysicalMeshData::BuildSelfCollisionData(const TMap<FName, TObjectPtr<UClothConfigBase>>& ClothConfigs)
-{
-	// Deprecated from 5.0 onwards
-	float SelfCollisionRadius = 0.f;
-	for (const TPair<FName, TObjectPtr<UClothConfigBase>>& ClothConfig : ClothConfigs)
-	{
-		SelfCollisionRadius = FMath::Max(SelfCollisionRadius, ClothConfig.Value->GetSelfCollisionRadius());
-	}
-	if (SelfCollisionRadius)
-	{
-		BuildSelfCollisionData(SelfCollisionRadius);
-	}
-}
-
 void FClothPhysicalMeshData::BuildSelfCollisionData(float SelfCollisionRadius)
 {
 	const float SelfCollisionDiamSq = 4.f *SelfCollisionRadius * SelfCollisionRadius;
@@ -267,41 +253,3 @@ void FClothPhysicalMeshData::CalculateTethers(bool bUseEuclideanDistance, bool b
 		}
 	}
 }
-
-// TODO: Deprecated, turn this into CalculateNormals() after 5.0
-void FClothPhysicalMeshData::ComputeFaceAveragedVertexNormals(TArray<FVector3f>& OutNormals) const
-{
-	OutNormals.Init(FVector3f{ 0.f, 0.f, 0.f }, Vertices.Num());
-
-	const int32 NumTris = Indices.Num() / 3;
-	for (int32 TID = 0; TID < NumTris; ++TID)
-	{
-		const uint32 VA = Indices[3 * TID + 0];
-		const uint32 VB = Indices[3 * TID + 1];
-		const uint32 VC = Indices[3 * TID + 2];
-		const FVector3f& PA = Vertices[VA];
-		const FVector3f& PB = Vertices[VB];
-		const FVector3f& PC = Vertices[VC];
-
-		FVector3f Normal = FVector3f::CrossProduct(PC - PA, PB - PA);
-		if (!Normal.Normalize())
-		{
-			// skip contributions from degenerate triangles
-			continue;
-		}
-
-		OutNormals[VA] += Normal;
-		OutNormals[VB] += Normal;
-		OutNormals[VC] += Normal;
-	}
-
-	for (int32 VertexID = 0; VertexID < OutNormals.Num(); ++VertexID )
-	{
-		FVector3f& N = OutNormals[VertexID];
-		if (!N.Normalize())
-		{
-			N = FVector3f::XAxisVector;
-		}
-	}
-}
-
diff --git a/Engine/Source/Runtime/ClothingSystemRuntimeCommon/Private/ClothingAsset.cpp b/Engine/Source/Runtime/ClothingSystemRuntimeCommon/Private/ClothingAsset.cpp
--- a/Engine/Source/Runtime/ClothingSystemRuntimeCommon/Private/ClothingAsset.cpp
+++ b/Engine/Source/Runtime/ClothingSystemRuntimeCommon/Private/ClothingAsset.cpp
@@ -40,135 +40,6 @@ DEFINE_LOG_CATEGORY(LogClothingAsset)
 // ClothingAssetUtils
 //==============================================================================
 
-//Deprecated function
-void ClothingAssetUtils::GetMeshClothingAssetBindings(
-	USkeletalMesh* InSkelMesh, 
-	TArray<FClothingAssetMeshBinding>& OutBindings)
-{
-	OutBindings.Empty();
-
-	if(!InSkelMesh)
-	{
-		return;
-	}
-#if WITH_EDITORONLY_DATA
-	if (InSkelMesh->GetImportedModel())
-	{
-		int32 LODNum = InSkelMesh->GetImportedModel()->LODModels.Num();
-		for (int32 LODIndex = 0; LODIndex < LODNum; ++LODIndex)
-		{
-			if (InSkelMesh->GetImportedModel()->LODModels[LODIndex].HasClothData())
-			{
-				TArray<FClothingAssetMeshBinding> LodBindings;
-PRAGMA_DISABLE_DEPRECATION_WARNINGS
-				GetMeshClothingAssetBindings(InSkelMesh, LodBindings, LODIndex);
-PRAGMA_ENABLE_DEPRECATION_WARNINGS
-				OutBindings.Append(LodBindings);
-			}
-		}
-		if (OutBindings.Num() > 0)
-		{
-			return;
-		}
-	}
-#endif
-
-	//Fallback on render data
-	if (FSkeletalMeshRenderData* Resource = InSkelMesh->GetResourceForRendering())
-	{
-		const int32 NumLods = Resource->LODRenderData.Num();
-
-		for (int32 LodIndex = 0; LodIndex < NumLods; ++LodIndex)
-		{
-			TArray<FClothingAssetMeshBinding> LodBindings;
-PRAGMA_DISABLE_DEPRECATION_WARNINGS
-			GetMeshClothingAssetBindings(InSkelMesh, LodBindings, LodIndex);
-PRAGMA_ENABLE_DEPRECATION_WARNINGS
-
-			OutBindings.Append(LodBindings);
-		}
-	}
-}
-
-//Deprecated function
-void ClothingAssetUtils::GetMeshClothingAssetBindings(
-	USkeletalMesh* InSkelMesh, 
-	TArray<FClothingAssetMeshBinding>& OutBindings, 
-	int32 InLodIndex)
-{
-	OutBindings.Empty();
-
-	if(!InSkelMesh)
-	{
-		return;
-	}
-
-#if WITH_EDITORONLY_DATA
-	if (InSkelMesh->GetImportedModel())
-	{
-		int32 LODNum = InSkelMesh->GetImportedModel()->LODModels.Num();
-		if (InSkelMesh->GetImportedModel()->LODModels[InLodIndex].HasClothData())
-		{
-			TArray<FClothingAssetMeshBinding> LodBindings;
-			int32 SectionNum = InSkelMesh->GetImportedModel()->LODModels[InLodIndex].Sections.Num();
-			for (int32 SectionIndex = 0; SectionIndex < SectionNum; ++SectionIndex)
-			{
-				const FSkelMeshSection& Section = InSkelMesh->GetImportedModel()->LODModels[InLodIndex].Sections[SectionIndex];
-				if (Section.HasClothingData())
-				{
-					UClothingAssetBase* ClothingAsset = InSkelMesh->GetClothingAsset(Section.ClothingData.AssetGuid);
-					FClothingAssetMeshBinding ClothBinding;
-					ClothBinding.Asset = Cast<UClothingAssetCommon>(ClothingAsset);
-					ClothBinding.AssetInternalLodIndex = Section.ClothingData.AssetLodIndex;// InSkelMesh->GetClothingAssetIndex(Section.ClothingData.AssetGuid);
-					check(ClothBinding.AssetInternalLodIndex == Section.ClothingData.AssetLodIndex);
-					ClothBinding.LODIndex = InLodIndex;
-					ClothBinding.SectionIndex = SectionIndex;
-					OutBindings.Add(ClothBinding);
-				}
-			}
-		}
-
-		if (OutBindings.Num() > 0)
-		{
-			return;
-		}
-	}
-#endif
-
-	//Fallback on render data
-	if(FSkeletalMeshRenderData* Resource = InSkelMesh->GetResourceForRendering())
-	{
-		if(Resource->LODRenderData.IsValidIndex(InLodIndex))
-		{
-			FSkeletalMeshLODRenderData& LodData = Resource->LODRenderData[InLodIndex];
-
-			const int32 NumSections = LodData.RenderSections.Num();
-
-			for(int32 SectionIndex = 0; SectionIndex < NumSections; ++SectionIndex)
-			{
-				FSkelMeshRenderSection& Section = LodData.RenderSections[SectionIndex];
-
-				if(Section.HasClothingData())
-				{
-					UClothingAssetCommon* SectionAsset = Cast<UClothingAssetCommon>(InSkelMesh->GetSectionClothingAsset(InLodIndex, SectionIndex));
-
-					if(SectionAsset)
-					{
-						// This is the original section of a clothing section pair
-						OutBindings.AddDefaulted();
-						FClothingAssetMeshBinding& Binding = OutBindings.Last();
-
-						Binding.Asset = SectionAsset;
-						Binding.LODIndex = InLodIndex;
-						Binding.SectionIndex = SectionIndex;
-						Binding.AssetInternalLodIndex = Section.ClothingData.AssetLodIndex;
-					}
-				}
-			}
-		}
-	}
-}
-
 #if WITH_EDITOR
 void ClothingAssetUtils::GetAllMeshClothingAssetBindings(const USkeletalMesh* SkeletalMesh, TArray<FClothingAssetMeshBinding>& OutBindings)
 {
@@ -251,7 +122,6 @@ UClothingAssetCommon::UClothingAssetCommon(const FObjectInitializer& ObjectIniti
 	, ChaosClothSimConfig_DEPRECATED(nullptr)
 #endif
 	, ReferenceBoneIndex(0)
-	, CustomData(nullptr)  // Deprecated
 {
 }
 
@@ -467,13 +337,6 @@ bool UClothingAssetCommon::BindToSkeletalMesh(
 		InSkelMesh->GetRefSkeleton().EnsureParentsExistAndSort(SkelLod.ActiveBoneIndices);
 	}
 
-PRAGMA_DISABLE_DEPRECATION_WARNINGS
-	if(CustomData)  // Deprecated from 5.0 onward
-	{
-		CustomData->BindToSkeletalMesh(InSkelMesh, InMeshLodIndex, InSectionIndex, InAssetLodIndex);
-	}
-PRAGMA_ENABLE_DEPRECATION_WARNINGS
-
 	// Make sure the LOD map is always big enough for the asset to use.
 	// This shouldn't grow to an unwieldy size but maybe consider compacting later.
 	while(LodMap.Num() - 1 < InMeshLodIndex)
diff --git a/Engine/Source/Runtime/ClothingSystemRuntimeCommon/Private/ClothingSimulation.cpp b/Engine/Source/Runtime/ClothingSystemRuntimeCommon/Private/ClothingSimulation.cpp
--- a/Engine/Source/Runtime/ClothingSystemRuntimeCommon/Private/ClothingSimulation.cpp
+++ b/Engine/Source/Runtime/ClothingSystemRuntimeCommon/Private/ClothingSimulation.cpp
@@ -47,13 +47,6 @@ FClothingSimulationContextCommon::~FClothingSimulationContextCommon()
 {}
 PRAGMA_ENABLE_DEPRECATION_WARNINGS
 
-void FClothingSimulationContextCommon::Fill(const USkeletalMeshComponent* InComponent, float InDeltaSeconds, float InMaxPhysicsDelta)
-{
-	// Deprecated version always fills RefToLocals with current animation results instead of using reference pose on initialization
-	const bool bIsInitialization = false;
-	Fill(InComponent, InDeltaSeconds, InMaxPhysicsDelta, bIsInitialization);
-}
-
 void FClothingSimulationContextCommon::Fill(const USkeletalMeshComponent* InComponent, float InDeltaSeconds, float InMaxPhysicsDelta, bool bIsInitialization)
 {
 	SCOPE_CYCLE_COUNTER(STAT_ClothFillContext);
diff --git a/Engine/Source/Runtime/ClothingSystemRuntimeCommon/Public/ClothLODData.h b/Engine/Source/Runtime/ClothingSystemRuntimeCommon/Public/ClothLODData.h
--- a/Engine/Source/Runtime/ClothingSystemRuntimeCommon/Public/ClothLODData.h
+++ b/Engine/Source/Runtime/ClothingSystemRuntimeCommon/Public/ClothLODData.h
@@ -24,11 +24,6 @@ PRAGMA_ENABLE_DEPRECATION_WARNINGS
 	UPROPERTY(EditAnywhere, Category = SimMesh)
 	FClothPhysicalMeshData PhysicalMeshData;
 
-	// Collision primitive and convex data for clothing collisions
-	UE_DEPRECATED(5.2, "This property is no longer supported. Use Physics Asset instead.")
-	UPROPERTY(EditAnywhere, Category = Collision, Meta = (DeprecatedProperty, DeprecationMessage = "This property is no longer supported. Use Physics Asset instead."))
-	FClothCollisionData CollisionData;
-
 	// Whether to use multiple triangles to interpolate from simulated cloth mesh to render mesh
 	UPROPERTY()
 	bool bUseMultipleInfluences = false;
@@ -42,7 +37,12 @@ PRAGMA_ENABLE_DEPRECATION_WARNINGS
 	bool bSmoothTransition = false;
 
 #if WITH_EDITORONLY_DATA
-	UPROPERTY()
+	UE_DEPRECATED(all, "This property is no longer supported. Use a Physics Asset instead.")
+	UPROPERTY(Meta = (DeprecatedProperty, DeprecationMessage = "This property is no longer supported. Use a Physics Asset instead."))
+	FClothCollisionData CollisionData_DEPRECATED;
+
+	UE_DEPRECATED(5.5, "This property is no longer supported. Use PointWeightMaps instead.")
+	UPROPERTY(Meta = (DeprecatedProperty, DeprecationMessage = "This property is no longer supported. Use PointWeightMaps instead."))
 	TArray<FClothParameterMask_Legacy> ParameterMasks_DEPRECATED;
 
 	// Parameter masks defining the physics mesh masked data
diff --git a/Engine/Source/Runtime/ClothingSystemRuntimeCommon/Public/ClothPhysicalMeshData.h b/Engine/Source/Runtime/ClothingSystemRuntimeCommon/Public/ClothPhysicalMeshData.h
--- a/Engine/Source/Runtime/ClothingSystemRuntimeCommon/Public/ClothPhysicalMeshData.h
+++ b/Engine/Source/Runtime/ClothingSystemRuntimeCommon/Public/ClothPhysicalMeshData.h
@@ -46,10 +46,6 @@ PRAGMA_ENABLE_DEPRECATION_WARNINGS
 	/** Clear out any default weight maps and delete any other ones. */
 	CLOTHINGSYSTEMRUNTIMECOMMON_API void ClearWeightMaps();
 
-	/** Build the self collision indices for the relevant config. */
-	UE_DEPRECATED(5.0, "Use BuildSelfCollisionData(float SelfCollisionRadius) instead.")
-	CLOTHINGSYSTEMRUNTIMECOMMON_API void BuildSelfCollisionData(const TMap<FName, TObjectPtr<UClothConfigBase>>& ClothConfigs);
-
 	/** Build the self collision indices with the specified radius. */
 	CLOTHINGSYSTEMRUNTIMECOMMON_API void BuildSelfCollisionData(float SelfCollisionRadius);
 
@@ -62,11 +58,6 @@ PRAGMA_ENABLE_DEPRECATION_WARNINGS
 	/** Recalculate the long range attachment tethers. */
 	CLOTHINGSYSTEMRUNTIMECOMMON_API void CalculateTethers(bool bUseEuclideanDistance, bool bUseGeodesicDistance);
 
-	/** Compute vertex normals as unweighted average of incident face normals. If all incident triangles are degenerate,
-	    the vertex normal is assigned a value of FVector3f::XAxisVector. */
-	UE_DEPRECATED(5.0, "This function is no longer part of this API and will soon be removed.")
-	CLOTHINGSYSTEMRUNTIMECOMMON_API void ComputeFaceAveragedVertexNormals(TArray<FVector3f>& OutNormals) const;
-
 	/** Retrieve whether a vertex weight array has already been registered. */
 	template<typename T>
 	bool HasWeightMap(const T Target) const { return WeightMaps.Contains((uint32)Target); }
diff --git a/Engine/Source/Runtime/ClothingSystemRuntimeCommon/Public/ClothingAsset.h b/Engine/Source/Runtime/ClothingSystemRuntimeCommon/Public/ClothingAsset.h
--- a/Engine/Source/Runtime/ClothingSystemRuntimeCommon/Public/ClothingAsset.h
+++ b/Engine/Source/Runtime/ClothingSystemRuntimeCommon/Public/ClothingAsset.h
@@ -30,27 +30,6 @@ namespace ClothingAssetUtils
 		int32 AssetInternalLodIndex;
 	};
 
-	/**
-	 * Given a skeletal mesh, find all of the currently bound clothing assets and their binding information
-	 * @param InSkelMesh - The skeletal mesh to search
-	 * @param OutBindings - The list of bindings to write to
-	 */
-	UE_DEPRECATED(5.0, "This Handler Is Now Deprecated, use GetAllMeshClothingAssetBindings instead.")
-	void CLOTHINGSYSTEMRUNTIMECOMMON_API 
-	GetMeshClothingAssetBindings(
-		USkeletalMesh* InSkelMesh, 
-		TArray<FClothingAssetMeshBinding>& OutBindings);
-	
-	/**
-	 * Similar to above, but only inspects the specified LOD.
-	 */
-	UE_DEPRECATED(5.0, "This Handler Is Now Deprecated, use GetAllLodMeshClothingAssetBindings instead.")
-	void CLOTHINGSYSTEMRUNTIMECOMMON_API 
-	GetMeshClothingAssetBindings(
-		USkeletalMesh* InSkelMesh, 
-		TArray<FClothingAssetMeshBinding>& OutBindings, 
-		int32 InLodIndex);
-
 #if WITH_EDITOR
 	/**
 	 * Given a skeletal mesh model, find all of the currently bound clothing assets and their binding information
@@ -87,9 +66,6 @@ UCLASS(Abstract, MinimalAPI)
 class UClothingAssetCustomData : public UObject
 {
 	GENERATED_BODY()
-public:
-	virtual void BindToSkeletalMesh(USkeletalMesh* InSkelMesh, int32 InMeshLodIndex, int32 InSectionIndex, int32 InAssetLodIndex)
-	{}
 };
 
 /**
@@ -207,22 +183,6 @@ public:
 	CLOTHINGSYSTEMRUNTIMECOMMON_API virtual void PostEditChangeChainProperty(FPropertyChangedChainEvent& ChainEvent) override;
 #endif // WITH_EDITOR
 
-	/**
-	 * Builds self collision data.
-	 * The default behavior is if the \c ClothSimConfig reports that self collisions
-	 * are enabled, it envokes \c BuildSelfCollisionData() on each \c ClothLodData's
-	 * \c ClothPhysicalMeshData member.
-	 */
-	UE_DEPRECATED(5.0, "Cached data are now all rebuilt by calling InvalidateCachedData")
-PRAGMA_DISABLE_DEPRECATION_WARNINGS
-	virtual void BuildSelfCollisionData() override
-	{
-#if WITH_EDITORONLY_DATA
-		InvalidateFlaggedCachedData(EClothingCachedDataFlagsCommon::SelfCollisionData);
-#endif // WITH_EDITORONLY_DATA
-	}
-PRAGMA_ENABLE_DEPRECATION_WARNINGS
-
 	/** Return a const cloth config pointer of the desired cloth config type, or nullptr if there isn't any suitable. */
 	template<typename ClothConfigType, typename = typename TEnableIf<TIsDerivedFrom<ClothConfigType, UClothConfigBase>::IsDerived>::Type>
 	const ClothConfigType* GetClothConfig() const
@@ -304,13 +264,13 @@ PRAGMA_ENABLE_DEPRECATION_WARNINGS
 	UPROPERTY()
 	int32 ReferenceBoneIndex;
 
+#if WITH_EDITORONLY_DATA
 PRAGMA_DISABLE_DEPRECATION_WARNINGS
-	// Custom data applied by the importer depending on where the asset was imported from.
-	UPROPERTY()
+	UE_DEPRECATED(5.5, "This property is obsolete.")
+	UPROPERTY(Meta = (DeprecatedProperty, DeprecationMessage = "This property is obsolete."))
 	TObjectPtr<UClothingAssetCustomData> CustomData;
 PRAGMA_ENABLE_DEPRECATION_WARNINGS
 
-#if WITH_EDITORONLY_DATA
 	/** 
 	 * Deprecated property for transitioning the \c FClothConfig struct to the 
 	 * \c UClothConfigBase array, in a new property called \c ClothConfigs.
diff --git a/Engine/Source/Runtime/ClothingSystemRuntimeCommon/Public/ClothingSimulation.h b/Engine/Source/Runtime/ClothingSystemRuntimeCommon/Public/ClothingSimulation.h
--- a/Engine/Source/Runtime/ClothingSystemRuntimeCommon/Public/ClothingSimulation.h
+++ b/Engine/Source/Runtime/ClothingSystemRuntimeCommon/Public/ClothingSimulation.h
@@ -22,15 +22,10 @@ class FClothingSimulationContextCommon : public IClothingSimulationContext
 {
 public:
 	CLOTHINGSYSTEMRUNTIMECOMMON_API FClothingSimulationContextCommon();
-	PRAGMA_DISABLE_DEPRECATION_WARNINGS
 	FClothingSimulationContextCommon(const FClothingSimulationContextCommon&) = default;
 	FClothingSimulationContextCommon& operator=(const FClothingSimulationContextCommon&) = default;
-	CLOTHINGSYSTEMRUNTIMECOMMON_API PRAGMA_ENABLE_DEPRECATION_WARNINGS
 
-	virtual ~FClothingSimulationContextCommon() override;
-
-	UE_DEPRECATED(4.27, "Use the version with bIsInitialization instead.")
-	CLOTHINGSYSTEMRUNTIMECOMMON_API virtual void Fill(const USkeletalMeshComponent* InComponent, float InDeltaSeconds, float InMaxPhysicsDelta);
+	CLOTHINGSYSTEMRUNTIMECOMMON_API virtual ~FClothingSimulationContextCommon() override;
 
 	// Fill this context using the given skeletal mesh component
 	CLOTHINGSYSTEMRUNTIMECOMMON_API virtual void Fill(const USkeletalMeshComponent* InComponent, float InDeltaSeconds, float InMaxPhysicsDelta, bool bIsInitialization);
@@ -83,14 +78,6 @@ public:
 	
 	// Data read from the cache.
 	FClothingSimulationCacheData CacheData;
-
-	// World space cached positions for the kinematics targets.
-	UE_DEPRECATED(5.3, "Use CacheData.CachedPositions instead")
-	TArray<FVector> CachedPositions;
-
-	// World space cached velocities for the kinematics targets.
-	UE_DEPRECATED(5.3, "Use CacheData.CachedVelocities instead")
-	TArray<FVector> CachedVelocities;
 };
 
 // Base simulation to fill in common data for the base context
diff --git a/Engine/Source/Runtime/ClothingSystemRuntimeInterface/Public/ClothConfigBase.h b/Engine/Source/Runtime/ClothingSystemRuntimeInterface/Public/ClothConfigBase.h
--- a/Engine/Source/Runtime/ClothingSystemRuntimeInterface/Public/ClothConfigBase.h
+++ b/Engine/Source/Runtime/ClothingSystemRuntimeInterface/Public/ClothConfigBase.h
@@ -22,10 +22,6 @@ public:
 	CLOTHINGSYSTEMRUNTIMEINTERFACE_API UClothConfigBase();
 	CLOTHINGSYSTEMRUNTIMEINTERFACE_API virtual ~UClothConfigBase();
 
-	/** Return the self collision radius if building self collision indices is required for this config, or 0.f otherwise. */
-	UE_DEPRECATED(5.0, "Use NeedsSelfCollisionData and GetSelfCollisionRadius instead.")
-	virtual float NeedsSelfCollisionIndices() const { return NeedsSelfCollisionData() ? GetSelfCollisionRadius() : 0.f; }
-
 	/** Return wherether to pre-compute self collision data. */
 	CLOTHINGSYSTEMRUNTIMEINTERFACE_API virtual bool NeedsSelfCollisionData() const
 	PURE_VIRTUAL(UClothConfigBase::NeedsSelfCollisionData, return false;);
diff --git a/Engine/Source/Runtime/ClothingSystemRuntimeInterface/Public/ClothingAssetBase.h b/Engine/Source/Runtime/ClothingSystemRuntimeInterface/Public/ClothingAssetBase.h
--- a/Engine/Source/Runtime/ClothingSystemRuntimeInterface/Public/ClothingAssetBase.h
+++ b/Engine/Source/Runtime/ClothingSystemRuntimeInterface/Public/ClothingAssetBase.h
@@ -51,18 +51,6 @@ public:
 	CLOTHINGSYSTEMRUNTIMEINTERFACE_API virtual void UpdateAllLODBiasMappings(USkeletalMesh* SkeletalMesh)
 	PURE_VIRTUAL(UClothingAssetBase::UpdateAllLODBiasMappings,);
 
-	/**
-	 * Called on the clothing asset when the base data (physical mesh etc.) has
-	 * changed, so any intermediate generated data can be regenerated.
-	 */
-	UE_DEPRECATED(5.0, "Use InvalidateAllCachedData() instead")
-	virtual void InvalidateCachedData()
-	{
-#if WITH_EDITORONLY_DATA
-		InvalidateAllCachedData();
-#endif
-	}
-
 	/** Add a new LOD class instance. */
 	CLOTHINGSYSTEMRUNTIMEINTERFACE_API virtual int32 AddNewLod()
 	PURE_VIRTUAL(UClothingAssetBase::AddNewLod(), return INDEX_NONE;);
@@ -82,7 +70,7 @@ public:
 	 * has changed, so any intermediate generated data can be regenerated.
 	 */
 	CLOTHINGSYSTEMRUNTIMEINTERFACE_API virtual void InvalidateAllCachedData()
-	PURE_VIRTUAL(UClothingAssetBase::InvalidateCachedData(),);
+	PURE_VIRTUAL(UClothingAssetBase::InvalidateAllCachedData(),);
 #endif // WITH_EDITORONLY_DATA
 
 	/** 
@@ -102,11 +90,6 @@ public:
 	CLOTHINGSYSTEMRUNTIMEINTERFACE_API virtual int32 GetNumLods() const
 	PURE_VIRTUAL(UClothingAssetBase::GetNumLods(), return 0;);
 
-	/** Builds self collision data */
-	UE_DEPRECATED(5.0, "Cached data are now all rebuilt by calling InvalidateCachedData()")
-	CLOTHINGSYSTEMRUNTIMEINTERFACE_API virtual void BuildSelfCollisionData()
-	PURE_VIRTUAL(UClothingAssetBase::BuildSelfCollisionData(), );
-
 	/** Called after all cloth assets sharing the same simulation are added or loaded */
 	CLOTHINGSYSTEMRUNTIMEINTERFACE_API virtual void PostUpdateAllAssets()
 	PURE_VIRTUAL(UClothingAssetBase::PostUpdateAllAssets(), );
diff --git a/Engine/Source/Runtime/ClothingSystemRuntimeInterface/Public/ClothingSimulationFactory.h b/Engine/Source/Runtime/ClothingSystemRuntimeInterface/Public/ClothingSimulationFactory.h
--- a/Engine/Source/Runtime/ClothingSystemRuntimeInterface/Public/ClothingSimulationFactory.h
+++ b/Engine/Source/Runtime/ClothingSystemRuntimeInterface/Public/ClothingSimulationFactory.h
@@ -31,12 +31,6 @@ public:
 	// The feature name to register against for providers
 	static CLOTHINGSYSTEMRUNTIMEINTERFACE_API const FName FeatureName;
 
-	// Called by the engine to get the default clothing simulation factory to use
-	// for skeletal mesh components (see USkeletalMeshComponent constructor).
-	// Returns Factory class for simulations or nullptr to disable clothing simulation
-	UE_DEPRECATED(4.25, "GetDefaultSimulationFactoryClass() has been deprecated. Use IClothingSimulationFactoryClassProvider::GetSimulationFactoryClass() or UClothingSimulationFactory::GetDefaultClothingSimulationFactoryClass() instead.")
-	virtual UClass* GetDefaultSimulationFactoryClass() { return nullptr; }
-
 	// Called by the engine to get the clothing simulation factory associated with this
 	// provider for skeletal mesh components (see USkeletalMeshComponent constructor).
 	// Returns Factory class for simulations or nullptr to disable clothing simulation
diff --git a/Engine/Source/Runtime/ClothingSystemRuntimeInterface/Public/ClothingSimulationInterface.h b/Engine/Source/Runtime/ClothingSystemRuntimeInterface/Public/ClothingSimulationInterface.h
--- a/Engine/Source/Runtime/ClothingSystemRuntimeInterface/Public/ClothingSimulationInterface.h
+++ b/Engine/Source/Runtime/ClothingSystemRuntimeInterface/Public/ClothingSimulationInterface.h
@@ -50,10 +50,6 @@ protected:
 	 */
 	virtual void CreateActor(USkeletalMeshComponent* InOwnerComponent, UClothingAssetBase* InAsset, int32 SimDataIndex) = 0;
 
-	/** Do any required initialization after all the actors have been loaded */
-	UE_DEPRECATED(4.26, "No longer required.")
-	virtual void PostActorCreationInitialize() {}
-
 	/** Create a new context, will not be filled, call FillContext before simulating with this context */
 	virtual IClothingSimulationContext* CreateContext() = 0;
 
diff --git a/Engine/Source/Runtime/ClothingSystemRuntimeNv/Public/ClothConfigNv.h b/Engine/Source/Runtime/ClothingSystemRuntimeNv/Public/ClothConfigNv.h
--- a/Engine/Source/Runtime/ClothingSystemRuntimeNv/Public/ClothConfigNv.h
+++ b/Engine/Source/Runtime/ClothingSystemRuntimeNv/Public/ClothConfigNv.h
@@ -92,10 +92,6 @@ public:
 	// Return the collision radius required to calculate the self collision indices, or 0.f if self collision is disabled.
 	virtual float GetSelfCollisionRadius() const override { return NeedsSelfCollisionData() ? SelfCollisionRadius * SelfCollisionCullScale : 0.0f; }
 
-	// Return whether this Nv config has self collision.
-	UE_DEPRECATED(5.0, "Use NeedsSelfCollisionData instead.")
-	bool UseSelfCollisions() const { return NeedsSelfCollisionData(); }
-
 	// How wind should be processed, Accurate uses drag and lift to make the cloth react differently, legacy applies similar forces to all clothing without drag and lift (similar to APEX)
 	UPROPERTY(EditAnywhere, Category = ClothConfig)
 	EClothingWindMethodNv ClothingWindMethod;
