Chaos Cloth: Fix external collisions bookkeeping when using new solver.

diff --git a/Engine/Plugins/ChaosCloth/Source/ChaosCloth/Private/ChaosCloth/ChaosClothingSimulationSolver.cpp b/Engine/Plugins/ChaosCloth/Source/ChaosCloth/Private/ChaosCloth/ChaosClothingSimulationSolver.cpp
--- a/Engine/Plugins/ChaosCloth/Source/ChaosCloth/Private/ChaosCloth/ChaosClothingSimulationSolver.cpp
+++ b/Engine/Plugins/ChaosCloth/Source/ChaosCloth/Private/ChaosCloth/ChaosClothingSimulationSolver.cpp
@@ -783,6 +783,15 @@ int32 FClothingSimulationSolver::AddCollisionParticles(int32 NumCollisionParticl
 	int32 CollisionRangeId = INDEX_NONE;
 	if (Evolution)
 	{
+		if (Evolution->IsValidCollisionParticleRange(RecycledCollisionRangeId))
+		{
+			if (Evolution->GetCollisionParticleRange(RecycledCollisionRangeId).GetRangeSize() == NumCollisionParticles)
+			{
+				return RecycledCollisionRangeId;
+			}
+			// Size has changed. Remove the recycled collision range and reallocate some new particles
+			Evolution->RemoveCollisionParticleRange(RecycledCollisionRangeId);
+		}
 		CollisionRangeId = Evolution->AddCollisionParticleRange(GroupId, NumCollisionParticles, bActivateFalse);
 	}
 	else
diff --git a/Engine/Source/Runtime/Experimental/Chaos/Public/Chaos/SoftsEvolution.h b/Engine/Source/Runtime/Experimental/Chaos/Public/Chaos/SoftsEvolution.h
--- a/Engine/Source/Runtime/Experimental/Chaos/Public/Chaos/SoftsEvolution.h
+++ b/Engine/Source/Runtime/Experimental/Chaos/Public/Chaos/SoftsEvolution.h
@@ -87,6 +87,7 @@ public:
 	const TSet<int32>& GetGroupActiveCollisionParticleRanges(uint32 GroupId) const { return Groups.ActiveCollisionParticleRanges[GroupId]; }
 	CHAOS_API TArray<FSolverCollisionParticlesRange> GetActiveCollisionParticles(uint32 GroupId) const;
 
+	bool IsValidCollisionParticleRange(int32 CollisionRangeId) const { return CollisionRanges.ParticleRanges.IsValidIndex(CollisionRangeId); }
 	FSolverCollisionParticlesRange& GetCollisionParticleRange(int32 CollisionRangeId) { return CollisionRanges.ParticleRanges[CollisionRangeId]; }
 	const FSolverCollisionParticlesRange& GetCollisionParticleRange(int32 CollisionRangeId) const { return CollisionRanges.ParticleRanges[CollisionRangeId]; }
 	
