Chaos Cloth: fix max distance display

diff --git a/Engine/Plugins/ChaosCloth/Source/ChaosCloth/Private/ChaosCloth/ChaosClothVisualization.cpp b/Engine/Plugins/ChaosCloth/Source/ChaosCloth/Private/ChaosCloth/ChaosClothVisualization.cpp
--- a/Engine/Plugins/ChaosCloth/Source/ChaosCloth/Private/ChaosCloth/ChaosClothVisualization.cpp
+++ b/Engine/Plugins/ChaosCloth/Source/ChaosCloth/Private/ChaosCloth/ChaosClothVisualization.cpp
@@ -5,6 +5,7 @@
 #include "ChaosCloth/ChaosClothingSimulationMesh.h"
 #include "ChaosCloth/ChaosClothingSimulationCloth.h"
 #include "ChaosCloth/ChaosClothingSimulationCollider.h"
+#include "ChaosCloth/ChaosClothingSimulationConfig.h"
 #include "ChaosCloth/ChaosWeightMapTarget.h"
 #include "Chaos/DebugDrawQueue.h"
 #include "Chaos/Capsule.h"
@@ -18,6 +19,7 @@
 #include "Chaos/PBDAnimDriveConstraint.h"
 #include "Chaos/PBDBendingConstraints.h"
 #include "Chaos/PBDCollisionSpringConstraints.h"
+#include "Chaos/PBDFlatWeightMap.h"
 #include "Chaos/PBDLongRangeConstraints.h"
 #include "Chaos/PBDSelfCollisionSphereConstraints.h"
 #include "Chaos/PBDSphericalConstraint.h"
@@ -303,20 +305,20 @@ FLinearColor PseudoRandomColor(int32 NumColorRotations)
 				continue;
 			}
 
-			const TConstArrayView<FRealSingle>& MaxDistances = Cloth->GetWeightMapByProperty(Solver, TEXT("MaxDistance"));
-			if (!MaxDistances.Num())
-			{
-				continue;
-			}
+			const int32 NumParticles = Cloth->GetNumParticles(Solver);
+			const Softs::FPBDFlatWeightMapView MaxDistances(
+				Cloth->GetConfig()->GetProperties(Cloth->GetLODIndex(Solver)).GetWeightedFloatValue(TEXT("MaxDistance"), FVector2f(0.f, 1.f)),
+				Cloth->GetWeightMapByProperty(Solver, TEXT("MaxDistance")),
+				NumParticles);
 
 			const TConstArrayView<Softs::FSolverVec3> Positions = Cloth->GetAnimationPositions(Solver);
 			const TConstArrayView<Softs::FSolverReal> InvMasses = Cloth->GetParticleInvMasses(Solver);
-			check(MaxDistances.Num() == Positions.Num());
-			check(MaxDistances.Num() == InvMasses.Num());
+			check(NumParticles == Positions.Num());
+			check(NumParticles == InvMasses.Num());
 
 			for (int32 Index = 0; Index < MaxDistances.Num(); ++Index)
 			{
-				const FReal MaxDistance = MaxDistances[Index];
+				const FReal MaxDistance = (FReal)MaxDistances.GetValue(Index);
 				const FVector Position = LocalSpaceLocation + FVector(Positions[Index]);
 
 				const FText Text = FText::AsNumber(MaxDistance, &NumberFormattingOptions);
@@ -1522,22 +1524,22 @@ FLinearColor PseudoRandomColor(int32 NumColorRotations)
 				continue;
 			}
 
-			const TConstArrayView<FRealSingle>& MaxDistances = Cloth->GetWeightMapByProperty(Solver, TEXT("MaxDistance"));
-			if (!MaxDistances.Num())
-			{
-				continue;
-			}
+			const int32 NumParticles = Cloth->GetNumParticles(Solver);
+			const Softs::FPBDFlatWeightMapView MaxDistances(
+				Cloth->GetConfig()->GetProperties(Cloth->GetLODIndex(Solver)).GetWeightedFloatValue(TEXT("MaxDistance"), FVector2f(0.f, 1.f)),
+				Cloth->GetWeightMapByProperty(Solver, TEXT("MaxDistance")), 
+				NumParticles);
 
 			const TConstArrayView<Softs::FSolverReal> InvMasses = Cloth->GetParticleInvMasses(Solver);
 			const TConstArrayView<Softs::FSolverVec3> Positions = Cloth->GetAnimationPositions(Solver);
 			const TConstArrayView<Softs::FSolverVec3> Normals = Cloth->GetAnimationNormals(Solver);
-			check(Normals.Num() == Positions.Num());
-			check(MaxDistances.Num() == Positions.Num());
-			check(InvMasses.Num() == Positions.Num());
+			check(NumParticles == InvMasses.Num());
+			check(NumParticles == Positions.Num());
+			check(NumParticles == Normals.Num());
 
 			for (int32 Index = 0; Index < MaxDistances.Num(); ++Index)
 			{
-				const FReal MaxDistance = (FReal)MaxDistances[Index];
+				const FReal MaxDistance = (FReal)MaxDistances.GetValue(Index);
 				const FVector Position = LocalSpaceLocation + FVector(Positions[Index]);
 				if (InvMasses[Index] == (Softs::FSolverReal)0.)
 				{
diff --git a/Engine/Plugins/ChaosClothAssetEditor/Source/ChaosClothAssetEditor/Private/ChaosClothAsset/ClothEditorSimulationVisualization.cpp b/Engine/Plugins/ChaosClothAssetEditor/Source/ChaosClothAssetEditor/Private/ChaosClothAsset/ClothEditorSimulationVisualization.cpp
--- a/Engine/Plugins/ChaosClothAssetEditor/Source/ChaosClothAssetEditor/Private/ChaosClothAsset/ClothEditorSimulationVisualization.cpp
+++ b/Engine/Plugins/ChaosClothAssetEditor/Source/ChaosClothAssetEditor/Private/ChaosClothAsset/ClothEditorSimulationVisualization.cpp
@@ -243,8 +243,7 @@ const FVisualizationOption FVisualizationOption::OptionData[] =
 			Visualization.DrawMaxDistances(PDI);
 		}),
 		LOCTEXT("ChaosVisName_MaxDistance", "Max Distances"), 
-		LOCTEXT("ChaosVisName_MaxDistance_ToolTip", "Draws the current max distances for the sim particles as a line along its normal"), 
-		/*bDisablesSimulation =*/true),
+		LOCTEXT("ChaosVisName_MaxDistance_ToolTip", "Draws the current max distances for the sim particles as a line along its normal")),
 	FVisualizationOption(
 		FClothVisualizationDebugDrawTexts::CreateLambda([](const FClothEditorSimulationVisualization&, const ::Chaos::FClothVisualization& Visualization, FCanvas* Canvas, const FSceneView* SceneView)
 		{
