Chaos Cloth: Fix buckling ratio and buckling stiffness not updating correctly.

There were 3 different issues.
* PBDWeightMaps were not dirtying correctly, so the initial value wasn't getting set (this probably affected a handful of other constraints as well). This was the main problem.
* Technically the order of update between updating properties and initializing the constraint each timestep was incorrect for bending elements when using the legacy solver. In reality this just meant that property updates would be delayed by one frame.
* We would fallback to the default BucklingStiffness = 1 when BucklingStiffness.Low == BucklingStiffness.High == 0 when converting the SKM config to the new property config.

diff --git a/Engine/Plugins/ChaosCloth/Source/ChaosCloth/Private/ChaosCloth/ChaosClothConfig.cpp b/Engine/Plugins/ChaosCloth/Source/ChaosCloth/Private/ChaosCloth/ChaosClothConfig.cpp
--- a/Engine/Plugins/ChaosCloth/Source/ChaosCloth/Private/ChaosCloth/ChaosClothConfig.cpp
+++ b/Engine/Plugins/ChaosCloth/Source/ChaosCloth/Private/ChaosCloth/ChaosClothConfig.cpp
@@ -208,6 +208,15 @@ void UChaosClothConfig::PostLoad()
 		LocalDampingCoefficient = DampingCoefficient;
 		DampingCoefficient = 0.f;
 	}
+
+	if (FortniteMainBranchObjectVersion < FFortniteMainBranchObjectVersion::ChaosClothAllowZeroBucklingStiffness)
+	{
+		// Before this version, there was a bug where BucklingStiffness.Low == BucklingStiffness.High == 0 would result in the parameter being set to default (1.f).
+		if (BucklingStiffnessWeighted.Low <= 0 && BucklingStiffnessWeighted.High <= 0)
+		{
+			BucklingStiffnessWeighted.Low = BucklingStiffnessWeighted.High = 1.f;
+		}
+	}
 #endif  // #if WITH_EDITORONLY_DATA
 }
 
diff --git a/Engine/Plugins/ChaosCloth/Source/ChaosCloth/Private/ChaosCloth/ChaosClothingSimulationConfig.cpp b/Engine/Plugins/ChaosCloth/Source/ChaosCloth/Private/ChaosCloth/ChaosClothingSimulationConfig.cpp
--- a/Engine/Plugins/ChaosCloth/Source/ChaosCloth/Private/ChaosCloth/ChaosClothingSimulationConfig.cpp
+++ b/Engine/Plugins/ChaosCloth/Source/ChaosCloth/Private/ChaosCloth/ChaosClothingSimulationConfig.cpp
@@ -88,12 +88,9 @@ namespace Chaos
 
 					Property->AddValue(TEXT("BucklingRatio"), ClothConfig->BucklingRatio, AnimatablePropertyFlags);
 
-					if (ClothConfig->BucklingStiffnessWeighted.Low > 0.f || ClothConfig->BucklingStiffnessWeighted.High > 0.f)
-					{
-						const int32 BucklingStiffnessIndex = Property->AddProperty(TEXT("BucklingStiffness"), AnimatablePropertyFlags);
-						Property->SetWeightedValue(BucklingStiffnessIndex, ClothConfig->BucklingStiffnessWeighted.Low, ClothConfig->BucklingStiffnessWeighted.High);
-						Property->SetStringValue(BucklingStiffnessIndex, TEXT("BucklingStiffness"));
-					}
+					const int32 BucklingStiffnessIndex = Property->AddProperty(TEXT("BucklingStiffness"), AnimatablePropertyFlags);
+					Property->SetWeightedValue(BucklingStiffnessIndex, ClothConfig->BucklingStiffnessWeighted.Low, ClothConfig->BucklingStiffnessWeighted.High);
+					Property->SetStringValue(BucklingStiffnessIndex, TEXT("BucklingStiffness"));
 				}
 				else  // Not using bending elements
 				{
diff --git a/Engine/Source/Runtime/Core/Public/UObject/FortniteMainBranchObjectVersion.h b/Engine/Source/Runtime/Core/Public/UObject/FortniteMainBranchObjectVersion.h
--- a/Engine/Source/Runtime/Core/Public/UObject/FortniteMainBranchObjectVersion.h
+++ b/Engine/Source/Runtime/Core/Public/UObject/FortniteMainBranchObjectVersion.h
@@ -488,6 +488,9 @@ struct FFortniteMainBranchObjectVersion
 
 		//Allow custom import of morph target
 		MorphTargetCustomImport,
+
+		// Fix chaos cloth buckling stiffness parameter bug 
+		ChaosClothAllowZeroBucklingStiffness,
 		
 		// -----<new versions can be added above this line>-------------------------------------------------
 		VersionPlusOne,
diff --git a/Engine/Source/Runtime/Experimental/Chaos/Public/Chaos/PBDWeightMap.h b/Engine/Source/Runtime/Experimental/Chaos/Public/Chaos/PBDWeightMap.h
--- a/Engine/Source/Runtime/Experimental/Chaos/Public/Chaos/PBDWeightMap.h
+++ b/Engine/Source/Runtime/Experimental/Chaos/Public/Chaos/PBDWeightMap.h
@@ -59,15 +59,7 @@ public:
 	 * The weight map table only gets updated after ApplyValues is called.
 	 * Low and high values are clamped between [0,1]
 	 */
-	void SetWeightedValue(const FSolverVec2& InWeightedValue) { bIsDirty = WeightedValue != InWeightedValue; WeightedValue = InWeightedValue;}
-
-	/**
-	 * Set the low and high values of the weight map.
-	 * The weight map table only gets updated after ApplyValues is called.
-	 * Low and high values are not clamped. Commonly used for XPBD Stiffness values which are not [0,1]
-	 */
-	UE_DEPRECATED(5.2, "Use SetWeightedValue.")
-	void SetWeightedValueUnclamped(const FSolverVec2& InWeightedValue) { bIsDirty = WeightedValue != InWeightedValue; WeightedValue = InWeightedValue;}
+	void SetWeightedValue(const FSolverVec2& InWeightedValue) { bIsDirty |= WeightedValue != InWeightedValue; WeightedValue = InWeightedValue;}
 
 	/**
 	 * Return the low and high values set for this weight map.
