Chaos Cloth Asset: skip external collision bodies that are missing bones in their ref skeleton. This matches the behavior of legacy SKM cloth.

diff --git a/Engine/Plugins/ChaosCloth/Source/ChaosCloth/Private/ChaosCloth/ChaosClothingSimulationCollider.cpp b/Engine/Plugins/ChaosCloth/Source/ChaosCloth/Private/ChaosCloth/ChaosClothingSimulationCollider.cpp
--- a/Engine/Plugins/ChaosCloth/Source/ChaosCloth/Private/ChaosCloth/ChaosClothingSimulationCollider.cpp
+++ b/Engine/Plugins/ChaosCloth/Source/ChaosCloth/Private/ChaosCloth/ChaosClothingSimulationCollider.cpp
@@ -486,7 +486,8 @@ void FClothingSimulationCollider::ExtractPhysicsAssetCollision(
 	TArray<FLevelSetCollisionData>& LevelSetCollisions,
 	TArray<FSkinnedLevelSetCollisionData>& SkinnedLevelSetCollisions,
 	TArray<int32>& UsedBoneIndices,
-	bool bUseSphylsOnly)
+	bool bUseSphylsOnly,
+	bool bSkipMissingBones)
 {
 	ClothCollisionData.Reset();
 	UsedBoneIndices.Reset();
@@ -503,6 +504,11 @@ void FClothingSimulationCollider::ExtractPhysicsAssetCollision(
 			}
 
 			const int32 MeshBoneIndex = ReferenceSkeleton ? ReferenceSkeleton->FindBoneIndex(BodySetup->BoneName) : INDEX_NONE;
+			if (bSkipMissingBones && MeshBoneIndex == INDEX_NONE)
+			{
+				continue;
+			}
+
 			const int32 MappedBoneIndex = UsedBoneIndices.Add(MeshBoneIndex);
 			
 			// Add capsules
diff --git a/Engine/Plugins/ChaosCloth/Source/ChaosCloth/Public/ChaosCloth/ChaosClothingSimulationCollider.h b/Engine/Plugins/ChaosCloth/Source/ChaosCloth/Public/ChaosCloth/ChaosClothingSimulationCollider.h
--- a/Engine/Plugins/ChaosCloth/Source/ChaosCloth/Public/ChaosCloth/ChaosClothingSimulationCollider.h
+++ b/Engine/Plugins/ChaosCloth/Source/ChaosCloth/Public/ChaosCloth/ChaosClothingSimulationCollider.h
@@ -100,7 +100,8 @@ namespace Chaos
 			TArray<FLevelSetCollisionData>& LevelSetCollisions,
 			TArray<FSkinnedLevelSetCollisionData>& SkinnedLevelSetCollisions,
 			TArray<int32>& UsedBoneIndices,
-			bool bUseSphylOnly = false);
+			bool bUseSphylOnly = false,
+			bool bSkipMissingBones = false);
 
 	private:
 		CHAOSCLOTH_API int32 GetNumGeometries(int32 InSlotIndex) const;
diff --git a/Engine/Plugins/ChaosClothAsset/Source/ChaosClothAssetEngine/Private/ChaosClothAsset/CollisionSources.cpp b/Engine/Plugins/ChaosClothAsset/Source/ChaosClothAssetEngine/Private/ChaosClothAsset/CollisionSources.cpp
--- a/Engine/Plugins/ChaosClothAsset/Source/ChaosClothAssetEngine/Private/ChaosClothAsset/CollisionSources.cpp
+++ b/Engine/Plugins/ChaosClothAsset/Source/ChaosClothAssetEngine/Private/ChaosClothAsset/CollisionSources.cpp
@@ -163,6 +163,7 @@ namespace UE::Chaos::ClothAsset
 					TArray<::Chaos::FClothingSimulationCollider::FLevelSetCollisionData> LevelSetCollisions;
 					TArray<::Chaos::FClothingSimulationCollider::FSkinnedLevelSetCollisionData> SkinnedLevelSetCollisions;
 
+					constexpr bool bSkipMissingBones = true;
 					::Chaos::FClothingSimulationCollider::ExtractPhysicsAssetCollision(
 						PhysicsAsset,
 						&SkinnedAsset->GetRefSkeleton(),
@@ -170,7 +171,8 @@ namespace UE::Chaos::ClothAsset
 						LevelSetCollisions,
 						SkinnedLevelSetCollisions,
 						CachedUsedBoneIndices,
-						bUseSphylsOnly);
+						bUseSphylsOnly,
+						bSkipMissingBones);
 				}
 			}
 
