Chaos Cloth: Velocity and Acceleration clamps. Similar to RBAN, the clamps are applied after the scaling.

diff --git a/Engine/Plugins/ChaosCloth/Source/ChaosCloth/Private/ChaosCloth/ChaosClothingSimulationCloth.cpp b/Engine/Plugins/ChaosCloth/Source/ChaosCloth/Private/ChaosCloth/ChaosClothingSimulationCloth.cpp
--- a/Engine/Plugins/ChaosCloth/Source/ChaosCloth/Private/ChaosCloth/ChaosClothingSimulationCloth.cpp
+++ b/Engine/Plugins/ChaosCloth/Source/ChaosCloth/Private/ChaosCloth/ChaosClothingSimulationCloth.cpp
@@ -37,6 +37,8 @@ namespace ClothingSimulationClothDefault
 	constexpr float GravityZOverride = Softs::FExternalForces::DefaultGravityZOverride; // -980.665f;
 	constexpr float VelocityScale = 0.75f;
 	constexpr float MaxVelocityScale = 1.f;
+	constexpr float MaxVelocity = TNumericLimits<float>::Max();
+	constexpr float MaxAcceleration = TNumericLimits<float>::Max();
 	constexpr float FictitiousAngularScale = Softs::FExternalForces::DefaultFictitiousAngularScale; // 1.f;
 	constexpr int32 MultiResCoarseLODIndex = INDEX_NONE;
 }
@@ -893,12 +895,18 @@ void FClothingSimulationCloth::Update(FClothingSimulationSolver* Solver)
 		FRealSingle OutAngularVelocityScale;
 		FRealSingle OutMaxVelocityScale;
 		bool bDisableFictitiousForces = false;
+		FVec3f MaxLinearVelocity(ClothingSimulationClothDefault::MaxVelocity);
+		FVec3f MaxLinearAcceleration(ClothingSimulationClothDefault::MaxAcceleration);
+		FRealSingle MaxAngularVelocity = ClothingSimulationClothDefault::MaxVelocity;
+		FRealSingle MaxAngularAcceleration = ClothingSimulationClothDefault::MaxAcceleration;
 		if (bNeedsReset)
 		{
 			// Make sure not to do any pre-sim transform just after a reset
 			OutLinearVelocityScale = FVec3f(1.f);
 			OutAngularVelocityScale = 1.f;
 			OutMaxVelocityScale = 1.f;
+			ReferenceSpaceAngularVelocity = FVec3(0.);
+			ReferenceSpaceVelocity = FVec3(0.);
 			bDisableFictitiousForces = true; // It doesn't actually matter what value we set here since AngularVelocityScale == 1 means fictitious forces will be 0.
 
 			// Reset to start pose
@@ -920,6 +928,8 @@ void FClothingSimulationCloth::Update(FClothingSimulationSolver* Solver)
 			OutLinearVelocityScale = FVec3f(0.f);
 			OutAngularVelocityScale = 0.f;
 			OutMaxVelocityScale = 1.f;
+			ReferenceSpaceAngularVelocity = FVec3(0.);
+			ReferenceSpaceVelocity = FVec3(0.);
 			bDisableFictitiousForces = true; // Disable fictitious forces. Otherwise they will be applied since AngularVelocityScale < 1.
 			UE_LOG(LogChaosCloth, VeryVerbose, TEXT("Cloth in group Id %d Needs teleport."), GroupId);
 		}
@@ -929,6 +939,10 @@ void FClothingSimulationCloth::Update(FClothingSimulationSolver* Solver)
 			OutLinearVelocityScale = ConfigProperties.GetValue<FVector3f>(TEXT("LinearVelocityScale"), FVector3f(ClothingSimulationClothDefault::VelocityScale));
 			OutAngularVelocityScale = ConfigProperties.GetValue<float>(TEXT("AngularVelocityScale"), ClothingSimulationClothDefault::VelocityScale);
 			OutMaxVelocityScale = ConfigProperties.GetValue<float>(TEXT("MaxVelocityScale"), ClothingSimulationClothDefault::MaxVelocityScale);
+			MaxLinearVelocity = ConfigProperties.GetValue<FVector3f>(TEXT("MaxLinearVelocity"), MaxLinearVelocity);
+			MaxLinearAcceleration = ConfigProperties.GetValue<FVector3f>(TEXT("MaxLinearAcceleration"), MaxLinearAcceleration);
+			MaxAngularVelocity = ConfigProperties.GetValue<float>(TEXT("MaxAngularVelocity"), MaxAngularVelocity);
+			MaxAngularAcceleration = ConfigProperties.GetValue<float>(TEXT("MaxAngularAcceleration"), MaxAngularAcceleration);
 		}
 
 		// NOTE: Force-based solver doesn't actually use FictitiousAngularScale here. It gets it from the property collection directly.
@@ -937,10 +951,16 @@ void FClothingSimulationCloth::Update(FClothingSimulationSolver* Solver)
 			GroupId,
 			OldReferenceSpaceTransform,
 			ReferenceSpaceTransform,
+			ReferenceSpaceVelocity,
+			ReferenceSpaceAngularVelocity,
 			OutLinearVelocityScale,
+			MaxLinearVelocity,
+			MaxLinearAcceleration,
 			OutAngularVelocityScale,
+			MaxAngularVelocity,
+			MaxAngularAcceleration,
 			FictitiousAngularScale,
-			OutMaxVelocityScale, 
+			OutMaxVelocityScale,
 			bDisableFictitiousForces);
 		if (!Solver->IsLegacySolver())
 		{
diff --git a/Engine/Plugins/ChaosCloth/Source/ChaosCloth/Private/ChaosCloth/ChaosClothingSimulationSolver.cpp b/Engine/Plugins/ChaosCloth/Source/ChaosCloth/Private/ChaosCloth/ChaosClothingSimulationSolver.cpp
--- a/Engine/Plugins/ChaosCloth/Source/ChaosCloth/Private/ChaosCloth/ChaosClothingSimulationSolver.cpp
+++ b/Engine/Plugins/ChaosCloth/Source/ChaosCloth/Private/ChaosCloth/ChaosClothingSimulationSolver.cpp
@@ -1079,23 +1079,29 @@ void FClothingSimulationSolver::SetParticleMassFromDensity(int32 ParticleRangeId
 	}
 }
 
-void FClothingSimulationSolver::SetReferenceVelocityScale(
-	uint32 GroupId,
-	const FRigidTransform3& OldReferenceSpaceTransform,  // Transforms are in world space so have to be FReal based for LWC
+void FClothingSimulationSolver::SetReferenceVelocityScale(uint32 GroupId,
+	const FRigidTransform3& OldReferenceSpaceTransform,
 	const FRigidTransform3& ReferenceSpaceTransform,
+	TVec3<FReal>& InOutReferenceVelocity, // Old reference velocity is passed in. New reference velocity is returned.
+	TVec3<FReal>& InOutReferenceAngularVelocity, // Old reference velocity is passed in. New reference velocity is returned.
 	const TVec3<FRealSingle>& LinearVelocityScale,
-	FRealSingle AngularVelocityScale, FRealSingle FictitiousAngularScale,
-	FRealSingle MaxVelocityScale, 
+	const TVec3<FRealSingle>& MaxLinearVelocity,
+	const TVec3<FRealSingle>& MaxLinearAcceleration,
+	FRealSingle AngularVelocityScale,
+	FRealSingle MaxAngularVelocity,
+	FRealSingle MaxAngularAcceleration,
+	FRealSingle FictitiousAngularScale,
+	FRealSingle MaxVelocityScale,
 	bool bDisableFictitiousForces
 )
 {
 	FRigidTransform3 OldRootBoneLocalTransform = OldReferenceSpaceTransform;
 	OldRootBoneLocalTransform.AddToTranslation(-OldLocalSpaceLocation);
 
-	const FReal SolverVelocityScale = bClothSolverUseVelocityScale ? VelocityScale : (FReal)1.;
+	const FVec3 OldReferenceVelocity(InOutReferenceVelocity); // In local space
+	const FVec3 OldReferenceAngularVelocity(InOutReferenceAngularVelocity); // In local space
 
-	// Calculate deltas
-	const FRigidTransform3 DeltaTransform = ReferenceSpaceTransform.GetRelativeTransform(OldReferenceSpaceTransform);
+	const FReal SolverVelocityScale = bClothSolverUseVelocityScale ? VelocityScale : (FReal)1.;
 
 	auto CalculateClampedVelocityScale = [MaxVelocityScale, SolverVelocityScale](FRealSingle InVelocityScale)
 	{
@@ -1109,27 +1115,94 @@ void FClothingSimulationSolver::SetReferenceVelocityScale(
 		return FMath::Clamp(SolverVelocityScale, (FReal)0., (FReal)1.) * FMath::Clamp(InVelocityScale, (FReal)0., (FReal)MaxVelocityScale);
 	};
 
-	// Apply linear velocity scale
-	const FVec3 LinearRatio = FVec3(1.) - FVec3(
+	// Calculate deltas
+	const FRotation3 OldReferenceSpaceRotationInverse = OldReferenceSpaceTransform.GetRotation().Inverse();
+	const FVec3 LocalDeltaTranslation = FRigidTransform3::SubtractTranslations(ReferenceSpaceTransform, OldReferenceSpaceTransform);
+	const FRotation3 FullDeltaRotationQ = OldReferenceSpaceRotationInverse * ReferenceSpaceTransform.GetRotation();
+	FReal FullDeltaAngle = FullDeltaRotationQ.GetAngle();
+	const FVec3 FullAxis = FullDeltaRotationQ.GetRotationAxis();
+	if (FullDeltaAngle > (FReal)PI)
+	{
+		FullDeltaAngle -= (FReal)2. * (FReal)PI;
+	}	
+	const FVec3 FullDeltaRotation = FullDeltaAngle * FullAxis;
+
+	// Do all linear velocity changes in old reference space.
+	const FVec3 ReferenceSpaceDeltaTranslation= OldReferenceSpaceRotationInverse * LocalDeltaTranslation;
+
+	// Apply linear velocity scale.
+	const FVec3 LinearRatio = FVec3(
 		CalculateClampedVelocityScale(LinearVelocityScale[0]),
 		CalculateClampedVelocityScale(LinearVelocityScale[1]),
 		CalculateClampedVelocityScale(LinearVelocityScale[2]));
-	const FVec3 DeltaPosition = LinearRatio * DeltaTransform.GetTranslation();
+
+	FVec3 AppliedDeltaTranslation = LinearRatio * ReferenceSpaceDeltaTranslation;
 
 	// Apply angular velocity scale
-	FRotation3 DeltaRotation = DeltaTransform.GetRotation();
-	FReal DeltaAngle = DeltaRotation.GetAngle();
-	FVec3 Axis = DeltaRotation.GetRotationAxis();
-	if (DeltaAngle > (FReal)PI)
+	const FReal AngularRatio = CalculateClampedVelocityScale(AngularVelocityScale);
+	FVec3 AppliedDeltaRotation = AngularRatio * FullDeltaRotation;
+
+	// Apply clamps (since we're ultimately just modifying a single transform, remove accelerations from the velocity first, then clamp velocities).
+	if (DeltaTime > 0.f)
 	{
-		DeltaAngle -= (FReal)2. * (FReal)PI;
+		// Only update old reference velocity when ticking.
+		InOutReferenceVelocity = LocalDeltaTranslation / DeltaTime;
+		InOutReferenceAngularVelocity = FullDeltaRotation / DeltaTime;
+	
+		FVec3 AppliedLinearVelocity = AppliedDeltaTranslation / DeltaTime;
+		FVec3 AppliedAngularVelocity = InOutReferenceAngularVelocity;
+
+		const FVec3 OldScaledReferenceSpaceVelocity = LinearRatio * (OldReferenceSpaceRotationInverse * OldReferenceVelocity);
+		const FVec3 ScaledReferenceSpaceAccelerationTimesDt = AppliedLinearVelocity - OldScaledReferenceSpaceVelocity;
+		for (int32 Index = 0; Index < 3; ++Index)
+		{
+			if (MaxLinearAcceleration[Index] != TNumericLimits<FRealSingle>::Max() && MaxLinearAcceleration[Index] >= 0.f
+				&& FMath::Abs(ScaledReferenceSpaceAccelerationTimesDt[Index]) > (FReal)(MaxLinearAcceleration[Index] * DeltaTime))
+			{
+				AppliedLinearVelocity[Index] = OldScaledReferenceSpaceVelocity[Index] + FMath::Sign(ScaledReferenceSpaceAccelerationTimesDt[Index]) * (FReal)(MaxLinearAcceleration[Index] * DeltaTime);
+			}
+			if (MaxLinearVelocity[Index] != TNumericLimits<FRealSingle>::Max() && MaxLinearVelocity[Index] >= 0.f &&
+				FMath::Abs(AppliedLinearVelocity[Index]) > (FReal)MaxLinearVelocity[Index])
+			{
+				AppliedLinearVelocity[Index] = FMath::Sign(AppliedLinearVelocity[Index]) * (FReal)MaxLinearVelocity[Index];
+			}
+
+		}
+
+		AppliedDeltaTranslation = AppliedLinearVelocity * DeltaTime;
+
+		const FVec3 OldScaledAngularVelocity = AngularRatio * OldReferenceAngularVelocity;
+		const FVec3 ScaledReferenceSpaceAngularAccelerationTimesDt = AppliedAngularVelocity - OldScaledAngularVelocity;
+		if (MaxAngularAcceleration != TNumericLimits<FRealSingle>::Max() && MaxAngularAcceleration >= 0.f)
+		{
+			const FReal LenSq = ScaledReferenceSpaceAngularAccelerationTimesDt.SquaredLength();
+			if (LenSq > FMath::Square(MaxAngularAcceleration * DeltaTime))
+			{
+				const FReal Scale = MaxAngularAcceleration * DeltaTime * FMath::InvSqrt(LenSq);
+				AppliedAngularVelocity = OldScaledAngularVelocity + Scale * ScaledReferenceSpaceAngularAccelerationTimesDt;
+			}
+		}
+		if (MaxAngularVelocity != TNumericLimits<FRealSingle>::Max() && MaxAngularVelocity >= 0.f)
+		{
+			const FReal LenSq = AppliedAngularVelocity.SquaredLength();
+			if (LenSq > FMath::Square(MaxAngularVelocity))
+			{
+				const FReal Scale = MaxAngularVelocity * FMath::InvSqrt(LenSq);
+				AppliedAngularVelocity *= Scale;
+			}
+		}
+
+		AppliedDeltaRotation = AppliedAngularVelocity * DeltaTime;
 	}
+	const FVec3 DeltaPosition = ReferenceSpaceDeltaTranslation - AppliedDeltaTranslation;
+	const FVec3 DeltaRotation = FullDeltaRotation - AppliedDeltaRotation;
+	const FReal DeltaAngle = DeltaRotation.Length();
+	const FVec3 Axis = DeltaAngle > UE_KINDA_SMALL_NUMBER ? DeltaRotation / DeltaAngle : FVec3::XAxisVector;
 
-	const FReal PartialDeltaAngle = DeltaAngle * ((FReal)1. - CalculateClampedVelocityScale(AngularVelocityScale));
-	DeltaRotation = UE::Math::TQuat<FReal>(Axis, PartialDeltaAngle);
+	const UE::Math::TQuat<FReal> DeltaRotationQuat(Axis, DeltaAngle);
 
 	// Transform points back into the previous frame of reference before applying the adjusted deltas
-	const FRigidTransform3 PreSimulationTransform = OldRootBoneLocalTransform.Inverse() * FRigidTransform3(DeltaPosition, DeltaRotation) * OldRootBoneLocalTransform;
+	const FRigidTransform3 PreSimulationTransform = OldRootBoneLocalTransform.Inverse() * FRigidTransform3(DeltaPosition, DeltaRotationQuat) * OldRootBoneLocalTransform;
 	PreSimulationTransforms[GroupId] = Softs::FSolverRigidTransform3(  // Store the delta in solver precision, no need for LWC here
 		Softs::FSolverVec3(PreSimulationTransform.GetTranslation()),
 		Softs::FSolverRotation3(PreSimulationTransform.GetRotation()));
@@ -1138,7 +1211,7 @@ void FClothingSimulationSolver::SetReferenceVelocityScale(
 	const FReal AppliedFictitiousAngularScale = bDisableFictitiousForces ? (FReal)0.f : (PBDEvolution ? FMath::Min((FReal)2., (FReal)FictitiousAngularScale) : (FReal)1.);
 
 	// Save the reference bone relative angular velocity for calculating the fictitious forces
-	const FVec3 FictitiousAngularDisplacement = ReferenceSpaceTransform.TransformVector(Axis * PartialDeltaAngle)
+	const FVec3 FictitiousAngularDisplacement = ReferenceSpaceTransform.TransformVector(DeltaRotation)
 		* AppliedFictitiousAngularScale;
 	FictitiousAngularVelocities[GroupId] = DeltaTime > (Softs::FSolverReal)0.f ? Softs::FSolverVec3(FictitiousAngularDisplacement) / DeltaTime : Softs::FSolverVec3(0.f);
 	ReferenceSpaceLocations[GroupId] = ReferenceSpaceTransform.GetLocation() - LocalSpaceLocation;
diff --git a/Engine/Plugins/ChaosCloth/Source/ChaosCloth/Public/ChaosCloth/ChaosClothingSimulationCloth.h b/Engine/Plugins/ChaosCloth/Source/ChaosCloth/Public/ChaosCloth/ChaosClothingSimulationCloth.h
--- a/Engine/Plugins/ChaosCloth/Source/ChaosCloth/Public/ChaosCloth/ChaosClothingSimulationCloth.h
+++ b/Engine/Plugins/ChaosCloth/Source/ChaosCloth/Public/ChaosCloth/ChaosClothingSimulationCloth.h
@@ -147,6 +147,8 @@ namespace Chaos
 
 		// Reference space transform
 		FRigidTransform3 ReferenceSpaceTransform;  // TODO: Add override in the style of LODIndexOverride
+		FVec3 ReferenceSpaceVelocity;
+		FVec3 ReferenceSpaceAngularVelocity;
 
 		// LOD data
 		TArray<TUniquePtr<FLODData>> LODData;
diff --git a/Engine/Plugins/ChaosCloth/Source/ChaosCloth/Public/ChaosCloth/ChaosClothingSimulationSolver.h b/Engine/Plugins/ChaosCloth/Source/ChaosCloth/Public/ChaosCloth/ChaosClothingSimulationSolver.h
--- a/Engine/Plugins/ChaosCloth/Source/ChaosCloth/Public/ChaosCloth/ChaosClothingSimulationSolver.h
+++ b/Engine/Plugins/ChaosCloth/Source/ChaosCloth/Public/ChaosCloth/ChaosClothingSimulationSolver.h
@@ -132,13 +132,36 @@ namespace Chaos
 		// Set the amount of velocity allowed to filter from the given change in reference space transform, including local simulation space.
 		// NOTE: Force-based solver does not apply FictitiousAngularScale here. It's applied directly via the PropertyCollection.
 		CHAOSCLOTH_API void SetReferenceVelocityScale(uint32 GroupId,
+			const FRigidTransform3& OldReferenceSpaceTransform,
+			const FRigidTransform3& ReferenceSpaceTransform,
+			TVec3<FReal>& InOutReferenceVelocity, // Old reference velocity is passed in. New reference velocity is returned.
+			TVec3<FReal>& InOutReferenceAngularVelocity, // Old reference velocity is passed in. New reference velocity is returned.
+			const TVec3<FRealSingle>& LinearVelocityScale,
+			const TVec3<FRealSingle>& MaxLinearVelocity,
+			const TVec3<FRealSingle>& MaxLinearAcceleration,
+			FRealSingle AngularVelocityScale,
+			FRealSingle MaxAngularVelocity,
+			FRealSingle MaxAngularAcceleration,
+			FRealSingle FictitiousAngularScale,
+			FRealSingle MaxVelocityScale,
+			bool bDisableFictitiousForces);
+
+		void SetReferenceVelocityScale(uint32 GroupId,
 			const FRigidTransform3& OldReferenceSpaceTransform,
 			const FRigidTransform3& ReferenceSpaceTransform,
 			const TVec3<FRealSingle>& LinearVelocityScale,
 			FRealSingle AngularVelocityScale,
 			FRealSingle FictitiousAngularScale,
 			FRealSingle MaxVelocityScale = 1.f,
-			bool bDisableFictitiousForces = false);
+			bool bDisableFictitiousForces = false)
+		{
+			// Acceleration clamps are disabled, so it doesn't matter what the old velocities were.
+			TVec3<FReal> ReferenceVelocity(0.);
+			TVec3<FReal> ReferenceAngularVelocity(0.);
+			SetReferenceVelocityScale(GroupId, OldReferenceSpaceTransform, ReferenceSpaceTransform, ReferenceVelocity, ReferenceAngularVelocity, LinearVelocityScale,
+				TVec3<FRealSingle>(TNumericLimits<FRealSingle>::Max()), TVec3<FRealSingle>(TNumericLimits<FRealSingle>::Max()), AngularVelocityScale, FictitiousAngularScale,
+				MaxVelocityScale, TNumericLimits<FRealSingle>::Max(), TNumericLimits<FRealSingle>::Max(), bDisableFictitiousForces);
+		}
 
 		/** PBDSolver version */
 		CHAOSCLOTH_API void SetProperties(
diff --git a/Engine/Plugins/ChaosClothAssetEditor/Source/ChaosClothAssetDataflowNodes/Private/ChaosClothAsset/SimulationVelocityScaleConfigNode.cpp b/Engine/Plugins/ChaosClothAssetEditor/Source/ChaosClothAssetDataflowNodes/Private/ChaosClothAsset/SimulationVelocityScaleConfigNode.cpp
--- a/Engine/Plugins/ChaosClothAssetEditor/Source/ChaosClothAssetDataflowNodes/Private/ChaosClothAsset/SimulationVelocityScaleConfigNode.cpp
+++ b/Engine/Plugins/ChaosClothAssetEditor/Source/ChaosClothAssetDataflowNodes/Private/ChaosClothAsset/SimulationVelocityScaleConfigNode.cpp
@@ -13,9 +13,13 @@ FChaosClothAssetSimulationVelocityScaleConfigNode::FChaosClothAssetSimulationVel
 
 void FChaosClothAssetSimulationVelocityScaleConfigNode::AddProperties(FPropertyHelper& PropertyHelper) const
 {
-	
+	static const FVector3f DefaultLinearClamp = { TNumericLimits<float>::Max(), TNumericLimits<float>::Max(), TNumericLimits<float>::Max() };
 	PropertyHelper.SetProperty(this, &LinearVelocityScale);
+	PropertyHelper.SetProperty(TEXT("MaxLinearVelocity"), bEnableLinearVelocityClamping ? MaxLinearVelocity : DefaultLinearClamp);
+	PropertyHelper.SetProperty(TEXT("MaxLinearAcceleration"), bEnableLinearAccelerationClamping ? MaxLinearAcceleration : DefaultLinearClamp);
 	PropertyHelper.SetProperty(this, &AngularVelocityScale);
+	PropertyHelper.SetProperty(TEXT("MaxAngularVelocity"), bEnableAngularVelocityClamping ? MaxAngularVelocity : TNumericLimits<float>::Max());
+	PropertyHelper.SetProperty(TEXT("MaxAngularAcceleration"), bEnableAngularAccelerationClamping ? MaxAngularAcceleration : TNumericLimits<float>::Max());
 	PropertyHelper.SetProperty(this, &MaxVelocityScale);
 	PropertyHelper.SetProperty(this, &FictitiousAngularScale);
 }
diff --git a/Engine/Plugins/ChaosClothAssetEditor/Source/ChaosClothAssetDataflowNodes/Public/ChaosClothAsset/SimulationVelocityScaleConfigNode.h b/Engine/Plugins/ChaosClothAssetEditor/Source/ChaosClothAssetDataflowNodes/Public/ChaosClothAsset/SimulationVelocityScaleConfigNode.h
--- a/Engine/Plugins/ChaosClothAssetEditor/Source/ChaosClothAssetDataflowNodes/Public/ChaosClothAsset/SimulationVelocityScaleConfigNode.h
+++ b/Engine/Plugins/ChaosClothAssetEditor/Source/ChaosClothAssetDataflowNodes/Public/ChaosClothAsset/SimulationVelocityScaleConfigNode.h
@@ -23,6 +23,32 @@ public:
 	UPROPERTY(EditAnywhere, Category = "Velocity Scale Properties", meta = (UIMin = "0", UIMax = "1", ClampMin = "0", ClampMax = "100"))
 	FVector3f LinearVelocityScale = { 0.75f, 0.75f, 0.75f };
 
+	/**
+	 * Enable linear velocity clamping.
+	 */
+	UPROPERTY(EditAnywhere, Category = "Velocity Scale Properties")
+	bool bEnableLinearVelocityClamping = false;
+
+	/**
+	 * The maximum amount of linear velocity sent to the local cloth space from the reference bone
+	 * (the closest bone to the root on which the cloth section has been skinned, or the root itself if the cloth isn't skinned).
+	 */
+	UPROPERTY(EditAnywhere, Category = "Velocity Scale Properties", meta = (ClampMin = "0", EditCondition = "bEnableLinearVelocityClamping"))
+	FVector3f MaxLinearVelocity = { 1000.f, 1000.f, 1000.f }; // Approx 22mph or 36kph per direction
+
+	/**
+	 * Enable linear acceleration clamping.
+	 */
+	UPROPERTY(EditAnywhere, Category = "Velocity Scale Properties")
+	bool bEnableLinearAccelerationClamping = false;
+
+	/**
+	 * The maximum amount of linear acceleration sent to the local cloth space from the reference bone
+	 * (the closest bone to the root on which the cloth section has been skinned, or the root itself if the cloth isn't skinned).
+	 */
+	UPROPERTY(EditAnywhere, Category = "Velocity Scale Properties", meta = (ClampMin = "0", EditCondition = "bEnableLinearAccelerationClamping"))
+	FVector3f MaxLinearAcceleration = { 60000.f, 60000.f, 60000.f };
+
 	/**
 	 * The amount of angular velocities sent to the local cloth space from the reference bone
 	 * (the closest bone to the root on which the cloth section has been skinned, or the root itself if the cloth isn't skinned).
@@ -31,6 +57,32 @@ public:
 	UPROPERTY(EditAnywhere, Category = "Velocity Scale Properties", meta = (UIMin = "0", UIMax = "1", ClampMin = "0", ClampMax = "100"))
 	float AngularVelocityScale = 0.75f;
 
+	/**
+	 * Enable angular velocity clamping.
+	 */
+	UPROPERTY(EditAnywhere, Category = "Velocity Scale Properties")
+	bool bEnableAngularVelocityClamping = false;
+
+	/**
+	 * The maximum amount of angular velocity sent to the local cloth space from the reference bone
+	 * (the closest bone to the root on which the cloth section has been skinned, or the root itself if the cloth isn't skinned).
+	 */
+	UPROPERTY(EditAnywhere, Category = "Velocity Scale Properties", meta = (ClampMin = "0", EditCondition = "bEnableAngularVelocityClamping"))
+	float MaxAngularVelocity = 200.f;
+
+	/**
+	 * Enable angular acceleration clamping.
+	 */
+	UPROPERTY(EditAnywhere, Category = "Velocity Scale Properties")
+	bool bEnableAngularAccelerationClamping = false;
+
+	/**
+	 * The maximum amount of angular acceleration sent to the local cloth space from the reference bone
+	 * (the closest bone to the root on which the cloth section has been skinned, or the root itself if the cloth isn't skinned).
+	 */
+	UPROPERTY(EditAnywhere, Category = "Velocity Scale Properties", meta = (ClampMin = "0", EditCondition = "bEnableAngularAccelerationClamping"))
+	float MaxAngularAcceleration = 12000.f;
+
 	/**
 	 * Clamp on Linear and Angular Velocity Scale. The final velocity scale (e.g., including contributions from blueprints) will be clamped to this value.
 	 */
@@ -40,8 +92,9 @@ public:
 	/**
 	 * The portion of the angular velocity that is used to calculate the strength of all fictitious forces (e.g. centrifugal force).
 	 * This parameter is only having an effect on the portion of the reference bone's angular velocity that has been removed from the
-	 * simulation via the Angular Velocity Scale parameter. This means it has no effect when AngularVelocityScale is set to 1 in which
-	 * case the cloth is simulated with full world space angular velocities and subjected to the true physical world inertial forces.
+	 * simulation via the Angular Velocity Scale parameter. This means it has no effect when AngularVelocityScale is set to 1 and 
+	 * Angular Velocity and Acceleration clamps are disabled, in which case the cloth is simulated with full world space angular 
+	 * velocities and subjected to the true physical world inertial forces.
 	 * Values range from 0 to 2, with 0 showing no centrifugal effect, 1 full centrifugal effect, and 2 an overdriven centrifugal effect.
 	 */
 	UPROPERTY(EditAnywhere, Category = "Velocity Scale Properties", meta = (UIMin = "0", UIMax = "1", ClampMin = "0", ClampMax = "2"))
