Chaos Cloth visualization updates.

* Fix cloth debug draw to display when PIE is paused. This was just a matter of switching to not draw in Foreground since the Foreground Linebatcher queue flushes every frame, even when paused.
* Add an option to draw in foreground or in the world via cvar (p.ChaosClothVisualization.DrawInForeground)
* Enable all text-based debug draw in PIE/outside of asset editors (doesn't work when paused, but it's a start).
* Update Draw Local Space to show the name of the reference bones.

diff --git a/Engine/Plugins/ChaosCloth/Source/ChaosCloth/Private/ChaosCloth/ChaosClothVisualization.cpp b/Engine/Plugins/ChaosCloth/Source/ChaosCloth/Private/ChaosCloth/ChaosClothVisualization.cpp
--- a/Engine/Plugins/ChaosCloth/Source/ChaosCloth/Private/ChaosCloth/ChaosClothVisualization.cpp
+++ b/Engine/Plugins/ChaosCloth/Source/ChaosCloth/Private/ChaosCloth/ChaosClothVisualization.cpp
@@ -91,6 +91,9 @@
 static FString WeightMapName = "";
 static FAutoConsoleVariableRef CVarClothVizWeightMapName(TEXT("p.ChaosClothVisualization.WeightMapName"), WeightMapName, TEXT("Weight map name to be visualized"));
 
+static bool bDrawInForeground = true;
+static FAutoConsoleVariableRef CVarClothVizDrawInForeground(TEXT("p.ChaosClothVisualization.DrawInForeground"), bDrawInForeground, TEXT("Draw in foreground when outside the cloth/SKM editor"));
+
 // copied from ClothEditorMode
 FLinearColor PseudoRandomColor(int32 NumColorRotations)
 {
@@ -103,6 +106,13 @@
 	}
 	return FLinearColor::MakeFromHSV8(Seed, 180, 140);
 }
+
+static uint8 GetDepthPriority()
+{
+	// set depth to SDPG_MAX not SDPG_Foreground when drawing in foreground.
+	// SDPG_Foreground does not draw when PIE is paused (its buffer is flushed).
+	return bDrawInForeground ? SDPG_MAX : SDPG_World; 
+}
 }// namespace Private
 
 	FClothVisualization::FClothVisualization(const ::Chaos::FClothingSimulationSolver* InSolver)
@@ -124,6 +134,27 @@
 		Solver = InSolver;
 	}
 
+	static void DrawText(FCanvas* Canvas, const FSceneView* SceneView, const FVector& Pos, const FText& Text, const FLinearColor& Color, const float Scale = 1.f)
+	{
 #if WITH_EDITOR
+		if (Canvas && SceneView)
+		{
+			FVector2D PixelLocation;
+			if (SceneView->WorldToPixel(Pos, PixelLocation))
+			{
+				FCanvasTextItem TextItem(PixelLocation, Text, GEngine->GetSmallFont(), Color);
+				TextItem.Scale = FVector2D::UnitVector  * Scale;
+				TextItem.EnableShadow(FLinearColor::Black);
+				TextItem.Draw(Canvas);
+			}
+		}
+		else
+#endif
+		{
+			FDebugDrawQueue::GetInstance().DrawDebugString(Pos, Text.ToString(), nullptr, Color.ToFColor(true), KINDA_SMALL_NUMBER, true, 0.75f * Scale);
+		}
+	}
+
+#if WITH_EDITOR
 	void FClothVisualization::AddReferencedObjects(FReferenceCollector& Collector)
 	{
@@ -185,18 +216,135 @@
 		MeshBuilder.Draw(PDI, LocalSimSpaceToWorld, ClothMaterial->GetRenderProxy(), SDPG_World, false, false);
 	}
 
-	static void DrawText(FCanvas* Canvas, const FSceneView* SceneView, const FVector& Pos, const FText& Text, const FLinearColor& Color)
+	void FClothVisualization::DrawWeightMapWithName(FPrimitiveDrawInterface* PDI, const FString& Name) const
 	{
-		FVector2D PixelLocation;
-		if (SceneView->WorldToPixel(Pos, PixelLocation))
+		if (!Solver || !ClothMaterialColor)
 		{
-			FCanvasTextItem TextItem(PixelLocation, Text, GEngine->GetSmallFont(), Color);
-			TextItem.Scale = FVector2D::UnitVector;
-			TextItem.EnableShadow(FLinearColor::Black);
-			TextItem.Draw(Canvas);
+			return;
 		}
+
+		FDynamicMeshBuilder MeshBuilder(PDI->View->GetFeatureLevel());
+		int32 VertexIndex = 0;
+
+		for (const FClothingSimulationCloth* const Cloth : Solver->GetCloths())
+		{
+			const int32 ParticleRangeId = Cloth->GetParticleRangeId(Solver);
+			if (ParticleRangeId == INDEX_NONE)
+			{
+				continue;
+			}
+			// Elements are local indexed for new solver
+			const int32 Offset = Solver->IsLegacySolver() ? ParticleRangeId : 0;
+			const TConstArrayView<TVec3<int32>> Elements = Cloth->GetTriangleMesh(Solver).GetElements();
+			const TConstArrayView<Softs::FSolverVec3> Positions = Cloth->GetParticlePositions(Solver);
+			const TConstArrayView<FRealSingle>& WeightMap = Cloth->GetWeightMapByName(Solver, Name);
+
+			for (int32 ElementIndex = 0; ElementIndex < Elements.Num(); ++ElementIndex, VertexIndex += 3)
+			{
+				const TVec3<int32>& Element = Elements[ElementIndex];
+
+				const FVector3f Pos0(Positions[Element.X - Offset]);
+				const FVector3f Pos1(Positions[Element.Y - Offset]);
+				const FVector3f Pos2(Positions[Element.Z - Offset]);
+
+				const FVector3f Normal = FVector3f::CrossProduct(Pos2 - Pos0, Pos1 - Pos0).GetSafeNormal();
+				const FVector3f Tangent = ((Pos1 + Pos2) * 0.5f - Pos0).GetSafeNormal();
+
+				FLinearColor VertexColor1 = FLinearColor::Black;
+				FLinearColor VertexColor2 = FLinearColor::Black;
+				FLinearColor VertexColor3 = FLinearColor::Black;
+
+				if (!WeightMap.IsEmpty() && WeightMap.Num() == Positions.Num()) // if map with that name exists and not empty
+				{
+					const FRealSingle Value0(WeightMap[Element.X - Offset]);
+					const FRealSingle Value1(WeightMap[Element.Y - Offset]);
+					const FRealSingle Value2(WeightMap[Element.Z - Offset]);
+
+					VertexColor1 = FLinearColor::LerpUsingHSV(FLinearColor::Black, FLinearColor::White, (float)Value0);
+					VertexColor2 = FLinearColor::LerpUsingHSV(FLinearColor::Black, FLinearColor::White, (float)Value1);
+					VertexColor3 = FLinearColor::LerpUsingHSV(FLinearColor::Black, FLinearColor::White, (float)Value2);
+				}
+	
+				MeshBuilder.AddVertex(FDynamicMeshVertex(Pos0, Tangent, Normal, FVector2f(0.f, 0.f), VertexColor1.ToFColor(true)));
+				MeshBuilder.AddVertex(FDynamicMeshVertex(Pos1, Tangent, Normal, FVector2f(0.f, 1.f), VertexColor2.ToFColor(true)));
+				MeshBuilder.AddVertex(FDynamicMeshVertex(Pos2, Tangent, Normal, FVector2f(1.f, 1.f), VertexColor3.ToFColor(true)));
+				MeshBuilder.AddTriangle(VertexIndex, VertexIndex + 1, VertexIndex + 2);
+			}
+		}
+		
+		FMatrix LocalSimSpaceToWorld(FMatrix::Identity);
+		LocalSimSpaceToWorld.SetOrigin(Solver->GetLocalSpaceLocation());
+		MeshBuilder.Draw(PDI, LocalSimSpaceToWorld, ClothMaterialColor->GetRenderProxy(), SDPG_World, false, false);
 	}
 
+	void FClothVisualization::DrawWeightMap(FPrimitiveDrawInterface* PDI) const
+	{
+		DrawWeightMapWithName(PDI, Private::WeightMapName);
+	}
+
+	void FClothVisualization::DrawInpaintWeightsMatched(FPrimitiveDrawInterface* PDI) const
+	{
+		DrawWeightMapWithName(PDI, TEXT("_InpaintWeightMask"));
+	}
+
+	void FClothVisualization::DrawSelfCollisionLayers(FPrimitiveDrawInterface* PDI) const
+	{
+
+		if (!Solver || !ClothMaterialColor)
+		{
+			return;
+		}
+
+		FDynamicMeshBuilder MeshBuilder(PDI->View->GetFeatureLevel());
+		int32 VertexIndex = 0;
+
+		for (const FClothingSimulationCloth* const Cloth : Solver->GetCloths())
+		{
+			const int32 ParticleRangeId = Cloth->GetParticleRangeId(Solver);
+			if (ParticleRangeId == INDEX_NONE)
+			{
+				continue;
+			}
+			// Elements are local indexed for new solver
+			const int32 Offset = Solver->IsLegacySolver() ? ParticleRangeId : 0;
+			const TConstArrayView<TVec3<int32>> Elements = Cloth->GetTriangleMesh(Solver).GetElements();
+			const TConstArrayView<Softs::FSolverVec3> Positions = Cloth->GetParticlePositions(Solver);
+			const TConstArrayView<int32>& WeightMap = Cloth->GetFaceIntMapByProperty(Solver, TEXT("SelfCollisionLayers"));
+
+			for (int32 ElementIndex = 0; ElementIndex < Elements.Num(); ++ElementIndex, VertexIndex += 3)
+			{
+
+				const TVec3<int32>& Element = Elements[ElementIndex];
+
+				const FVector3f Pos0(Positions[Element.X - Offset]);
+				const FVector3f Pos1(Positions[Element.Y - Offset]);
+				const FVector3f Pos2(Positions[Element.Z - Offset]);
+
+				const FVector3f Normal = FVector3f::CrossProduct(Pos2 - Pos0, Pos1 - Pos0).GetSafeNormal();
+				const FVector3f Tangent = ((Pos1 + Pos2) * 0.5f - Pos0).GetSafeNormal();
+
+				FLinearColor VertexColor1 = FLinearColor::Gray;
+				FLinearColor VertexColor2 = FLinearColor::Gray;
+				FLinearColor VertexColor3 = FLinearColor::Gray;
+
+				if (!WeightMap.IsEmpty() && WeightMap.Num() == Elements.Num() && WeightMap[ElementIndex]!=INDEX_NONE) // if map with that name exists and not empty
+				{
+					VertexColor1 = VertexColor2 = VertexColor3 = Chaos::Private::PseudoRandomColor(WeightMap[ElementIndex]);
+				}
+
+				MeshBuilder.AddVertex(FDynamicMeshVertex(Pos0, Tangent, Normal, FVector2f(0.f, 0.f), VertexColor1.ToFColor(true)));
+				MeshBuilder.AddVertex(FDynamicMeshVertex(Pos1, Tangent, Normal, FVector2f(0.f, 1.f), VertexColor2.ToFColor(true)));
+				MeshBuilder.AddVertex(FDynamicMeshVertex(Pos2, Tangent, Normal, FVector2f(1.f, 1.f), VertexColor3.ToFColor(true)));
+				MeshBuilder.AddTriangle(VertexIndex, VertexIndex + 1, VertexIndex + 2);
+			}
+		}
+
+		FMatrix LocalSimSpaceToWorld(FMatrix::Identity);
+		LocalSimSpaceToWorld.SetOrigin(Solver->GetLocalSpaceLocation());
+		MeshBuilder.Draw(PDI, LocalSimSpaceToWorld, ClothMaterialColor->GetRenderProxy(), SDPG_World, false, false);
+	}
+#endif  // #if WITH_EDITOR
+
 	void FClothVisualization::DrawParticleIndices(FCanvas* Canvas, const FSceneView* SceneView) const
 	{
 		if (!Solver)
@@ -327,140 +475,11 @@
 		}
 	}
 
-	void FClothVisualization::DrawWeightMapWithName(FPrimitiveDrawInterface* PDI, const FString& Name) const
-	{
-		if (!Solver || !ClothMaterialColor)
-		{
-			return;
-		}
-
-		FDynamicMeshBuilder MeshBuilder(PDI->View->GetFeatureLevel());
-		int32 VertexIndex = 0;
-
-		for (const FClothingSimulationCloth* const Cloth : Solver->GetCloths())
-		{
-			const int32 ParticleRangeId = Cloth->GetParticleRangeId(Solver);
-			if (ParticleRangeId == INDEX_NONE)
-			{
-				continue;
-			}
-			// Elements are local indexed for new solver
-			const int32 Offset = Solver->IsLegacySolver() ? ParticleRangeId : 0;
-			const TConstArrayView<TVec3<int32>> Elements = Cloth->GetTriangleMesh(Solver).GetElements();
-			const TConstArrayView<Softs::FSolverVec3> Positions = Cloth->GetParticlePositions(Solver);
-			const TConstArrayView<FRealSingle>& WeightMap = Cloth->GetWeightMapByName(Solver, Name);
-
-			for (int32 ElementIndex = 0; ElementIndex < Elements.Num(); ++ElementIndex, VertexIndex += 3)
-			{
-				const TVec3<int32>& Element = Elements[ElementIndex];
-
-				const FVector3f Pos0(Positions[Element.X - Offset]);
-				const FVector3f Pos1(Positions[Element.Y - Offset]);
-				const FVector3f Pos2(Positions[Element.Z - Offset]);
-
-				const FVector3f Normal = FVector3f::CrossProduct(Pos2 - Pos0, Pos1 - Pos0).GetSafeNormal();
-				const FVector3f Tangent = ((Pos1 + Pos2) * 0.5f - Pos0).GetSafeNormal();
-
-				FLinearColor VertexColor1 = FLinearColor::Black;
-				FLinearColor VertexColor2 = FLinearColor::Black;
-				FLinearColor VertexColor3 = FLinearColor::Black;
-
-				if (!WeightMap.IsEmpty() && WeightMap.Num() == Positions.Num()) // if map with that name exists and not empty
-				{
-					const FRealSingle Value0(WeightMap[Element.X - Offset]);
-					const FRealSingle Value1(WeightMap[Element.Y - Offset]);
-					const FRealSingle Value2(WeightMap[Element.Z - Offset]);
-
-					VertexColor1 = FLinearColor::LerpUsingHSV(FLinearColor::Black, FLinearColor::White, (float)Value0);
-					VertexColor2 = FLinearColor::LerpUsingHSV(FLinearColor::Black, FLinearColor::White, (float)Value1);
-					VertexColor3 = FLinearColor::LerpUsingHSV(FLinearColor::Black, FLinearColor::White, (float)Value2);
-				}
-	
-				MeshBuilder.AddVertex(FDynamicMeshVertex(Pos0, Tangent, Normal, FVector2f(0.f, 0.f), VertexColor1.ToFColor(true)));
-				MeshBuilder.AddVertex(FDynamicMeshVertex(Pos1, Tangent, Normal, FVector2f(0.f, 1.f), VertexColor2.ToFColor(true)));
-				MeshBuilder.AddVertex(FDynamicMeshVertex(Pos2, Tangent, Normal, FVector2f(1.f, 1.f), VertexColor3.ToFColor(true)));
-				MeshBuilder.AddTriangle(VertexIndex, VertexIndex + 1, VertexIndex + 2);
-			}
-		}
-		
-		FMatrix LocalSimSpaceToWorld(FMatrix::Identity);
-		LocalSimSpaceToWorld.SetOrigin(Solver->GetLocalSpaceLocation());
-		MeshBuilder.Draw(PDI, LocalSimSpaceToWorld, ClothMaterialColor->GetRenderProxy(), SDPG_World, false, false);
-	}
-
-	void FClothVisualization::DrawWeightMap(FPrimitiveDrawInterface* PDI) const
-	{
-		DrawWeightMapWithName(PDI, Private::WeightMapName);
-	}
-
-	void FClothVisualization::DrawInpaintWeightsMatched(FPrimitiveDrawInterface* PDI) const
-	{
-		DrawWeightMapWithName(PDI, TEXT("_InpaintWeightMask"));
-	}
-
-	void FClothVisualization::DrawSelfCollisionLayers(FPrimitiveDrawInterface* PDI) const
-	{
-
-		if (!Solver || !ClothMaterialColor)
-		{
-			return;
-		}
-
-		FDynamicMeshBuilder MeshBuilder(PDI->View->GetFeatureLevel());
-		int32 VertexIndex = 0;
-
-		for (const FClothingSimulationCloth* const Cloth : Solver->GetCloths())
-		{
-			const int32 ParticleRangeId = Cloth->GetParticleRangeId(Solver);
-			if (ParticleRangeId == INDEX_NONE)
-			{
-				continue;
-			}
-			// Elements are local indexed for new solver
-			const int32 Offset = Solver->IsLegacySolver() ? ParticleRangeId : 0;
-			const TConstArrayView<TVec3<int32>> Elements = Cloth->GetTriangleMesh(Solver).GetElements();
-			const TConstArrayView<Softs::FSolverVec3> Positions = Cloth->GetParticlePositions(Solver);
-			const TConstArrayView<int32>& WeightMap = Cloth->GetFaceIntMapByProperty(Solver, TEXT("SelfCollisionLayers"));
-
-			for (int32 ElementIndex = 0; ElementIndex < Elements.Num(); ++ElementIndex, VertexIndex += 3)
-			{
-
-				const TVec3<int32>& Element = Elements[ElementIndex];
-
-				const FVector3f Pos0(Positions[Element.X - Offset]);
-				const FVector3f Pos1(Positions[Element.Y - Offset]);
-				const FVector3f Pos2(Positions[Element.Z - Offset]);
-
-				const FVector3f Normal = FVector3f::CrossProduct(Pos2 - Pos0, Pos1 - Pos0).GetSafeNormal();
-				const FVector3f Tangent = ((Pos1 + Pos2) * 0.5f - Pos0).GetSafeNormal();
-
-				FLinearColor VertexColor1 = FLinearColor::Gray;
-				FLinearColor VertexColor2 = FLinearColor::Gray;
-				FLinearColor VertexColor3 = FLinearColor::Gray;
-
-				if (!WeightMap.IsEmpty() && WeightMap.Num() == Elements.Num() && WeightMap[ElementIndex]!=INDEX_NONE) // if map with that name exists and not empty
-				{
-					VertexColor1 = VertexColor2 = VertexColor3 = Chaos::Private::PseudoRandomColor(WeightMap[ElementIndex]);
-				}
-
-				MeshBuilder.AddVertex(FDynamicMeshVertex(Pos0, Tangent, Normal, FVector2f(0.f, 0.f), VertexColor1.ToFColor(true)));
-				MeshBuilder.AddVertex(FDynamicMeshVertex(Pos1, Tangent, Normal, FVector2f(0.f, 1.f), VertexColor2.ToFColor(true)));
-				MeshBuilder.AddVertex(FDynamicMeshVertex(Pos2, Tangent, Normal, FVector2f(1.f, 1.f), VertexColor3.ToFColor(true)));
-				MeshBuilder.AddTriangle(VertexIndex, VertexIndex + 1, VertexIndex + 2);
-			}
-		}
-
-		FMatrix LocalSimSpaceToWorld(FMatrix::Identity);
-		LocalSimSpaceToWorld.SetOrigin(Solver->GetLocalSpaceLocation());
-		MeshBuilder.Draw(PDI, LocalSimSpaceToWorld, ClothMaterialColor->GetRenderProxy(), SDPG_World, false, false);
-	}
-#endif  // #if WITH_EDITOR
-
 	static void DrawPoint(FPrimitiveDrawInterface* PDI, const FVector& Pos, const FLinearColor& Color, const UMaterial* ClothMaterialVertex, const float Thickness = 1.f)  // Use color or material
 	{
 		if (!PDI)
 		{
-			FDebugDrawQueue::GetInstance().DrawDebugPoint(Pos, Color.ToFColor(true), false, KINDA_SMALL_NUMBER, SDPG_Foreground, Thickness);
+			FDebugDrawQueue::GetInstance().DrawDebugPoint(Pos, Color.ToFColor(true), false, KINDA_SMALL_NUMBER, Chaos::Private::GetDepthPriority(), Thickness);
 			return;
 		}
 #if WITH_EDITOR
@@ -482,7 +501,7 @@
 	{
 		if (!PDI)
 		{
-			FDebugDrawQueue::GetInstance().DrawDebugLine(Pos0, Pos1, Color.ToFColor(true), false, KINDA_SMALL_NUMBER, SDPG_Foreground, 0.f);
+			FDebugDrawQueue::GetInstance().DrawDebugLine(Pos0, Pos1, Color.ToFColor(true), false, KINDA_SMALL_NUMBER, Chaos::Private::GetDepthPriority(), 0.f);
 			return;
 		}
 #if WITH_EDITOR
@@ -512,7 +531,7 @@
 		const FVec3 Center = Position + Rotation.RotateVector(Sphere.GetCenter());
 		if (!PDI)
 		{
-			FDebugDrawQueue::GetInstance().DrawDebugSphere(Center, Radius, 12, Color.ToFColor(true), false, KINDA_SMALL_NUMBER, SDPG_Foreground, 0.f);
+			FDebugDrawQueue::GetInstance().DrawDebugSphere(Center, Radius, 12, Color.ToFColor(true), false, KINDA_SMALL_NUMBER, Chaos::Private::GetDepthPriority(), 0.f);
 			return;
 		}
 #if WITH_EDITOR
@@ -526,7 +545,7 @@
 		if (!PDI)
 		{
 			const FVec3 Center = Position + Rotation.RotateVector(Box.GetCenter());
-			FDebugDrawQueue::GetInstance().DrawDebugBox(Center, Box.Extents() * 0.5f, Rotation, Color.ToFColor(true), false, KINDA_SMALL_NUMBER, SDPG_Foreground, 0.f);
+			FDebugDrawQueue::GetInstance().DrawDebugBox(Center, Box.Extents() * 0.5f, Rotation, Color.ToFColor(true), false, KINDA_SMALL_NUMBER, Chaos::Private::GetDepthPriority(), 0.f);
 			return;
 		}
 #if WITH_EDITOR
@@ -543,7 +562,7 @@
 		if (!PDI)
 		{
 			const FQuat Orientation = FQuat::FindBetweenNormals(FVec3::UpVector, Capsule.GetAxis());
-			FDebugDrawQueue::GetInstance().DrawDebugCapsule(Center, HalfHeight, Radius, Rotation * Orientation, Color.ToFColor(true), false, KINDA_SMALL_NUMBER, SDPG_Foreground, 0.f);
+			FDebugDrawQueue::GetInstance().DrawDebugCapsule(Center, HalfHeight, Radius, Rotation * Orientation, Color.ToFColor(true), false, KINDA_SMALL_NUMBER, Chaos::Private::GetDepthPriority(), 0.f);
 			return;
 		}
 #if WITH_EDITOR
@@ -661,15 +680,15 @@
 		}
 	}
 
-	static void DrawCoordinateSystem(FPrimitiveDrawInterface* PDI, const FQuat& Rotation, const FVector& Position)
+	static void DrawCoordinateSystem(FPrimitiveDrawInterface* PDI, const FQuat& Rotation, const FVector& Position, const float LengthScale = 10.f, const float ColorScale = 1.f)
 	{
-		const FVector X = Rotation.RotateVector(FVector::ForwardVector) * 10.f;
-		const FVector Y = Rotation.RotateVector(FVector::RightVector) * 10.f;
-		const FVector Z = Rotation.RotateVector(FVector::UpVector) * 10.f;
+		const FVector X = Rotation.RotateVector(FVector::ForwardVector) * LengthScale;
+		const FVector Y = Rotation.RotateVector(FVector::RightVector) * LengthScale;
+		const FVector Z = Rotation.RotateVector(FVector::UpVector) * LengthScale;
 
-		DrawLine(PDI, Position, Position + X, FLinearColor::Red);
-		DrawLine(PDI, Position, Position + Y, FLinearColor::Green);
-		DrawLine(PDI, Position, Position + Z, FLinearColor::Blue);
+		DrawLine(PDI, Position, Position + X, FLinearColor::Red * ColorScale);
+		DrawLine(PDI, Position, Position + Y, FLinearColor::Green * ColorScale);
+		DrawLine(PDI, Position, Position + Z, FLinearColor::Blue * ColorScale);
 	}
 
 	static void DrawLevelSet(FPrimitiveDrawInterface* PDI, const FTransform& Transform, const FMaterialRenderProxy* MaterialRenderProxy, const FLevelSet& LevelSet)
@@ -1115,7 +1134,7 @@
 		}
 	}
 
-	void FClothVisualization::DrawAnimNormals(FPrimitiveDrawInterface* PDI) const
+	void FClothVisualization::DrawAnimNormals(FPrimitiveDrawInterface* PDI, const FReal NormalLength) const
 	{
 		if (!Solver)
 		{
@@ -1123,7 +1142,6 @@
 		}
 
 		static const FLinearColor KinematicColor = FColor::Magenta;
-		constexpr FReal NormalLength = (FReal)20.;
 
 		const FVec3& LocalSpaceLocation = Solver->GetLocalSpaceLocation();
 
@@ -1148,7 +1166,7 @@
 		}
 	}
 
-	void FClothVisualization::DrawPointNormals(FPrimitiveDrawInterface* PDI) const
+	void FClothVisualization::DrawPointNormals(FPrimitiveDrawInterface* PDI, const FReal NormalLength) const
 	{
 		if (!Solver)
 		{
@@ -1157,7 +1175,6 @@
 
 		static const FLinearColor DynamicColor = FColor::White;
 		static const FLinearColor KinematicColor = FColor::Purple;
-		constexpr FReal NormalLength = (FReal)20.;
 
 		const FVec3& LocalSpaceLocation = Solver->GetLocalSpaceLocation();
 
@@ -2248,7 +2265,7 @@
 		}
 
 		// Draw local space
-		DrawCoordinateSystem(PDI, FQuat::Identity, Solver->GetLocalSpaceLocation());
+		DrawCoordinateSystem(PDI, FQuat::Identity, Solver->GetLocalSpaceLocation(), 5.f, 0.25f);
 
 		// Draw reference spaces
 		for (const FClothingSimulationCloth* const Cloth : Solver->GetCloths())
@@ -2260,6 +2277,8 @@
 			}
 			const FRigidTransform3& ReferenceSpaceTransform = Cloth->GetReferenceSpaceTransform();
 			DrawCoordinateSystem(PDI, ReferenceSpaceTransform.GetRotation(), ReferenceSpaceTransform.GetLocation());
+			check(Cloth->GetMesh());
+			DrawText(nullptr, nullptr, ReferenceSpaceTransform.GetLocation() + FVec3(2.), FText::FromName(Cloth->GetMesh()->GetReferenceBoneName()), FLinearColor::White, 2.f);
 		}
 	}
 
diff --git a/Engine/Plugins/ChaosCloth/Source/ChaosCloth/Private/ChaosCloth/ChaosClothingSimulation.cpp b/Engine/Plugins/ChaosCloth/Source/ChaosCloth/Private/ChaosCloth/ChaosClothingSimulation.cpp
--- a/Engine/Plugins/ChaosCloth/Source/ChaosCloth/Private/ChaosCloth/ChaosClothingSimulation.cpp
+++ b/Engine/Plugins/ChaosCloth/Source/ChaosCloth/Private/ChaosCloth/ChaosClothingSimulation.cpp
@@ -79,6 +79,7 @@ namespace ClothingSimulationCVar
 	TAutoConsoleVariable<bool> DebugDrawBackstops           (TEXT("p.ChaosCloth.DebugDrawBackstops"           ), false, TEXT("Whether to debug draw the Chaos Cloth backstops"), ECVF_Cheat);
 	TAutoConsoleVariable<bool> DebugDrawBackstopDistances   (TEXT("p.ChaosCloth.DebugDrawBackstopDistances"   ), false, TEXT("Whether to debug draw the Chaos Cloth backstop distances"), ECVF_Cheat);
 	TAutoConsoleVariable<bool> DebugDrawMaxDistances        (TEXT("p.ChaosCloth.DebugDrawMaxDistances"        ), false, TEXT("Whether to debug draw the Chaos Cloth max distances"), ECVF_Cheat);
+	TAutoConsoleVariable<bool> DebugDrawMaxDistanceValues   (TEXT("p.ChaosCloth.DebugDrawMaxDistanceValues"   ), false, TEXT("Whether to debug draw the Chaos Cloth max distances as numbers"), ECVF_Cheat);
 	TAutoConsoleVariable<bool> DebugDrawAnimDrive           (TEXT("p.ChaosCloth.DebugDrawAnimDrive"           ), false, TEXT("Whether to debug draw the Chaos Cloth anim drive"), ECVF_Cheat);
 	TAutoConsoleVariable<bool> DebugDrawEdgeConstraint      (TEXT("p.ChaosCloth.DebugDrawEdgeConstraint"      ), false, TEXT("Whether to debug draw the Chaos Cloth edge constraint"), ECVF_Cheat);
 	TAutoConsoleVariable<bool> DebugDrawBendingConstraint   (TEXT("p.ChaosCloth.DebugDrawBendingConstraint"   ), false, TEXT("Whether to debug draw the Chaos Cloth bending constraint"), ECVF_Cheat);
@@ -86,6 +87,8 @@ namespace ClothingSimulationCVar
 	TAutoConsoleVariable<bool> DebugDrawWindForces          (TEXT("p.ChaosCloth.DebugDrawWindForces"          ), false, TEXT("Whether to debug draw the Chaos Cloth wind forces"), ECVF_Cheat);
 	TAutoConsoleVariable<bool> DebugDrawSelfCollision       (TEXT("p.ChaosCloth.DebugDrawSelfCollision"       ), false, TEXT("Whether to debug draw the Chaos Cloth self collision information"), ECVF_Cheat);
 	TAutoConsoleVariable<bool> DebugDrawSelfIntersection    (TEXT("p.ChaosCloth.DebugDrawSelfIntersection"    ), false, TEXT("Whether to debug draw the Chaos Cloth self intersection information"), ECVF_Cheat);
+	TAutoConsoleVariable<bool> DebugDrawParticleIndices     (TEXT("p.ChaosCloth.DebugDrawParticleIndices"     ), false, TEXT("Whether to debug draw the Chaos Cloth particle indices"), ECVF_Cheat);
+	TAutoConsoleVariable<bool> DebugDrawElementIndices      (TEXT("p.ChaosCloth.DebugDrawElementIndices"      ), false, TEXT("Whether to debug draw the Chaos Cloth element indices"), ECVF_Cheat);
 }
 #endif  // #if CHAOS_DEBUG_DRAW
 
@@ -362,9 +365,10 @@ void FClothingSimulation::CreateActor(USkeletalMeshComponent* InOwnerComponent,
 		InOwnerComponent));
 
 	// Create collider runtime simulation object
+	const FReferenceSkeleton* const ReferenceSkeleton = &CastChecked<USkeletalMesh>(Asset->GetOuter())->GetRefSkeleton();
 	const int32 ColliderIndex = Colliders.Emplace(MakeUnique<FClothingSimulationCollider>(
 		Asset->PhysicsAsset,
-		&CastChecked<USkeletalMesh>(Asset->GetOuter())->GetRefSkeleton()));
+		ReferenceSkeleton));
 
 	// Set the external collision data to get updated at every frame
 	Colliders[ColliderIndex]->SetCollisionData(&ExternalCollisionData);
@@ -576,6 +580,7 @@ void FClothingSimulation::Simulate(IClothingSimulationContext* InContext)
 	if (ClothingSimulationCVar::DebugDrawBackstops           .GetValueOnAnyThread()) { DebugDrawBackstops           (); }
 	if (ClothingSimulationCVar::DebugDrawBackstopDistances   .GetValueOnAnyThread()) { DebugDrawBackstopDistances   (); }
 	if (ClothingSimulationCVar::DebugDrawMaxDistances        .GetValueOnAnyThread()) { DebugDrawMaxDistances        (); }
+	if (ClothingSimulationCVar::DebugDrawMaxDistanceValues   .GetValueOnAnyThread()) { DebugDrawMaxDistanceValues   (); }
 	if (ClothingSimulationCVar::DebugDrawAnimDrive           .GetValueOnAnyThread()) { DebugDrawAnimDrive           (); }
 	if (ClothingSimulationCVar::DebugDrawEdgeConstraint      .GetValueOnAnyThread()) { DebugDrawEdgeConstraint      (); }
 	if (ClothingSimulationCVar::DebugDrawBendingConstraint   .GetValueOnAnyThread()) { DebugDrawBendingConstraint   (); }
diff --git a/Engine/Plugins/ChaosCloth/Source/ChaosCloth/Private/ChaosCloth/ChaosClothingSimulationSkeletalMesh.cpp b/Engine/Plugins/ChaosCloth/Source/ChaosCloth/Private/ChaosCloth/ChaosClothingSimulationSkeletalMesh.cpp
--- a/Engine/Plugins/ChaosCloth/Source/ChaosCloth/Private/ChaosCloth/ChaosClothingSimulationSkeletalMesh.cpp
+++ b/Engine/Plugins/ChaosCloth/Source/ChaosCloth/Private/ChaosCloth/ChaosClothingSimulationSkeletalMesh.cpp
@@ -5,6 +5,7 @@
 #include "ChaosWeightMapTarget.h"
 #include "ClothingAsset.h"
 #include "ClothingSimulation.h"
+#include "Engine/SkeletalMesh.h"
 
 namespace Chaos
 {
@@ -15,6 +16,14 @@ namespace Chaos
 		, Asset(InAsset)
 		, SkeletalMeshComponent(InSkeletalMeshComponent)
 	{
+#if CHAOS_DEBUG_DRAW
+		const FReferenceSkeleton* const ReferenceSkeleton = &CastChecked<USkeletalMesh>(Asset->GetOuter())->GetRefSkeleton();
+		const int32 ReferenceBoneIndex = GetReferenceBoneIndex();
+		if (ReferenceSkeleton && ReferenceBoneIndex >= 0 && ReferenceBoneIndex < ReferenceSkeleton->GetNum())
+		{
+			ReferenceBoneName = ReferenceSkeleton->GetBoneName(ReferenceBoneIndex);
+		}
+#endif
 	}
 
 	int32 FClothingSimulationSkeletalMesh::GetNumLODs() const
diff --git a/Engine/Plugins/ChaosCloth/Source/ChaosCloth/Public/ChaosCloth/ChaosClothVisualization.h b/Engine/Plugins/ChaosCloth/Source/ChaosCloth/Public/ChaosCloth/ChaosClothVisualization.h
--- a/Engine/Plugins/ChaosCloth/Source/ChaosCloth/Public/ChaosCloth/ChaosClothVisualization.h
+++ b/Engine/Plugins/ChaosCloth/Source/ChaosCloth/Public/ChaosCloth/ChaosClothVisualization.h
@@ -31,11 +31,15 @@ namespace Chaos
 		// Editor & runtime functions
 		CHAOSCLOTH_API void SetSolver(const ::Chaos::FClothingSimulationSolver* InSolver);
 
+		CHAOSCLOTH_API void DrawParticleIndices(FCanvas* Canvas = nullptr, const FSceneView* SceneView = nullptr) const;
+		CHAOSCLOTH_API void DrawElementIndices(FCanvas* Canvas = nullptr, const FSceneView* SceneView = nullptr) const;
+		CHAOSCLOTH_API void DrawMaxDistanceValues(FCanvas* Canvas = nullptr, const FSceneView* SceneView = nullptr) const;
+
 		CHAOSCLOTH_API void DrawPhysMeshWired(FPrimitiveDrawInterface* PDI = nullptr) const;
 		CHAOSCLOTH_API void DrawAnimMeshWired(FPrimitiveDrawInterface* PDI = nullptr) const;
-		CHAOSCLOTH_API void DrawAnimNormals(FPrimitiveDrawInterface* PDI = nullptr) const;
+		CHAOSCLOTH_API void DrawAnimNormals(FPrimitiveDrawInterface* PDI = nullptr, const FReal Length = 20.) const;
 		CHAOSCLOTH_API void DrawOpenEdges(FPrimitiveDrawInterface* PDI = nullptr) const;
-		CHAOSCLOTH_API void DrawPointNormals(FPrimitiveDrawInterface* PDI = nullptr) const;
+		CHAOSCLOTH_API void DrawPointNormals(FPrimitiveDrawInterface* PDI = nullptr, const FReal Length = 20.) const;
 		CHAOSCLOTH_API void DrawPointVelocities(FPrimitiveDrawInterface* PDI = nullptr) const;
 		CHAOSCLOTH_API void DrawCollision(FPrimitiveDrawInterface* PDI = nullptr) const;
 		CHAOSCLOTH_API void DrawBackstops(FPrimitiveDrawInterface* PDI = nullptr) const;
@@ -59,11 +63,14 @@ namespace Chaos
 #else  // #if CHAOS_DEBUG_DRAW
 		void SetSolver(const ::Chaos::FClothingSimulationSolver* /*InSolver*/) {}
 
+		void DrawParticleIndices(FCanvas* /*Canvas*/, const FSceneView* /*SceneView*/) const {}
+		void DrawElementIndices(FCanvas* /*Canvas*/, const FSceneView* /*SceneView*/) const {}
+		void DrawMaxDistanceValues(FCanvas* /*Canvas*/, const FSceneView* /*SceneView*/) const {}
 		void DrawPhysMeshWired(FPrimitiveDrawInterface* /*PDI*/ = nullptr) const {}
 		void DrawAnimMeshWired(FPrimitiveDrawInterface* /*PDI*/ = nullptr) const {}
-		void DrawAnimNormals(FPrimitiveDrawInterface* /*PDI*/ = nullptr) const {}
+		void DrawAnimNormals(FPrimitiveDrawInterface* /*PDI*/ = nullptr, const FReal /*Length*/ = 20.) const {}
 		void DrawOpenEdges(FPrimitiveDrawInterface* /*PDI*/ = nullptr) const {}
-		void DrawPointNormals(FPrimitiveDrawInterface* /*PDI*/ = nullptr) const {}
+		void DrawPointNormals(FPrimitiveDrawInterface* /*PDI*/ = nullptr, const FReal /*Length*/ = 20.) const {}
 		void DrawPointVelocities(FPrimitiveDrawInterface* /*PDI*/ = nullptr) const {}
 		void DrawCollision(FPrimitiveDrawInterface* /*PDI*/ = nullptr) const {}
 		void DrawBackstops(FPrimitiveDrawInterface* /*PDI*/ = nullptr) const {}
@@ -89,9 +96,6 @@ namespace Chaos
 #if WITH_EDITOR && CHAOS_DEBUG_DRAW
 		// Editor only functions
 		CHAOSCLOTH_API void DrawPhysMeshShaded(FPrimitiveDrawInterface* PDI) const;
-		CHAOSCLOTH_API void DrawParticleIndices(FCanvas* Canvas, const FSceneView* SceneView) const;
-		CHAOSCLOTH_API void DrawElementIndices(FCanvas* Canvas, const FSceneView* SceneView) const;
-		CHAOSCLOTH_API void DrawMaxDistanceValues(FCanvas* Canvas, const FSceneView* SceneView) const;
 		CHAOSCLOTH_API void DrawWeightMap(FPrimitiveDrawInterface* PDI = nullptr) const;
 		CHAOSCLOTH_API void DrawSelfCollisionLayers(FPrimitiveDrawInterface* PDI = nullptr) const;
 		CHAOSCLOTH_API void DrawInpaintWeightsMatched(FPrimitiveDrawInterface* PDI = nullptr) const;
@@ -106,9 +110,6 @@ namespace Chaos
 
 #else  // #if WITH_EDITOR && CHAOS_DEBUG_DRAW
 		void DrawPhysMeshShaded(FPrimitiveDrawInterface* /*PDI*/) const {}
-		void DrawParticleIndices(FCanvas* /*Canvas*/, const FSceneView* /*SceneView*/) const {}
-		void DrawElementIndices(FCanvas* /*Canvas*/, const FSceneView* /*SceneView*/) const {}
-		void DrawMaxDistanceValues(FCanvas* /*Canvas*/, const FSceneView* /*SceneView*/) const {}
 		CHAOSCLOTH_API void DrawKinematicColliderShaded(FPrimitiveDrawInterface* /*PDI*/) const {}
 #endif  // #if WITH_EDITOR && CHAOS_DEBUG_DRAW
 
diff --git a/Engine/Plugins/ChaosCloth/Source/ChaosCloth/Public/ChaosCloth/ChaosClothingSimulation.h b/Engine/Plugins/ChaosCloth/Source/ChaosCloth/Public/ChaosCloth/ChaosClothingSimulation.h
--- a/Engine/Plugins/ChaosCloth/Source/ChaosCloth/Public/ChaosCloth/ChaosClothingSimulation.h
+++ b/Engine/Plugins/ChaosCloth/Source/ChaosCloth/Public/ChaosCloth/ChaosClothingSimulation.h
@@ -87,13 +87,13 @@ namespace Chaos
 #if WITH_EDITOR
 		// Editor only debug draw function
 		void DebugDrawPhysMeshShaded(FPrimitiveDrawInterface* PDI) const { Visualization.DrawPhysMeshShaded(PDI); }
-		void DebugDrawParticleIndices(FCanvas* Canvas, const FSceneView* SceneView) const { Visualization.DrawParticleIndices(Canvas, SceneView); }
-		void DebugDrawElementIndices(FCanvas* Canvas, const FSceneView* SceneView) const { Visualization.DrawElementIndices(Canvas, SceneView); }
-		void DebugDrawMaxDistanceValues(FCanvas* Canvas, const FSceneView* SceneView) const { Visualization.DrawMaxDistanceValues(Canvas, SceneView); }
 #endif  // #if WITH_EDITOR
 
 #if CHAOS_DEBUG_DRAW
 		// Editor & runtime debug draw functions
+		void DebugDrawParticleIndices(FCanvas* Canvas = nullptr, const FSceneView* SceneView = nullptr) const { Visualization.DrawParticleIndices(Canvas, SceneView); }
+		void DebugDrawElementIndices(FCanvas* Canvas = nullptr, const FSceneView* SceneView = nullptr) const { Visualization.DrawElementIndices(Canvas, SceneView); }
+		void DebugDrawMaxDistanceValues(FCanvas* Canvas = nullptr, const FSceneView* SceneView = nullptr) const { Visualization.DrawMaxDistanceValues(Canvas, SceneView); }
 		void DebugDrawPhysMeshWired(FPrimitiveDrawInterface* PDI = nullptr) const { Visualization.DrawPhysMeshWired(PDI); }
 		void DebugDrawAnimMeshWired(FPrimitiveDrawInterface* PDI = nullptr) const { Visualization.DrawAnimMeshWired(PDI); }
 		void DebugDrawAnimNormals(FPrimitiveDrawInterface* PDI = nullptr) const { Visualization.DrawAnimNormals(PDI); }
diff --git a/Engine/Plugins/ChaosCloth/Source/ChaosCloth/Public/ChaosCloth/ChaosClothingSimulationMesh.h b/Engine/Plugins/ChaosCloth/Source/ChaosCloth/Public/ChaosCloth/ChaosClothingSimulationMesh.h
--- a/Engine/Plugins/ChaosCloth/Source/ChaosCloth/Public/ChaosCloth/ChaosClothingSimulationMesh.h
+++ b/Engine/Plugins/ChaosCloth/Source/ChaosCloth/Public/ChaosCloth/ChaosClothingSimulationMesh.h
@@ -5,6 +5,7 @@
 #include "Chaos/PBDSoftsEvolutionFwd.h"
 #include "Containers/ContainersFwd.h"
 #include "ClothVertBoneData.h"
+#include "EngineDefines.h"
 
 struct FMeshToMeshVertData;
 struct FClothVertBoneData;
@@ -37,6 +38,14 @@ namespace Chaos
 #else
 		const FString& GetDebugName() const { return FText::GetEmpty().ToString(); }
 #endif
+		FName GetReferenceBoneName() const
+		{
+#if UE_ENABLE_DEBUG_DRAWING
+			return ReferenceBoneName;
+#else
+			return NAME_None;
+#endif
+		}
 
 		/* Return the number of LODs on this mesh. */
 		virtual int32 GetNumLODs() const = 0;
@@ -163,6 +172,10 @@ namespace Chaos
 #if !UE_BUILD_SHIPPING
 		/** Debug name of the source component. */
 		FString DebugName;
+#endif
+	protected:
+#if UE_ENABLE_DEBUG_DRAWING
+		FName ReferenceBoneName = NAME_None;
 #endif
 	};
 } // namespace Chaos
diff --git a/Engine/Plugins/ChaosClothAsset/Source/ChaosClothAssetEngine/Private/ChaosClothAsset/ClothSimulationMesh.cpp b/Engine/Plugins/ChaosClothAsset/Source/ChaosClothAssetEngine/Private/ChaosClothAsset/ClothSimulationMesh.cpp
--- a/Engine/Plugins/ChaosClothAsset/Source/ChaosClothAssetEngine/Private/ChaosClothAsset/ClothSimulationMesh.cpp
+++ b/Engine/Plugins/ChaosClothAsset/Source/ChaosClothAssetEngine/Private/ChaosClothAsset/ClothSimulationMesh.cpp
@@ -5,6 +5,7 @@
 #include "ChaosClothAsset/ClothSimulationModel.h"
 #include "ChaosClothAsset/ClothSimulationContext.h"
 #include "ClothSimulationMesh.h"
+#include "ReferenceSkeleton.h"
 
 namespace UE::Chaos::ClothAsset
 {
@@ -13,6 +14,14 @@ namespace UE::Chaos::ClothAsset
 		, ClothSimulationModel(InClothSimulationModel)
 		, ClothSimulationContext(InClothSimulationContext)
 	{
+#if CHAOS_DEBUG_DRAW
+		const int32 ReferenceBoneIndex = GetReferenceBoneIndex();
+		const int32 UsedBoneNameIndex = ClothSimulationModel.UsedBoneIndices.Find(ReferenceBoneIndex);
+		if (ClothSimulationModel.UsedBoneNames.IsValidIndex(UsedBoneNameIndex))
+		{
+			ReferenceBoneName = ClothSimulationModel.UsedBoneNames[UsedBoneNameIndex];
+		}
+#endif
 	}
 
 	int32 FClothSimulationMesh::GetNumLODs() const
diff --git a/Engine/Plugins/ChaosClothAsset/Source/ChaosClothAssetEngine/Private/ChaosClothAsset/ClothSimulationProxy.cpp b/Engine/Plugins/ChaosClothAsset/Source/ChaosClothAssetEngine/Private/ChaosClothAsset/ClothSimulationProxy.cpp
--- a/Engine/Plugins/ChaosClothAsset/Source/ChaosClothAssetEngine/Private/ChaosClothAsset/ClothSimulationProxy.cpp
+++ b/Engine/Plugins/ChaosClothAsset/Source/ChaosClothAssetEngine/Private/ChaosClothAsset/ClothSimulationProxy.cpp
@@ -370,6 +370,7 @@ namespace UE::Chaos::ClothAsset
 		static const TConsoleVariableData<bool>* const DebugDrawBackstopsCVar = IConsoleManager::Get().FindTConsoleVariableDataBool(TEXT("p.ChaosCloth.DebugDrawBackstops"));
 		static const TConsoleVariableData<bool>* const DebugDrawBackstopDistancesCVar = IConsoleManager::Get().FindTConsoleVariableDataBool(TEXT("p.ChaosCloth.DebugDrawBackstopDistances"));
 		static const TConsoleVariableData<bool>* const DebugDrawMaxDistancesCVar = IConsoleManager::Get().FindTConsoleVariableDataBool(TEXT("p.ChaosCloth.DebugDrawMaxDistances"));
+		static const TConsoleVariableData<bool>* const DebugDrawMaxDistanceValuesCVar = IConsoleManager::Get().FindTConsoleVariableDataBool(TEXT("p.ChaosCloth.DebugDrawMaxDistanceValues"));
 		static const TConsoleVariableData<bool>* const DebugDrawAnimDriveCVar = IConsoleManager::Get().FindTConsoleVariableDataBool(TEXT("p.ChaosCloth.DebugDrawAnimDrive"));
 		static const TConsoleVariableData<bool>* const DebugDrawEdgeConstraintCVar = IConsoleManager::Get().FindTConsoleVariableDataBool(TEXT("p.ChaosCloth.DebugDrawEdgeConstraint"));
 		static const TConsoleVariableData<bool>* const DebugDrawBendingConstraintCVar = IConsoleManager::Get().FindTConsoleVariableDataBool(TEXT("p.ChaosCloth.DebugDrawBendingConstraint"));
@@ -377,6 +378,8 @@ namespace UE::Chaos::ClothAsset
 		static const TConsoleVariableData<bool>* const DebugDrawWindForcesCVar = IConsoleManager::Get().FindTConsoleVariableDataBool(TEXT("p.ChaosCloth.DebugDrawWindForces"));
 		static const TConsoleVariableData<bool>* const DebugDrawSelfCollisionCVar = IConsoleManager::Get().FindTConsoleVariableDataBool(TEXT("p.ChaosCloth.DebugDrawSelfCollision"));
 		static const TConsoleVariableData<bool>* const DebugDrawSelfIntersectionCVar = IConsoleManager::Get().FindTConsoleVariableDataBool(TEXT("p.ChaosCloth.DebugDrawSelfIntersection"));
+		static const TConsoleVariableData<bool>* const DebugDrawParticleIndicesCVar = IConsoleManager::Get().FindTConsoleVariableDataBool(TEXT("p.ChaosCloth.DebugDrawParticleIndices"));
+		static const TConsoleVariableData<bool>* const DebugDrawElementIndicesCVar = IConsoleManager::Get().FindTConsoleVariableDataBool(TEXT("p.ChaosCloth.DebugDrawElementIndices"));
 
 		if (DebugDrawLocalSpaceCVar && DebugDrawLocalSpaceCVar->GetValueOnAnyThread()) { Visualization->DrawLocalSpace(); }
 		if (DebugDrawBoundsCVar && DebugDrawBoundsCVar->GetValueOnAnyThread()) { Visualization->DrawBounds(); }
@@ -388,8 +391,9 @@ namespace UE::Chaos::ClothAsset
 		if (DebugDrawPointNormalsCVar && DebugDrawPointNormalsCVar->GetValueOnAnyThread()) { Visualization->DrawPointNormals(); }
 		if (DebugDrawCollisionCVar && DebugDrawCollisionCVar->GetValueOnAnyThread()) { Visualization->DrawCollision(); }
 		if (DebugDrawBackstopsCVar && DebugDrawBackstopsCVar->GetValueOnAnyThread()) { Visualization->DrawBackstops(); }
-		if (DebugDrawBackstopDistancesCVar && DebugDrawBackstopDistancesCVar->GetValueOnAnyThread()) { Visualization->DrawMaxDistances(); }
+		if (DebugDrawBackstopDistancesCVar && DebugDrawBackstopDistancesCVar->GetValueOnAnyThread()) { Visualization->DrawBackstopDistances(); }
 		if (DebugDrawMaxDistancesCVar && DebugDrawMaxDistancesCVar->GetValueOnAnyThread()) { Visualization->DrawMaxDistances(); }
+		if (DebugDrawMaxDistanceValuesCVar && DebugDrawMaxDistanceValuesCVar->GetValueOnAnyThread()) { Visualization->DrawMaxDistanceValues(); }
 		if (DebugDrawAnimDriveCVar && DebugDrawAnimDriveCVar->GetValueOnAnyThread()) { Visualization->DrawAnimDrive(); }
 		if (DebugDrawEdgeConstraintCVar && DebugDrawEdgeConstraintCVar->GetValueOnAnyThread()) { Visualization->DrawEdgeConstraint(); }
 		if (DebugDrawBendingConstraintCVar && DebugDrawBendingConstraintCVar->GetValueOnAnyThread()) { Visualization->DrawBendingConstraint(); }
@@ -397,6 +401,8 @@ namespace UE::Chaos::ClothAsset
 		if (DebugDrawWindForcesCVar && DebugDrawWindForcesCVar->GetValueOnAnyThread()) { Visualization->DrawWindAndPressureForces(); }
 		if (DebugDrawSelfCollisionCVar && DebugDrawSelfCollisionCVar->GetValueOnAnyThread()) { Visualization->DrawSelfCollision(); }
 		if (DebugDrawSelfIntersectionCVar && DebugDrawSelfIntersectionCVar->GetValueOnAnyThread()) { Visualization->DrawSelfIntersection(); }
+		if (DebugDrawParticleIndicesCVar && DebugDrawParticleIndicesCVar->GetValueOnAnyThread()) { Visualization->DrawParticleIndices(); }
+		if (DebugDrawElementIndicesCVar && DebugDrawElementIndicesCVar->GetValueOnAnyThread()) { Visualization->DrawElementIndices(); }
 #endif  // #if CHAOS_DEBUG_DRAW
 	}
 
diff --git a/Engine/Plugins/ChaosClothAssetEditor/Source/ChaosClothAssetEditor/Private/ChaosClothAsset/ClothEditorSimulationVisualization.cpp b/Engine/Plugins/ChaosClothAssetEditor/Source/ChaosClothAssetEditor/Private/ChaosClothAsset/ClothEditorSimulationVisualization.cpp
--- a/Engine/Plugins/ChaosClothAssetEditor/Source/ChaosClothAssetEditor/Private/ChaosClothAsset/ClothEditorSimulationVisualization.cpp
+++ b/Engine/Plugins/ChaosClothAssetEditor/Source/ChaosClothAssetEditor/Private/ChaosClothAsset/ClothEditorSimulationVisualization.cpp
@@ -10,6 +10,7 @@
 #include "Rendering/SkeletalMeshRenderData.h"
 #include "Delegates/Delegate.h"
 #include "Delegates/DelegateCombinations.h"
+#include "Widgets/Input/SSpinBox.h"
 #include "Widgets/Input/STextComboBox.h"
 #include "Types/SlateEnums.h"
 
@@ -195,12 +196,17 @@ const FVisualizationOption FVisualizationOption::OptionData[] =
 		LOCTEXT("ChaosVisName_ElementIndices", "Element Indices"), 
 		LOCTEXT("ChaosVisName_ElementIndices_ToolTip", "Draws the element's (triangle or other) indices as instantiated by the solver")),
 	FVisualizationOption(
-		FClothVisualizationDebugDraw::CreateLambda([](const FClothEditorSimulationVisualization&, const ::Chaos::FClothVisualization& Visualization, FPrimitiveDrawInterface* PDI)
+		FClothVisualizationDebugDraw::CreateLambda([](const FClothEditorSimulationVisualization& EditorVisualization, const ::Chaos::FClothVisualization& Visualization, FPrimitiveDrawInterface* PDI)
 		{
-			Visualization.DrawPointNormals(PDI);
+			Visualization.DrawPointNormals(PDI, EditorVisualization.GetPointNormalLength());
 		}),
 		LOCTEXT("ChaosVisName_PointNormals", "Physical Mesh Normals"), 
-		LOCTEXT("ChaosVisName_PointNormals_ToolTip", "Draws the current point normals for the simulation mesh")),
+		LOCTEXT("ChaosVisName_PointNormals_ToolTip", "Draws the current point normals for the simulation mesh"),
+		/*bDisablesSimulation =*/false, /*bHidesClothSections=*/false,
+		FAdditionalMenus::CreateLambda([](FClothEditorSimulationVisualization& EditorVisualization, FMenuBuilder& MenuBuilder, TSharedRef<FChaosClothAssetEditor3DViewportClient> ViewportClient)
+		{
+			EditorVisualization.ExtendViewportShowMenuPointNormalsLength(MenuBuilder);
+		})),
 	FVisualizationOption(
 		FClothVisualizationDebugDraw::CreateLambda([](const FClothEditorSimulationVisualization&, const ::Chaos::FClothVisualization& Visualization, FPrimitiveDrawInterface* PDI)
 		{
@@ -209,12 +215,17 @@ const FVisualizationOption FVisualizationOption::OptionData[] =
 		LOCTEXT("ChaosVisName_PointVelocities", "Point Velocities"), 
 		LOCTEXT("ChaosVisName_PointVelocities_ToolTip", "Draws the current point velocities for the simulation mesh")),
 	FVisualizationOption(
-		FClothVisualizationDebugDraw::CreateLambda([](const FClothEditorSimulationVisualization&, const ::Chaos::FClothVisualization& Visualization, FPrimitiveDrawInterface* PDI)
+		FClothVisualizationDebugDraw::CreateLambda([](const FClothEditorSimulationVisualization& EditorVisualization, const ::Chaos::FClothVisualization& Visualization, FPrimitiveDrawInterface* PDI)
 		{
-			Visualization.DrawAnimNormals(PDI);
+			Visualization.DrawAnimNormals(PDI, EditorVisualization.GetAnimatedNormalLength());
 		}),
 		LOCTEXT("ChaosVisName_AnimNormals", "Animated Mesh Normals"), 
-		LOCTEXT("ChaosVisName_AnimNormals_ToolTip", "Draws the current point normals for the animated mesh")),
+		LOCTEXT("ChaosVisName_AnimNormals_ToolTip", "Draws the current point normals for the animated mesh"),
+		/*bDisablesSimulation =*/false, /*bHidesClothSections=*/false,
+		FAdditionalMenus::CreateLambda([](FClothEditorSimulationVisualization& EditorVisualization, FMenuBuilder& MenuBuilder, TSharedRef<FChaosClothAssetEditor3DViewportClient> ViewportClient)
+		{
+			EditorVisualization.ExtendViewportShowMenuAnimatedNormalsLength(MenuBuilder);
+		})),
 	FVisualizationOption(
 		FClothVisualizationDebugDraw::CreateLambda([](const FClothEditorSimulationVisualization&, const ::Chaos::FClothVisualization& Visualization, FPrimitiveDrawInterface* PDI)
 		{
@@ -457,6 +468,21 @@ void FClothEditorSimulationVisualization::WeightMapSelectionChanged(TSharedPtr<F
 	CurrentlySelectedWeightMap = Selection;
 }
 
+void FClothEditorSimulationVisualization::ExtendViewportShowMenuSpinBox(FMenuBuilder& MenuBuilder, float& Value, const float MinValue, const float MaxValue, const float MinSliderValue, const float MaxSliderValue)
+{
+	TSharedRef<SWidget> SpinBox =
+		SNew(SSpinBox<float>)
+		.Value(Value)
+		.OnValueChanged_Lambda([&Value](float NewValue) { Value = NewValue; })
+		.MinValue(MinValue)
+		.MaxValue(MaxValue)
+		.MinSliderValue(MinSliderValue)
+		.MaxSliderValue(MaxSliderValue);
+
+	MenuBuilder.AddMenuEntry(FUIAction(), SpinBox,
+		NAME_None, LOCTEXT("SetNormalLength", "Set the normal display length."), EUserInterfaceActionType::None);
+}
+
 void FClothEditorSimulationVisualization::RefreshMenusForClothComponent(const UChaosClothComponent* ClothComponent)
 {
 	TArray<FString> WeightMaps;
diff --git a/Engine/Plugins/ChaosClothAssetEditor/Source/ChaosClothAssetEditor/Private/ChaosClothAsset/ClothEditorSimulationVisualization.h b/Engine/Plugins/ChaosClothAssetEditor/Source/ChaosClothAssetEditor/Private/ChaosClothAsset/ClothEditorSimulationVisualization.h
--- a/Engine/Plugins/ChaosClothAssetEditor/Source/ChaosClothAssetEditor/Private/ChaosClothAsset/ClothEditorSimulationVisualization.h
+++ b/Engine/Plugins/ChaosClothAssetEditor/Source/ChaosClothAssetEditor/Private/ChaosClothAsset/ClothEditorSimulationVisualization.h
@@ -35,6 +35,17 @@ public:
 	// WeightMaps 
 	const FString* GetCurrentlySelectedWeightMap() const { return CurrentlySelectedWeightMap.Get(); }
 	void ExtendViewportShowMenuWeightMapSelector(FMenuBuilder& MenuBuilder, TSharedRef<FChaosClothAssetEditor3DViewportClient> ViewportClient);
+	// Normals
+	void ExtendViewportShowMenuPointNormalsLength(FMenuBuilder& MenuBuilder)
+	{
+		ExtendViewportShowMenuSpinBox(MenuBuilder, PointNormalLength, 0.f, FLT_MAX, 0.f, 40.f);
+	}
+	void ExtendViewportShowMenuAnimatedNormalsLength(FMenuBuilder& MenuBuilder)
+	{
+		ExtendViewportShowMenuSpinBox(MenuBuilder, AnimatedNormalLength, 0.f, FLT_MAX, 0.f, 40.f);
+	}
+	float GetPointNormalLength() const { return PointNormalLength; }
+	float GetAnimatedNormalLength() const { return AnimatedNormalLength; }
 private:
 	/** Return whether or not - given the current enabled options - the simulation should be disabled. */
 	bool ShouldDisableSimulation() const;
@@ -42,13 +53,15 @@ private:
 	void ShowClothSections(UChaosClothComponent* ClothComponent, bool bIsClothSectionsVisible) const;
 	/** Callback for weight map selection*/
 	void WeightMapSelectionChanged(TSharedPtr<FString> Selection, ESelectInfo::Type SelectInfo);
-
+	void ExtendViewportShowMenuSpinBox(FMenuBuilder& MenuBuilder, float& Value, const float MinValue, const float MaxValue, const float MinSliderValue, const float MaxSliderValue);
 private:
 	/** Flags used to store the checked status for the visualization options. */
 	TBitArray<> Flags;
 	TSharedPtr<STextComboBox> WeightMapSelector;
 	TArray<TSharedPtr<FString>> WeightMapNames;
 	TSharedPtr<FString> CurrentlySelectedWeightMap;
+	float PointNormalLength = 20.f;
+	float AnimatedNormalLength = 20.f;
 	
 };
 } // namespace UE::Chaos::ClothAsset
