Chaos Cloth: fix regression where AngularRatio was no longer being applied to velocity scale.

diff --git a/Engine/Plugins/ChaosCloth/Source/ChaosCloth/Private/ChaosCloth/ChaosClothingSimulationSolver.cpp b/Engine/Plugins/ChaosCloth/Source/ChaosCloth/Private/ChaosCloth/ChaosClothingSimulationSolver.cpp
--- a/Engine/Plugins/ChaosCloth/Source/ChaosCloth/Private/ChaosCloth/ChaosClothingSimulationSolver.cpp
+++ b/Engine/Plugins/ChaosCloth/Source/ChaosCloth/Private/ChaosCloth/ChaosClothingSimulationSolver.cpp
@@ -1150,7 +1150,7 @@ void FClothingSimulationSolver::SetReferenceVelocityScale(uint32 GroupId,
 		InOutReferenceAngularVelocity = FullDeltaRotation / DeltaTime;
 	
 		FVec3 AppliedLinearVelocity = AppliedDeltaTranslation / DeltaTime;
-		FVec3 AppliedAngularVelocity = InOutReferenceAngularVelocity;
+		FVec3 AppliedAngularVelocity = AppliedDeltaRotation / DeltaTime;
 
 		const FVec3 OldScaledReferenceSpaceVelocity = LinearRatio * (OldReferenceSpaceRotationInverse * OldReferenceVelocity);
 		const FVec3 ScaledReferenceSpaceAccelerationTimesDt = AppliedLinearVelocity - OldScaledReferenceSpaceVelocity;
