Legacy Chaos Cloth: add FlatnessRatio to the config.

diff --git a/Engine/Plugins/ChaosCloth/Source/ChaosCloth/Private/ChaosCloth/ChaosClothingSimulationConfig.cpp b/Engine/Plugins/ChaosCloth/Source/ChaosCloth/Private/ChaosCloth/ChaosClothingSimulationConfig.cpp
--- a/Engine/Plugins/ChaosCloth/Source/ChaosCloth/Private/ChaosCloth/ChaosClothingSimulationConfig.cpp
+++ b/Engine/Plugins/ChaosCloth/Source/ChaosCloth/Private/ChaosCloth/ChaosClothingSimulationConfig.cpp
@@ -3,6 +3,7 @@
 #include "ChaosCloth/ChaosClothConfig.h"
 #include "Chaos/CollectionPropertyFacade.h"
 #include "Chaos/PBDLongRangeConstraints.h"  // For Tether modes
+#include "Chaos/PBDBendingConstraintsBase.h" // For ERestAngleConstructionType
 #include "GeometryCollection/ManagedArrayCollection.h"
 
 namespace Chaos
@@ -91,6 +92,13 @@ namespace Chaos
 					const int32 BucklingStiffnessIndex = Property->AddProperty(TEXT("BucklingStiffness"), AnimatablePropertyFlags);
 					Property->SetWeightedValue(BucklingStiffnessIndex, ClothConfig->BucklingStiffnessWeighted.Low, ClothConfig->BucklingStiffnessWeighted.High);
 					Property->SetStringValue(BucklingStiffnessIndex, TEXT("BucklingStiffness"));
+
+					const int32 RestAngleTypeIndex = Property->AddProperty(TEXT("RestAngleType"), NonAnimatablePropertyFlags);
+					Property->SetValue(RestAngleTypeIndex, (int32)Chaos::Softs::FPBDBendingConstraintsBase::ERestAngleConstructionType::FlatnessRatio);
+
+					const int32 FlatnessRatioIndex = Property->AddProperty(TEXT("FlatnessRatio"), NonAnimatablePropertyFlags);
+					Property->SetWeightedValue(FlatnessRatioIndex, ClothConfig->FlatnessRatio.Low, ClothConfig->FlatnessRatio.High);
+					Property->SetStringValue(FlatnessRatioIndex, TEXT("FlatnessRatio"));
 				}
 				else  // Not using bending elements
 				{
diff --git a/Engine/Plugins/ChaosCloth/Source/ChaosCloth/Private/ChaosCloth/ChaosWeightMapTarget.h b/Engine/Plugins/ChaosCloth/Source/ChaosCloth/Private/ChaosCloth/ChaosWeightMapTarget.h
--- a/Engine/Plugins/ChaosCloth/Source/ChaosCloth/Private/ChaosCloth/ChaosWeightMapTarget.h
+++ b/Engine/Plugins/ChaosCloth/Source/ChaosCloth/Private/ChaosCloth/ChaosWeightMapTarget.h
@@ -25,4 +25,5 @@ enum class EChaosWeightMapTarget : uint8
 	AreaStiffness,
 	BucklingStiffness,
 	Pressure,
+	FlatnessRatio
 };
diff --git a/Engine/Plugins/ChaosCloth/Source/ChaosCloth/Public/ChaosCloth/ChaosClothConfig.h b/Engine/Plugins/ChaosCloth/Source/ChaosCloth/Public/ChaosCloth/ChaosClothConfig.h
--- a/Engine/Plugins/ChaosCloth/Source/ChaosCloth/Public/ChaosCloth/ChaosClothConfig.h
+++ b/Engine/Plugins/ChaosCloth/Source/ChaosCloth/Public/ChaosCloth/ChaosClothConfig.h
@@ -160,6 +160,17 @@ public:
 	UPROPERTY(EditAnywhere, Category = "Material Properties", DisplayName = "Buckling Stiffness", meta = (UIMin = "0", UIMax = "1", ClampMin = "0", ClampMax = "1", EditCondition = "bUseBendingElements"))
 	FChaosClothWeightedValue BucklingStiffnessWeighted = { 1.f, 1.f };
 
+	/**
+	* Calculate rest angles as a ratio between completely flat and whatever is the 3D rest angle.
+	* When FlatnessRatio = 0, the rest angle will match the input mesh..
+	* When FlatnessRatio = 1, the rest angle will be 0 (completely flat).
+	* If an enabled Weight Map (Mask with values in the range [0;1]) targeting the "Flatness Ratio" is added to the cloth, 
+	* then both the Low and High values will be used in conjunction with the per particle Weight stored in the Weight Map to interpolate the final value from them.
+	* Otherwise only the Low value is meaningful and sufficient to enable this constraint.
+	*/
+	UPROPERTY(EditAnywhere, Category = "Material Properties", meta = (UIMin = "0", UIMax = "1", ClampMin = "0", ClampMax = "1", EditCondition = "bUseBendingElements"))
+	FChaosClothWeightedValue FlatnessRatio = { 0.f, 1.f };
+
 	/**
 	 * The stiffness of the surface area preservation constraints. Increase the iteration count for stiffer materials.
 	 * If an enabled Weight Map (Mask with values in the range [0;1]) targeting the "Bend Stiffness" is added to the cloth, 
