Add missing files from last CL submitted (35541346)

diff --git a/Engine/Plugins/ChaosCloth/Source/ChaosCloth/Private/ChaosCloth/ChaosClothingSimulationSolver.cpp b/Engine/Plugins/ChaosCloth/Source/ChaosCloth/Private/ChaosCloth/ChaosClothingSimulationSolver.cpp
--- a/Engine/Plugins/ChaosCloth/Source/ChaosCloth/Private/ChaosCloth/ChaosClothingSimulationSolver.cpp
+++ b/Engine/Plugins/ChaosCloth/Source/ChaosCloth/Private/ChaosCloth/ChaosClothingSimulationSolver.cpp
@@ -1423,13 +1423,15 @@ void FClothingSimulationSolver::SetWindAndPressureGeometry(
 	uint32 GroupId,
 	const FTriangleMesh& TriangleMesh,
 	const TConstArrayView<FRealSingle>& DragMultipliers,
+	const TConstArrayView<FRealSingle>& OuterDragMultipliers,
 	const TConstArrayView<FRealSingle>& LiftMultipliers,
+	const TConstArrayView<FRealSingle>& OuterLiftMultipliers,
 	const TConstArrayView<FRealSingle>& PressureMultipliers)
 {
 	if (PBDEvolution)
 	{
 		Softs::FVelocityAndPressureField& VelocityAndPressureField = PBDEvolution->GetVelocityAndPressureField(GroupId);
-		VelocityAndPressureField.SetGeometry(&TriangleMesh, DragMultipliers, LiftMultipliers, PressureMultipliers);
+		VelocityAndPressureField.SetGeometry(&TriangleMesh, DragMultipliers, OuterDragMultipliers, LiftMultipliers, OuterLiftMultipliers, PressureMultipliers);
 	}
 }
 
@@ -1449,7 +1451,9 @@ void FClothingSimulationSolver::SetWindAndPressureProperties(
 void FClothingSimulationSolver::SetWindAndPressureProperties(
 	uint32 GroupId,
 	const TVec2<FRealSingle>& Drag,
+	const TVec2<FRealSingle>& OuterDrag,
 	const TVec2<FRealSingle>& Lift,
+	const TVec2<FRealSingle>& OuterLift,
 	FRealSingle FluidDensity,
 	const TVec2<FRealSingle>& Pressure)
 {
@@ -1458,7 +1462,9 @@ void FClothingSimulationSolver::SetWindAndPressureProperties(
 		Softs::FVelocityAndPressureField& VelocityAndPressureField = PBDEvolution->GetVelocityAndPressureField(GroupId);
 		VelocityAndPressureField.SetProperties(
 			Drag,
+			OuterDrag,
 			Lift,
+			OuterLift,
 			FluidDensity / FMath::Cube(ClothingSimulationSolverConstant::WorldScale),  // Fluid density is given in kg/m^3. Need to convert to kg/cm^3 for solver.
 			Pressure / ClothingSimulationSolverConstant::WorldScale);  // UI Pressure is in kg/m s^2. Need to convert to kg/cm s^2 for solver.
 	}
diff --git a/Engine/Plugins/ChaosCloth/Source/ChaosCloth/Public/ChaosCloth/ChaosClothingSimulationSolver.h b/Engine/Plugins/ChaosCloth/Source/ChaosCloth/Public/ChaosCloth/ChaosClothingSimulationSolver.h
--- a/Engine/Plugins/ChaosCloth/Source/ChaosCloth/Public/ChaosCloth/ChaosClothingSimulationSolver.h
+++ b/Engine/Plugins/ChaosCloth/Source/ChaosCloth/Public/ChaosCloth/ChaosClothingSimulationSolver.h
@@ -196,9 +196,22 @@ namespace Chaos
 			uint32 GroupId,
 			const FTriangleMesh& TriangleMesh,
 			const TConstArrayView<FRealSingle>& DragMultipliers,
+			const TConstArrayView<FRealSingle>& OuterDragMultipliers,
 			const TConstArrayView<FRealSingle>& LiftMultipliers,
+			const TConstArrayView<FRealSingle>& OuterLiftMultipliers,
 			const TConstArrayView<FRealSingle>& PressureMultipliers);
 
+		UE_DEPRECATED(5.5, "Use version with PropertyCollection (preferred) or OuterDrag and OuterLift")
+		void SetWindAndPressureGeometry(
+			uint32 GroupId,
+			const FTriangleMesh& TriangleMesh,
+			const TConstArrayView<FRealSingle>& DragMultipliers,
+			const TConstArrayView<FRealSingle>& LiftMultipliers,
+			const TConstArrayView<FRealSingle>& PressureMultipliers)
+		{
+			SetWindAndPressureGeometry(GroupId, TriangleMesh, DragMultipliers, DragMultipliers, LiftMultipliers, LiftMultipliers, PressureMultipliers);
+		}
+
 		// Set the wind and pressure properties.
 		CHAOSCLOTH_API void SetWindAndPressureProperties(
 			uint32 GroupId,
@@ -207,11 +220,24 @@ namespace Chaos
 			bool bEnableAerodynamics);
 
 		CHAOSCLOTH_API void SetWindAndPressureProperties(
+			uint32 GroupId,
+			const TVec2<FRealSingle>& Drag,
+			const TVec2<FRealSingle>& OuterDrag,
+			const TVec2<FRealSingle>& Lift,
+			const TVec2<FRealSingle>& OuterLift,
+			FRealSingle FluidDensity,
+			const TVec2<FRealSingle>& Pressure);
+
+		UE_DEPRECATED(5.5, "Use version with PropertyCollection (preferred) or OuterDrag and OuterLift")
+		void SetWindAndPressureProperties(
 			uint32 GroupId,
 			const TVec2<FRealSingle>& Drag,
 			const TVec2<FRealSingle>& Lift,
 			FRealSingle FluidDensity = 1.225f,
-			const TVec2<FRealSingle>& Pressure = TVec2<FRealSingle>::ZeroVector);
+			const TVec2<FRealSingle>& Pressure = TVec2<FRealSingle>::ZeroVector)
+		{
+			SetWindAndPressureProperties(GroupId, Drag, Drag, Lift, Lift, FluidDensity, Pressure);
+		}
 
 		// Return the wind velocity and pressure field associated with a given group id.
 		CHAOSCLOTH_API const Softs::FVelocityAndPressureField& GetWindVelocityAndPressureField(uint32 GroupId) const;
