Chaos Cloth - Moved visualization material loads to a singleton instantiation to prevent unecessary hangs when launching PIE.

[RN] Chaos Cloth - Moved visualization material loads to a singleton instantiation to prevent unecessary hangs when launching PIE.

diff --git a/Engine/Plugins/ChaosCloth/Source/ChaosCloth/Private/ChaosCloth/ChaosClothVisualization.cpp b/Engine/Plugins/ChaosCloth/Source/ChaosCloth/Private/ChaosCloth/ChaosClothVisualization.cpp
--- a/Engine/Plugins/ChaosCloth/Source/ChaosCloth/Private/ChaosCloth/ChaosClothVisualization.cpp
+++ b/Engine/Plugins/ChaosCloth/Source/ChaosCloth/Private/ChaosCloth/ChaosClothVisualization.cpp
@@ -35,101 +35,150 @@
 #include "Chaos/WeightedLatticeImplicitObject.h"
 #include "DynamicMeshBuilder.h"
 #include "Engine/EngineTypes.h"
+#include "Misc/LazySingleton.h"
 #include "SceneManagement.h"
 #include "SceneView.h"
 #if WITH_EDITOR
-#include "Materials/Material.h"
-#include "Engine/Canvas.h"  // For draw text
 #include "CanvasItem.h"     //
+#include "Engine/Canvas.h"  // For draw text
 #include "Engine/Engine.h"  //
+#include "Materials/Material.h"
 #include "UObject/ICookInfo.h"
 #endif  // #if WITH_EDITOR
 
 namespace Chaos
 {
+	namespace Private
+	{
+		static int DrawSkinnedLattice = 0;
+		static FAutoConsoleVariableRef CVarClothVizDrawSkinnedLattice(TEXT("p.ChaosClothVisualization.DrawSkinnedLattice"), DrawSkinnedLattice, TEXT("Draw skinned lattice, 0 = none, 1 = filled, 2 = empty, 3 = both"));
 
-namespace Private
-{
-static int DrawSkinnedLattice = 0;
-static FAutoConsoleVariableRef CVarClothVizDrawSkinnedLattice(TEXT("p.ChaosClothVisualization.DrawSkinnedLattice"), DrawSkinnedLattice, TEXT("Draw skinned lattice, 0 = none, 1 = filled, 2 = empty, 3 = both"));
+		// TODO: move these options to be somewhere the new cloth editor visualization can use.
+		enum class EBendingDrawMode : int
+		{
+			BuckleStatus = 0,
+			ParallelGraphColor = 1,
+			Anisotropy = 2,
+			RestAngle = 3
+		};
+		static int32 BendingDrawMode = (int32)EBendingDrawMode::BuckleStatus;
+		static FAutoConsoleVariableRef CVarClothVizBendDrawMode(TEXT("p.ChaosClothVisualization.BendingDrawMode"), BendingDrawMode, TEXT("Bending draw mode, 0 = BuckleStatus, 1 = Parallel graph color, 2 = Anisotropy, 3 = RestAngle"));
 
-// TODO: move these options to be somewhere the new cloth editor visualization can use.
-enum class EBendingDrawMode : int
-{
-	BuckleStatus = 0,
-	ParallelGraphColor = 1,
-	Anisotropy = 2,
-	RestAngle = 3
-};
-static int32 BendingDrawMode = (int32)EBendingDrawMode::BuckleStatus;
-static FAutoConsoleVariableRef CVarClothVizBendDrawMode(TEXT("p.ChaosClothVisualization.BendingDrawMode"), BendingDrawMode, TEXT("Bending draw mode, 0 = BuckleStatus, 1 = Parallel graph color, 2 = Anisotropy, 3 = RestAngle"));
-
-enum class EStretchBiasDrawMode : int
-{
-	ParallelGraphColor = 0,
-	WarpStretch = 1,
-	WeftStretch = 2,
-	BiasStretch = 3
-};
-static int32 StretchBiasDrawMode = (int32)EStretchBiasDrawMode::ParallelGraphColor;
-static FAutoConsoleVariableRef CVarClothVizStretchBiasDrawMode(TEXT("p.ChaosClothVisualization.StretchBiasDrawMode"), StretchBiasDrawMode, TEXT("Stretch draw mode, 0 = Parallel graph color, 1 = Warp Stretch, 2 = Weft Stretch, 3 = BiasStretch"));
-static float StretchBiasDrawRangeMin = -1.f;
-static float StretchBiasDrawRangeMax = 1.f;
-static FAutoConsoleVariableRef CVarClothVizStretchBiasDrawRangeMin(TEXT("p.ChaosClothVisualization.StretchBiasDrawRangeMin"), StretchBiasDrawRangeMin, TEXT("Min stretch in draw color range. Negative = compressed, 0 = undeformed, positive = stretched. (When drawing warp/weft stretch)"));
-static FAutoConsoleVariableRef CVarClothVizStretchBiasDrawRangeMax(TEXT("p.ChaosClothVisualization.StretchBiasDrawRangeMax"), StretchBiasDrawRangeMax, TEXT("Max stretch in draw color range. Negative = compressed, 0 = undeformed, positive = stretched. (When drawing warp/weft stretch)"));
-static bool bStretchBiasDrawOutOfRange = true;
-static FAutoConsoleVariableRef CVarClothVizStretchBiasDrawOutOfRange(TEXT("p.ChaosClothVisualization.StretchBiasDrawOutOfRange"), bStretchBiasDrawOutOfRange, TEXT("Draw out of range elements (When drawing warp/weft stretch)"));
-
-enum class EAnisoSpringDrawMode : int
-{
-	ParallelGraphColor = 0,
-	Anisotropy = 1,
-};
-static int32 AnisoSpringDrawMode = (int32)EAnisoSpringDrawMode::ParallelGraphColor;
-static FAutoConsoleVariableRef CVarClothVizAnisoSpringDrawMode(TEXT("p.ChaosClothVisualization.AnisoSpringDrawMode"), AnisoSpringDrawMode, TEXT("Stretch draw mode, 0 = Parallel graph color, 1 = Anisotropy"));
+		enum class EStretchBiasDrawMode : int
+		{
+			ParallelGraphColor = 0,
+			WarpStretch = 1,
+			WeftStretch = 2,
+			BiasStretch = 3
+		};
+		static int32 StretchBiasDrawMode = (int32)EStretchBiasDrawMode::ParallelGraphColor;
+		static FAutoConsoleVariableRef CVarClothVizStretchBiasDrawMode(TEXT("p.ChaosClothVisualization.StretchBiasDrawMode"), StretchBiasDrawMode, TEXT("Stretch draw mode, 0 = Parallel graph color, 1 = Warp Stretch, 2 = Weft Stretch, 3 = BiasStretch"));
+		static float StretchBiasDrawRangeMin = -1.f;
+		static float StretchBiasDrawRangeMax = 1.f;
+		static FAutoConsoleVariableRef CVarClothVizStretchBiasDrawRangeMin(TEXT("p.ChaosClothVisualization.StretchBiasDrawRangeMin"), StretchBiasDrawRangeMin, TEXT("Min stretch in draw color range. Negative = compressed, 0 = undeformed, positive = stretched. (When drawing warp/weft stretch)"));
+		static FAutoConsoleVariableRef CVarClothVizStretchBiasDrawRangeMax(TEXT("p.ChaosClothVisualization.StretchBiasDrawRangeMax"), StretchBiasDrawRangeMax, TEXT("Max stretch in draw color range. Negative = compressed, 0 = undeformed, positive = stretched. (When drawing warp/weft stretch)"));
+		static bool bStretchBiasDrawOutOfRange = true;
+		static FAutoConsoleVariableRef CVarClothVizStretchBiasDrawOutOfRange(TEXT("p.ChaosClothVisualization.StretchBiasDrawOutOfRange"), bStretchBiasDrawOutOfRange, TEXT("Draw out of range elements (When drawing warp/weft stretch)"));
 
-static FString WeightMapName = "";
-static FAutoConsoleVariableRef CVarClothVizWeightMapName(TEXT("p.ChaosClothVisualization.WeightMapName"), WeightMapName, TEXT("Weight map name to be visualized"));
+		enum class EAnisoSpringDrawMode : int
+		{
+			ParallelGraphColor = 0,
+			Anisotropy = 1,
+		};
+		static int32 AnisoSpringDrawMode = (int32)EAnisoSpringDrawMode::ParallelGraphColor;
+		static FAutoConsoleVariableRef CVarClothVizAnisoSpringDrawMode(TEXT("p.ChaosClothVisualization.AnisoSpringDrawMode"), AnisoSpringDrawMode, TEXT("Stretch draw mode, 0 = Parallel graph color, 1 = Anisotropy"));
 
-static bool bDrawInForeground = true;
-static FAutoConsoleVariableRef CVarClothVizDrawInForeground(TEXT("p.ChaosClothVisualization.DrawInForeground"), bDrawInForeground, TEXT("Draw in foreground when outside the cloth/SKM editor"));
+		static FString WeightMapName = "";
+		static FAutoConsoleVariableRef CVarClothVizWeightMapName(TEXT("p.ChaosClothVisualization.WeightMapName"), WeightMapName, TEXT("Weight map name to be visualized"));
 
-// copied from ClothEditorMode
-FLinearColor PseudoRandomColor(int32 NumColorRotations)
-{
-	constexpr uint8 Spread = 157;  // Prime number that gives a good spread of colors without getting too similar as a rand might do.
-	uint8 Seed = Spread;
-	NumColorRotations = FMath::Abs(NumColorRotations);
-	for (int32 Rotation = 0; Rotation < NumColorRotations; ++Rotation)
-	{
-		Seed += Spread;
-	}
-	return FLinearColor::MakeFromHSV8(Seed, 180, 140);
-}
+		static bool bDrawInForeground = true;
+		static FAutoConsoleVariableRef CVarClothVizDrawInForeground(TEXT("p.ChaosClothVisualization.DrawInForeground"), bDrawInForeground, TEXT("Draw in foreground when outside the cloth/SKM editor"));
 
-static uint8 GetDepthPriority()
-{
-	// set depth to SDPG_MAX not SDPG_Foreground when drawing in foreground.
-	// SDPG_Foreground does not draw when PIE is paused (its buffer is flushed).
-	return bDrawInForeground ? SDPG_MAX : SDPG_World; 
-}
-}// namespace Private
+		// copied from ClothEditorMode
+		FLinearColor PseudoRandomColor(int32 NumColorRotations)
+		{
+			constexpr uint8 Spread = 157;  // Prime number that gives a good spread of colors without getting too similar as a rand might do.
+			uint8 Seed = Spread;
+			NumColorRotations = FMath::Abs(NumColorRotations);
+			for (int32 Rotation = 0; Rotation < NumColorRotations; ++Rotation)
+			{
+				Seed += Spread;
+			}
+			return FLinearColor::MakeFromHSV8(Seed, 180, 140);
+		}
 
-	FClothVisualization::FClothVisualization(const ::Chaos::FClothingSimulationSolver* InSolver)
+		static uint8 GetDepthPriority()
+		{
+			// set depth to SDPG_MAX not SDPG_Foreground when drawing in foreground.
+			// SDPG_Foreground does not draw when PIE is paused (its buffer is flushed).
+			return bDrawInForeground ? SDPG_MAX : SDPG_World; 
+		}
+	}  // namespace Private
+
+	// Delay loading materials to their first point of use to avoid causing sync flush issues and slow loading when starting PIE
+	class FClothVisualizationNoGC::FMaterials
+#if WITH_EDITOR && CHAOS_DEBUG_DRAW
+		: public FGCObject
+	{
+	private:
+		FMaterials()
+			: ClothMaterial(LoadObject<UMaterial>(nullptr, TEXT("/Engine/EditorMaterials/Cloth/CameraLitDoubleSided.CameraLitDoubleSided"), nullptr, LOAD_None, nullptr))
+			, ClothMaterialColor(LoadObject<UMaterial>(nullptr, TEXT("/Engine/EditorMaterials/Cloth/CameraLitVertexColor.CameraLitVertexColor"), nullptr, LOAD_None, nullptr))
+			, ClothMaterialVertex(LoadObject<UMaterial>(nullptr, TEXT("/Engine/EditorMaterials/WidgetVertexColorMaterial"), nullptr, LOAD_None, nullptr))
+			, CollisionMaterial(LoadObject<UMaterial>(nullptr, TEXT("/Engine/EditorMaterials/PhAT_UnselectedMaterial"), nullptr, LOAD_None, nullptr))
+		{}
+
+		~FMaterials() = default;
+
+		//~ Begin FGCObject interface
+		virtual void AddReferencedObjects(FReferenceCollector& Collector) override
+		{
+			Collector.AddReferencedObject(ClothMaterial);
+			Collector.AddReferencedObject(ClothMaterialColor);
+			Collector.AddReferencedObject(ClothMaterialVertex);
+			Collector.AddReferencedObject(CollisionMaterial);
+		}
+
+		virtual FString GetReferencerName() const override { return TEXT("Chaos::Private::FClothVisualizationMaterials"); }
+		//~ End FGCObject interface
+
+		// Visualization material
+		TObjectPtr<const UMaterial> ClothMaterial;
+		TObjectPtr<const UMaterial> ClothMaterialColor;
+		TObjectPtr<const UMaterial> ClothMaterialVertex;
+		TObjectPtr<const UMaterial> CollisionMaterial;
+
+	public:
+		TObjectPtr<const UMaterial> GetClothMaterial() const { return ClothMaterial; }
+		TObjectPtr<const UMaterial> GetClothMaterialColor() const { return ClothMaterialColor; }
+		TObjectPtr<const UMaterial> GetClothMaterialVertex() const { return ClothMaterialVertex; }
+		TObjectPtr<const UMaterial> GetCollisionMaterial() const { return CollisionMaterial; }
+#else
+	{
+	private:
+		FMaterials() = default;
+		~FMaterials() = default;
+
+	public:
+		TObjectPtr<const UMaterial> GetClothMaterial() const { return nullptr; }
+		TObjectPtr<const UMaterial> GetClothMaterialColor() const { return nullptr; }
+		TObjectPtr<const UMaterial> GetClothMaterialVertex() const { return nullptr; }
+		TObjectPtr<const UMaterial> GetCollisionMaterial() const { return nullptr; }
+#endif  // #if WITH_EDITOR && CHAOS_DEBUG_DRAW
+
+		friend class ::FLazySingleton;
+		static FMaterials& GetInstance() { return TLazySingleton<FMaterials>::Get(); }  // Don't move the implementation outside of the cpp file
+		static void TearDown() { TLazySingleton<FMaterials>::TearDown(); }  // Don't move the implementation outside of the cpp file
+	};
+
+	FClothVisualizationNoGC::FClothVisualizationNoGC(const ::Chaos::FClothingSimulationSolver* InSolver)
 		: Solver(InSolver)
 	{
-#if WITH_EDITOR
-		FCookLoadScope CookLoadScope(ECookLoadType::EditorOnly);
-		ClothMaterial = LoadObject<UMaterial>(nullptr, TEXT("/Engine/EditorMaterials/Cloth/CameraLitDoubleSided.CameraLitDoubleSided"), nullptr, LOAD_None, nullptr);  // LOAD_EditorOnly
-		ClothMaterialColor = LoadObject<UMaterial>(nullptr, TEXT("/Engine/EditorMaterials/Cloth/CameraLitVertexColor.CameraLitVertexColor"), nullptr, LOAD_None, nullptr);  // LOAD_EditorOnly
-		ClothMaterialVertex = LoadObject<UMaterial>(nullptr, TEXT("/Engine/EditorMaterials/WidgetVertexColorMaterial"), nullptr, LOAD_None, nullptr);  // LOAD_EditorOnly
-		CollisionMaterial = LoadObject<UMaterial>(nullptr, TEXT("/Engine/EditorMaterials/PhAT_UnselectedMaterial"), nullptr, LOAD_None, nullptr);
-#endif  // #if WITH_EDITOR
 	}
 
-	FClothVisualization::~FClothVisualization() = default;
+	FClothVisualizationNoGC::~FClothVisualizationNoGC() = default;
 
-	void FClothVisualization::SetSolver(const ::Chaos::FClothingSimulationSolver* InSolver)
+	void FClothVisualizationNoGC::SetSolver(const ::Chaos::FClothingSimulationSolver* InSolver)
 	{
 		Solver = InSolver;
 	}
@@ -156,17 +205,9 @@ static uint8 GetDepthPriority()
 	}
 
 #if WITH_EDITOR
-	void FClothVisualization::AddReferencedObjects(FReferenceCollector& Collector)
+	void FClothVisualizationNoGC::DrawPhysMeshShaded(FPrimitiveDrawInterface* PDI) const
 	{
-		Collector.AddReferencedObject(ClothMaterial);
-		Collector.AddReferencedObject(ClothMaterialColor);
-		Collector.AddReferencedObject(ClothMaterialVertex);
-		Collector.AddReferencedObject(CollisionMaterial);
-	}
-
-	void FClothVisualization::DrawPhysMeshShaded(FPrimitiveDrawInterface* PDI) const
-	{
-		if (!Solver || !ClothMaterial)
+		if (!Solver || !FMaterials::GetInstance().GetClothMaterial())
 		{
 			return;
 		}
@@ -213,12 +254,12 @@ static uint8 GetDepthPriority()
 
 		FMatrix LocalSimSpaceToWorld(FMatrix::Identity);
 		LocalSimSpaceToWorld.SetOrigin(Solver->GetLocalSpaceLocation());
-		MeshBuilder.Draw(PDI, LocalSimSpaceToWorld, ClothMaterial->GetRenderProxy(), SDPG_World, false, false);
+		MeshBuilder.Draw(PDI, LocalSimSpaceToWorld, FMaterials::GetInstance().GetClothMaterial()->GetRenderProxy(), SDPG_World, false, false);
 	}
 
-	void FClothVisualization::DrawWeightMapWithName(FPrimitiveDrawInterface* PDI, const FString& Name) const
+	void FClothVisualizationNoGC::DrawWeightMapWithName(FPrimitiveDrawInterface* PDI, const FString& Name) const
 	{
-		if (!Solver || !ClothMaterialColor)
+		if (!Solver || !FMaterials::GetInstance().GetClothMaterialColor())
 		{
 			return;
 		}
@@ -274,23 +315,23 @@ static uint8 GetDepthPriority()
 		
 		FMatrix LocalSimSpaceToWorld(FMatrix::Identity);
 		LocalSimSpaceToWorld.SetOrigin(Solver->GetLocalSpaceLocation());
-		MeshBuilder.Draw(PDI, LocalSimSpaceToWorld, ClothMaterialColor->GetRenderProxy(), SDPG_World, false, false);
+		MeshBuilder.Draw(PDI, LocalSimSpaceToWorld, FMaterials::GetInstance().GetClothMaterialColor()->GetRenderProxy(), SDPG_World, false, false);
 	}
 
-	void FClothVisualization::DrawWeightMap(FPrimitiveDrawInterface* PDI) const
+	void FClothVisualizationNoGC::DrawWeightMap(FPrimitiveDrawInterface* PDI) const
 	{
 		DrawWeightMapWithName(PDI, Private::WeightMapName);
 	}
 
-	void FClothVisualization::DrawInpaintWeightsMatched(FPrimitiveDrawInterface* PDI) const
+	void FClothVisualizationNoGC::DrawInpaintWeightsMatched(FPrimitiveDrawInterface* PDI) const
 	{
 		DrawWeightMapWithName(PDI, TEXT("_InpaintWeightMask"));
 	}
 
-	void FClothVisualization::DrawSelfCollisionLayers(FPrimitiveDrawInterface* PDI) const
+	void FClothVisualizationNoGC::DrawSelfCollisionLayers(FPrimitiveDrawInterface* PDI) const
 	{
 
-		if (!Solver || !ClothMaterialColor)
+		if (!Solver || !FMaterials::GetInstance().GetClothMaterialColor())
 		{
 			return;
 		}
@@ -341,11 +382,11 @@ static uint8 GetDepthPriority()
 
 		FMatrix LocalSimSpaceToWorld(FMatrix::Identity);
 		LocalSimSpaceToWorld.SetOrigin(Solver->GetLocalSpaceLocation());
-		MeshBuilder.Draw(PDI, LocalSimSpaceToWorld, ClothMaterialColor->GetRenderProxy(), SDPG_World, false, false);
+		MeshBuilder.Draw(PDI, LocalSimSpaceToWorld, FMaterials::GetInstance().GetClothMaterialColor()->GetRenderProxy(), SDPG_World, false, false);
 	}
 #endif  // #if WITH_EDITOR
 
-	void FClothVisualization::DrawParticleIndices(FCanvas* Canvas, const FSceneView* SceneView) const
+	void FClothVisualizationNoGC::DrawParticleIndices(FCanvas* Canvas, const FSceneView* SceneView) const
 	{
 		if (!Solver)
 		{
@@ -380,7 +421,7 @@ static uint8 GetDepthPriority()
 		}
 	}
 
-	void FClothVisualization::DrawElementIndices(FCanvas* Canvas, const FSceneView* SceneView) const
+	void FClothVisualizationNoGC::DrawElementIndices(FCanvas* Canvas, const FSceneView* SceneView) const
 	{
 		if (!Solver)
 		{
@@ -425,7 +466,7 @@ static uint8 GetDepthPriority()
 		}
 	}
 
-	void FClothVisualization::DrawMaxDistanceValues(FCanvas* Canvas, const FSceneView* SceneView) const
+	void FClothVisualizationNoGC::DrawMaxDistanceValues(FCanvas* Canvas, const FSceneView* SceneView) const
 	{
 		if (!Solver)
 		{
@@ -798,7 +839,7 @@ static uint8 GetDepthPriority()
 		}
 	}
 
-	void FClothVisualization::DrawBounds(FPrimitiveDrawInterface* PDI) const
+	void FClothVisualizationNoGC::DrawBounds(FPrimitiveDrawInterface* PDI) const
 	{
 		if (!Solver)
 		{
@@ -827,7 +868,7 @@ static uint8 GetDepthPriority()
 		}
 	}
 
-	void FClothVisualization::DrawGravity(FPrimitiveDrawInterface* PDI) const
+	void FClothVisualizationNoGC::DrawGravity(FPrimitiveDrawInterface* PDI) const
 	{
 		if (!Solver)
 		{
@@ -879,7 +920,7 @@ static uint8 GetDepthPriority()
 	}
 
 
-	void FClothVisualization::DrawFictitiousAngularForces(FPrimitiveDrawInterface* PDI) const
+	void FClothVisualizationNoGC::DrawFictitiousAngularForces(FPrimitiveDrawInterface* PDI) const
 	{
 		if (!Solver)
 		{
@@ -916,7 +957,7 @@ static uint8 GetDepthPriority()
 		}
 	}
 
-	void FClothVisualization::DrawPhysMeshWired(FPrimitiveDrawInterface* PDI) const
+	void FClothVisualizationNoGC::DrawPhysMeshWired(FPrimitiveDrawInterface* PDI) const
 	{
 		if (!Solver)
 		{
@@ -961,7 +1002,7 @@ static uint8 GetDepthPriority()
 		}
 	}
 
-	void FClothVisualization::DrawAnimMeshWired(FPrimitiveDrawInterface* PDI) const
+	void FClothVisualizationNoGC::DrawAnimMeshWired(FPrimitiveDrawInterface* PDI) const
 	{
 		if (!Solver)
 		{
@@ -1000,7 +1041,7 @@ static uint8 GetDepthPriority()
 		}
 	}
 
-	void FClothVisualization::DrawOpenEdges(FPrimitiveDrawInterface* PDI) const
+	void FClothVisualizationNoGC::DrawOpenEdges(FPrimitiveDrawInterface* PDI) const
 	{
 		auto MakeSortedUintVector2 = [](uint32 Index0, uint32 Index1) -> FUintVector2
 			{
@@ -1068,7 +1109,7 @@ static uint8 GetDepthPriority()
 		}
 	}
 
-	void FClothVisualization::DrawMultiResConstraint(FPrimitiveDrawInterface* PDI) const
+	void FClothVisualizationNoGC::DrawMultiResConstraint(FPrimitiveDrawInterface* PDI) const
 	{
 		if (!Solver)
 		{
@@ -1134,7 +1175,7 @@ static uint8 GetDepthPriority()
 		}
 	}
 
-	void FClothVisualization::DrawAnimNormals(FPrimitiveDrawInterface* PDI, const FReal NormalLength) const
+	void FClothVisualizationNoGC::DrawAnimNormals(FPrimitiveDrawInterface* PDI, const FReal NormalLength) const
 	{
 		if (!Solver)
 		{
@@ -1166,7 +1207,7 @@ static uint8 GetDepthPriority()
 		}
 	}
 
-	void FClothVisualization::DrawPointNormals(FPrimitiveDrawInterface* PDI, const FReal NormalLength) const
+	void FClothVisualizationNoGC::DrawPointNormals(FPrimitiveDrawInterface* PDI, const FReal NormalLength) const
 	{
 		if (!Solver)
 		{
@@ -1202,7 +1243,7 @@ static uint8 GetDepthPriority()
 		}
 	}
 
-	void FClothVisualization::DrawPointVelocities(FPrimitiveDrawInterface* PDI) const
+	void FClothVisualizationNoGC::DrawPointVelocities(FPrimitiveDrawInterface* PDI) const
 	{
 		if (!Solver)
 		{
@@ -1238,7 +1279,7 @@ static uint8 GetDepthPriority()
 		}
 	}
 
-	void FClothVisualization::DrawCollision(FPrimitiveDrawInterface* PDI) const
+	void FClothVisualizationNoGC::DrawCollision(FPrimitiveDrawInterface* PDI) const
 	{
 		if (!Solver)
 		{
@@ -1335,7 +1376,7 @@ static uint8 GetDepthPriority()
 								const FLevelSet& LevelSet = Object->GetObjectChecked<TImplicitObjectTransformed<FReal, 3>>().GetGeometry()->GetObjectChecked<FLevelSet>();
 								const FMaterialRenderProxy* MaterialRenderProxy =
 #if WITH_EDITOR
-									CollisionMaterial->GetRenderProxy();
+									FMaterials::GetInstance().GetCollisionMaterial() ? FMaterials::GetInstance().GetCollisionMaterial()->GetRenderProxy() : nullptr;
 #else
 									nullptr;
 #endif
@@ -1347,7 +1388,7 @@ static uint8 GetDepthPriority()
 								const TWeightedLatticeImplicitObject<FLevelSet>& WeightedLevelset = Object->GetObjectChecked< TWeightedLatticeImplicitObject<FLevelSet> >();
 								const FMaterialRenderProxy* MaterialRenderProxy =
 #if WITH_EDITOR
-									CollisionMaterial->GetRenderProxy();
+									FMaterials::GetInstance().GetCollisionMaterial() ? FMaterials::GetInstance().GetCollisionMaterial()->GetRenderProxy() : nullptr;
 #else
 									nullptr;
 #endif
@@ -1408,7 +1449,7 @@ static uint8 GetDepthPriority()
 		}
 	}
 
-	void FClothVisualization::DrawBackstops(FPrimitiveDrawInterface* PDI) const
+	void FClothVisualizationNoGC::DrawBackstops(FPrimitiveDrawInterface* PDI) const
 	{
 		if (!Solver)
 		{
@@ -1477,7 +1518,7 @@ static uint8 GetDepthPriority()
 		}
 	}
 
-	void FClothVisualization::DrawBackstopDistances(FPrimitiveDrawInterface* PDI) const
+	void FClothVisualizationNoGC::DrawBackstopDistances(FPrimitiveDrawInterface* PDI) const
 	{
 		if (!Solver)
 		{
@@ -1523,7 +1564,7 @@ static uint8 GetDepthPriority()
 		}
 	}
 
-	void FClothVisualization::DrawMaxDistances(FPrimitiveDrawInterface* PDI) const
+	void FClothVisualizationNoGC::DrawMaxDistances(FPrimitiveDrawInterface* PDI) const
 	{
 		if (!Solver)
 		{
@@ -1561,7 +1602,7 @@ static uint8 GetDepthPriority()
 				if (InvMasses[Index] == (Softs::FSolverReal)0.)
 				{
 #if WITH_EDITOR
-					DrawPoint(PDI, Position, FLinearColor::Red, ClothMaterialVertex);
+					DrawPoint(PDI, Position, FLinearColor::Red, FMaterials::GetInstance().GetClothMaterialVertex());
 #else
 					DrawPoint(nullptr, Position, FLinearColor::Red, nullptr);
 #endif
@@ -1574,7 +1615,7 @@ static uint8 GetDepthPriority()
 		}
 	}
 
-	void FClothVisualization::DrawAnimDrive(FPrimitiveDrawInterface* PDI) const
+	void FClothVisualizationNoGC::DrawAnimDrive(FPrimitiveDrawInterface* PDI) const
 	{
 		if (!Solver)
 		{
@@ -1911,7 +1952,7 @@ static uint8 GetDepthPriority()
 		});
 	}
 
-	void FClothVisualization::DrawEdgeConstraint(FPrimitiveDrawInterface* PDI) const
+	void FClothVisualizationNoGC::DrawEdgeConstraint(FPrimitiveDrawInterface* PDI) const
 	{
 		if (!Solver)
 		{
@@ -1963,7 +2004,7 @@ static uint8 GetDepthPriority()
 			{
 				const FMaterialRenderProxy* MaterialRenderProxy =
 #if WITH_EDITOR
-					ClothMaterialColor->GetRenderProxy();
+					FMaterials::GetInstance().GetClothMaterialColor()->GetRenderProxy();
 #else
 					nullptr;
 #endif
@@ -2023,7 +2064,7 @@ static uint8 GetDepthPriority()
 		}
 	}
 
-	void FClothVisualization::DrawBendingConstraint(FPrimitiveDrawInterface* PDI) const
+	void FClothVisualizationNoGC::DrawBendingConstraint(FPrimitiveDrawInterface* PDI) const
 	{
 		if (!Solver)
 		{
@@ -2112,7 +2153,7 @@ static uint8 GetDepthPriority()
 		}
 	}
 
-	void FClothVisualization::DrawLongRangeConstraint(FPrimitiveDrawInterface* PDI) const
+	void FClothVisualizationNoGC::DrawLongRangeConstraint(FPrimitiveDrawInterface* PDI) const
 	{
 		if (!Solver)
 		{
@@ -2177,7 +2218,7 @@ static uint8 GetDepthPriority()
 
 						DrawLine(PDI, Pos0, Pos1, Color);
 #if WITH_EDITOR
-						DrawPoint(PDI, Pos1, Color, ClothMaterialVertex);
+						DrawPoint(PDI, Pos1, Color, FMaterials::GetInstance().GetClothMaterialVertex());
 #else
 						DrawPoint(nullptr, Pos1, Color, nullptr);
 #endif
@@ -2197,7 +2238,7 @@ static uint8 GetDepthPriority()
 		}
 	}
 
-	void FClothVisualization::DrawWindAndPressureForces(FPrimitiveDrawInterface* PDI) const
+	void FClothVisualizationNoGC::DrawWindAndPressureForces(FPrimitiveDrawInterface* PDI) const
 	{
 		if (!Solver)
 		{
@@ -2257,7 +2298,7 @@ static uint8 GetDepthPriority()
 		}
 	}
 
-	void FClothVisualization::DrawLocalSpace(FPrimitiveDrawInterface* PDI) const
+	void FClothVisualizationNoGC::DrawLocalSpace(FPrimitiveDrawInterface* PDI) const
 	{
 		if (!Solver)
 		{
@@ -2282,7 +2323,7 @@ static uint8 GetDepthPriority()
 		}
 	}
 
-	void FClothVisualization::DrawSelfCollision(FPrimitiveDrawInterface* PDI) const
+	void FClothVisualizationNoGC::DrawSelfCollision(FPrimitiveDrawInterface* PDI) const
 	{
 		if (!Solver)
 		{
@@ -2402,7 +2443,7 @@ static uint8 GetDepthPriority()
 		}
 	}
 
-	void FClothVisualization::DrawSelfIntersection(FPrimitiveDrawInterface* PDI) const
+	void FClothVisualizationNoGC::DrawSelfIntersection(FPrimitiveDrawInterface* PDI) const
 	{
 		if (!Solver)
 		{
@@ -2544,7 +2585,7 @@ static uint8 GetDepthPriority()
 		}
 	}
 
-	void FClothVisualization::DrawSelfCollisionThickness(FPrimitiveDrawInterface* PDI) const
+	void FClothVisualizationNoGC::DrawSelfCollisionThickness(FPrimitiveDrawInterface* PDI) const
 	{
 		if (!Solver)
 		{
@@ -2554,7 +2595,7 @@ static uint8 GetDepthPriority()
 		const FVec3& LocalSpaceLocation = Solver->GetLocalSpaceLocation();
 
 #if WITH_EDITOR
-		if (ClothMaterialColor)
+		if (FMaterials::GetInstance().GetClothMaterialColor())
 		{
 			FDynamicMeshBuilder MeshBuilder(PDI->View->GetFeatureLevel());
 
@@ -2603,7 +2644,7 @@ static uint8 GetDepthPriority()
 			}
 			FMatrix LocalSimSpaceToWorld(FMatrix::Identity);
 			LocalSimSpaceToWorld.SetOrigin(Solver->GetLocalSpaceLocation());
-			MeshBuilder.Draw(PDI, LocalSimSpaceToWorld, ClothMaterialColor->GetRenderProxy(), SDPG_World, false, false);
+			MeshBuilder.Draw(PDI, LocalSimSpaceToWorld, FMaterials::GetInstance().GetClothMaterialColor()->GetRenderProxy(), SDPG_World, false, false);
 		}
 		else
 #endif
@@ -2660,7 +2701,7 @@ static uint8 GetDepthPriority()
 		}
 	}
 
-	void FClothVisualization::DrawKinematicColliderWired(FPrimitiveDrawInterface* PDI) const
+	void FClothVisualizationNoGC::DrawKinematicColliderWired(FPrimitiveDrawInterface* PDI) const
 	{
 		if (!Solver)
 		{
@@ -2703,7 +2744,7 @@ static uint8 GetDepthPriority()
 	}
 
 #if WITH_EDITOR
-	TArray<FString> FClothVisualization::GetAllWeightMapNames() const
+	TArray<FString> FClothVisualizationNoGC::GetAllWeightMapNames() const
 	{
 		TSet<FString> AllNames;
 		for (const FClothingSimulationCloth* const Cloth : Solver->GetCloths())
@@ -2713,9 +2754,9 @@ static uint8 GetDepthPriority()
 		return AllNames.Array();
 	}
 
-	void FClothVisualization::DrawKinematicColliderShaded(FPrimitiveDrawInterface* PDI) const
+	void FClothVisualizationNoGC::DrawKinematicColliderShaded(FPrimitiveDrawInterface* PDI) const
 	{
-		if (!Solver || !CollisionMaterial)
+		if (!Solver || !FMaterials::GetInstance().GetCollisionMaterial())
 		{
 			return;
 		}
@@ -2760,15 +2801,21 @@ static uint8 GetDepthPriority()
 		const FVec3& LocalSpaceLocation = Solver->GetLocalSpaceLocation();
 		FMatrix LocalSimSpaceToWorld(FMatrix::Identity);
 		LocalSimSpaceToWorld.SetOrigin(Solver->GetLocalSpaceLocation());
-		MeshBuilder.Draw(PDI, LocalSimSpaceToWorld, CollisionMaterial->GetRenderProxy(), SDPG_World, false, false);
+		MeshBuilder.Draw(PDI, LocalSimSpaceToWorld, FMaterials::GetInstance().GetCollisionMaterial()->GetRenderProxy(), SDPG_World, false, false);
 	}
 #endif // #if WITH_EDITOR
-}  // End namespace Chaos
 #else  // #if CHAOS_DEBUG_DRAW
 namespace Chaos
 {
-	FClothVisualization::FClothVisualization(const ::Chaos::FClothingSimulationSolver* /*InSolver*/) {}
+	FClothVisualizationNoGC::FClothVisualizationNoGC(const ::Chaos::FClothingSimulationSolver* /*InSolver*/)
+	{}
+
+	FClothVisualizationNoGC::~FClothVisualizationNoGC() = default;
+#endif  // #else #if CHAOS_DEBUG_DRAW
+
+	FClothVisualization::FClothVisualization(const ::Chaos::FClothingSimulationSolver* InSolver)
+		: FClothVisualizationNoGC(InSolver)
+	{}
 
 	FClothVisualization::~FClothVisualization() = default;
-}  // End namespace UE::Chaos::Cloth
-#endif  // #if CHAOS_DEBUG_DRAW
+}  // End namespace Chaos
diff --git a/Engine/Plugins/ChaosCloth/Source/ChaosCloth/Private/ChaosCloth/ChaosClothingSimulation.cpp b/Engine/Plugins/ChaosCloth/Source/ChaosCloth/Private/ChaosCloth/ChaosClothingSimulation.cpp
--- a/Engine/Plugins/ChaosCloth/Source/ChaosCloth/Private/ChaosCloth/ChaosClothingSimulation.cpp
+++ b/Engine/Plugins/ChaosCloth/Source/ChaosCloth/Private/ChaosCloth/ChaosClothingSimulation.cpp
@@ -262,6 +262,7 @@ namespace ChaosClothingSimulationDefault
 	static const FReal MaxDistancesMultipliers = (FReal)1.;
 }
 
+PRAGMA_DISABLE_DEPRECATION_WARNINGS  // For Visualization
 FClothingSimulation::FClothingSimulation()
 	: ClothSharedSimConfig(nullptr)
 	, bUseLocalSpaceSimulation(true)
@@ -284,6 +285,7 @@ FClothingSimulation::FClothingSimulation()
 
 FClothingSimulation::~FClothingSimulation()
 {}
+PRAGMA_ENABLE_DEPRECATION_WARNINGS
 
 void FClothingSimulation::Initialize()
 {
diff --git a/Engine/Plugins/ChaosCloth/Source/ChaosCloth/Public/ChaosCloth/ChaosClothVisualization.h b/Engine/Plugins/ChaosCloth/Source/ChaosCloth/Public/ChaosCloth/ChaosClothVisualization.h
--- a/Engine/Plugins/ChaosCloth/Source/ChaosCloth/Public/ChaosCloth/ChaosClothVisualization.h
+++ b/Engine/Plugins/ChaosCloth/Source/ChaosCloth/Public/ChaosCloth/ChaosClothVisualization.h
@@ -3,6 +3,7 @@
 
 #include "Misc/Build.h"      // For CHAOS_DEBUG_DRAW
 #include "Chaos/Declares.h"  //
+#include "Misc/CoreMiscDefines.h"
 
 #if WITH_EDITOR
 #include "UObject/GCObject.h"
@@ -18,14 +19,11 @@ namespace Chaos
 {
 	class FClothingSimulationSolver;
 
-	class FClothVisualization
-#if WITH_EDITOR
-		: public FGCObject  // Add garbage collection for cloth material
-#endif  // #if WITH_EDITOR
+	class FClothVisualizationNoGC
 	{
 	public:
-		CHAOSCLOTH_API explicit FClothVisualization(const ::Chaos::FClothingSimulationSolver* InSolver = nullptr);
-		CHAOSCLOTH_API virtual ~FClothVisualization();
+		CHAOSCLOTH_API explicit FClothVisualizationNoGC(const ::Chaos::FClothingSimulationSolver* InSolver = nullptr);
+		CHAOSCLOTH_API virtual ~FClothVisualizationNoGC();
 
 #if CHAOS_DEBUG_DRAW
 		// Editor & runtime functions
@@ -102,12 +100,6 @@ namespace Chaos
 		CHAOSCLOTH_API void DrawKinematicColliderShaded(FPrimitiveDrawInterface* PDI) const;
 		CHAOSCLOTH_API void DrawWeightMapWithName(FPrimitiveDrawInterface* PDI, const FString& Name) const;
 		CHAOSCLOTH_API TArray<FString> GetAllWeightMapNames() const;
-	protected:
-		// FGCObject interface
-		CHAOSCLOTH_API virtual void AddReferencedObjects(FReferenceCollector& Collector) override;
-		virtual FString GetReferencerName() const override { return TEXT("UE::Chaos::Cloth::FVisualization"); }
-		// End of FGCObject interface
-
 #else  // #if WITH_EDITOR && CHAOS_DEBUG_DRAW
 		void DrawPhysMeshShaded(FPrimitiveDrawInterface* /*PDI*/) const {}
 		CHAOSCLOTH_API void DrawKinematicColliderShaded(FPrimitiveDrawInterface* /*PDI*/) const {}
@@ -117,13 +109,28 @@ namespace Chaos
 	private:
 		// Simulation objects
 		const ::Chaos::FClothingSimulationSolver* Solver;
-#if WITH_EDITOR
-		// Visualization material
-		TObjectPtr<const UMaterial> ClothMaterial = nullptr;
-		TObjectPtr<const UMaterial> ClothMaterialColor = nullptr;
-		TObjectPtr<const UMaterial> ClothMaterialVertex = nullptr;
-		TObjectPtr<const UMaterial> CollisionMaterial = nullptr;
-#endif  // #if WITH_EDITOR
 #endif  // #if CHAOS_DEBUG_DRAW
+
+		class FMaterials;
 	};
-} // End namespace UE::Chaos::Cloth
+
+	// Deprecated FGCObject inheritance
+	class UE_DEPRECATED(5.5, "Use FClothVisualizationNoGC instead") FClothVisualization : public ::Chaos::FClothVisualizationNoGC
+#if WITH_EDITOR
+		, public FGCObject
+#endif  // #if WITH_EDITOR
+	{
+	public:
+		CHAOSCLOTH_API explicit FClothVisualization(const ::Chaos::FClothingSimulationSolver* InSolver = nullptr);
+		CHAOSCLOTH_API virtual ~FClothVisualization() override;
+#if WITH_EDITOR
+	protected:
+		//~ Begin FGCObject interface
+		UE_DEPRECATED(5.5, "The FGCObject inheritance is no longer needed.")
+		virtual void AddReferencedObjects(FReferenceCollector& /*Collector*/) override {}
+		UE_DEPRECATED(5.5, "The FGCObject inheritance is no longer needed.")
+		virtual FString GetReferencerName() const override { return TEXT("Chaos::FVisualization"); }
+		//~ End FGCObject interface
+#endif  // #if WITH_EDITOR
+ 	};
+} // End namespace Chaos
diff --git a/Engine/Plugins/ChaosCloth/Source/ChaosCloth/Public/ChaosCloth/ChaosClothingSimulation.h b/Engine/Plugins/ChaosCloth/Source/ChaosCloth/Public/ChaosCloth/ChaosClothingSimulation.h
--- a/Engine/Plugins/ChaosCloth/Source/ChaosCloth/Public/ChaosCloth/ChaosClothingSimulation.h
+++ b/Engine/Plugins/ChaosCloth/Source/ChaosCloth/Public/ChaosCloth/ChaosClothingSimulation.h
@@ -123,7 +123,9 @@ namespace Chaos
 
 	private:
 		// Visualization object
-		FClothVisualization Visualization;
+PRAGMA_DISABLE_DEPRECATION_WARNINGS
+		FClothVisualization Visualization;  // TODO: Replace with FClothVisualizationNoGC in 5.7
+PRAGMA_ENABLE_DEPRECATION_WARNINGS
 
 		// Simulation objects
 		TUniquePtr<FClothingSimulationSolver> Solver;  // Default solver
diff --git a/Engine/Plugins/ChaosClothAsset/Source/ChaosClothAssetEngine/Private/ChaosClothAsset/ClothSimulationProxy.cpp b/Engine/Plugins/ChaosClothAsset/Source/ChaosClothAssetEngine/Private/ChaosClothAsset/ClothSimulationProxy.cpp
--- a/Engine/Plugins/ChaosClothAsset/Source/ChaosClothAssetEngine/Private/ChaosClothAsset/ClothSimulationProxy.cpp
+++ b/Engine/Plugins/ChaosClothAsset/Source/ChaosClothAssetEngine/Private/ChaosClothAsset/ClothSimulationProxy.cpp
@@ -111,7 +111,9 @@ namespace UE::Chaos::ClothAsset
 		// Use new SoftsEvolution, not PBDEvolution.
 		constexpr bool bUseLegacySolver = false;
 		Solver = MakeUnique<::Chaos::FClothingSimulationSolver>(Configs[SolverConfigIndex].Get(), bUseLegacySolver);
+PRAGMA_DISABLE_DEPRECATION_WARNINGS
 		Visualization = MakeUnique<::Chaos::FClothVisualization>(Solver.Get());
+PRAGMA_ENABLE_DEPRECATION_WARNINGS
 
 		// Need a valid context to initialize the mesh
 		constexpr bool bIsInitialization = true;
@@ -613,6 +615,11 @@ namespace UE::Chaos::ClothAsset
 		return FBoxSphereBounds(ForceInit);
 	}
 
+	const::Chaos::FClothVisualizationNoGC* FClothSimulationProxy::GetClothVisualization() const
+	{
+		return Visualization.Get();
+	}
+
 	void FClothSimulationProxy::InitializeConfigs()
 	{
 		// Replace physics thread's configs with the game thread's configs
diff --git a/Engine/Plugins/ChaosClothAsset/Source/ChaosClothAssetEngine/Public/ChaosClothAsset/ClothSimulationProxy.h b/Engine/Plugins/ChaosClothAsset/Source/ChaosClothAssetEngine/Public/ChaosClothAsset/ClothSimulationProxy.h
--- a/Engine/Plugins/ChaosClothAsset/Source/ChaosClothAssetEngine/Public/ChaosClothAsset/ClothSimulationProxy.h
+++ b/Engine/Plugins/ChaosClothAsset/Source/ChaosClothAssetEngine/Public/ChaosClothAsset/ClothSimulationProxy.h
@@ -17,6 +17,7 @@ namespace Chaos
 	class FClothingSimulationConfig;
 	class FClothingSimulationCollider;
 	class FClothVisualization;
+	class FClothVisualizationNoGC;
 }
 
 class UChaosClothComponent;
@@ -67,7 +68,9 @@ namespace UE::Chaos::ClothAsset
 
 		FBoxSphereBounds CalculateBounds_AnyThread() const;
 
+		UE_DEPRECATED(5.5, "Use GetClothVisualization instead.")
 		const ::Chaos::FClothVisualization* GetVisualization() const { return Visualization.Get(); }
+		const ::Chaos::FClothVisualizationNoGC* GetClothVisualization() const;
 
 		int32 GetNumCloths() const { return NumCloths; }
 		int32 GetNumKinematicParticles() const { return NumKinematicParticles; }
@@ -121,7 +124,9 @@ namespace UE::Chaos::ClothAsset
 		TArray<TUniquePtr<::Chaos::FClothingSimulationCloth>> Cloths;
 		TArray<TUniquePtr<::Chaos::FClothingSimulationConfig>> Configs;
 		TArray<TUniquePtr<::Chaos::FClothingSimulationCollider>> Colliders;
+PRAGMA_DISABLE_DEPRECATION_WARNINGS
 		TUniquePtr<::Chaos::FClothVisualization> Visualization;
+PRAGMA_ENABLE_DEPRECATION_WARNINGS
 
 		// Chaos Cache needs to have access to the solver.
 		friend class FClothComponentCacheAdapter;
diff --git a/Engine/Plugins/ChaosClothAssetEditor/Source/ChaosClothAssetEditor/Private/ChaosClothAsset/ClothEditorSimulationVisualization.cpp b/Engine/Plugins/ChaosClothAssetEditor/Source/ChaosClothAssetEditor/Private/ChaosClothAsset/ClothEditorSimulationVisualization.cpp
--- a/Engine/Plugins/ChaosClothAssetEditor/Source/ChaosClothAssetEditor/Private/ChaosClothAsset/ClothEditorSimulationVisualization.cpp
+++ b/Engine/Plugins/ChaosClothAssetEditor/Source/ChaosClothAssetEditor/Private/ChaosClothAsset/ClothEditorSimulationVisualization.cpp
@@ -89,9 +89,9 @@ static FText GetSimulationStatisticsString(const FClothSimulationProxy& SimProxy
 	return TextValue;
 }
 
-DECLARE_DELEGATE_ThreeParams(FClothVisualizationDebugDraw, const FClothEditorSimulationVisualization&, const ::Chaos::FClothVisualization&, FPrimitiveDrawInterface*);
+DECLARE_DELEGATE_ThreeParams(FClothVisualizationDebugDraw, const FClothEditorSimulationVisualization&, const ::Chaos::FClothVisualizationNoGC&, FPrimitiveDrawInterface*);
 
-DECLARE_DELEGATE_FourParams(FClothVisualizationDebugDrawTexts, const FClothEditorSimulationVisualization&, const ::Chaos::FClothVisualization&, FCanvas*, const FSceneView*);
+DECLARE_DELEGATE_FourParams(FClothVisualizationDebugDrawTexts, const FClothEditorSimulationVisualization&, const ::Chaos::FClothVisualizationNoGC&, FCanvas*, const FSceneView*);
 
 DECLARE_DELEGATE_RetVal_TwoParams(FText, FLocalDebugDisplayString, const FClothEditorSimulationVisualization&, const FClothSimulationProxy& );
 
@@ -156,7 +156,7 @@ struct FVisualizationOption
 const FVisualizationOption FVisualizationOption::OptionData[] =
 {
 	FVisualizationOption(
-		FClothVisualizationDebugDraw::CreateLambda([](const FClothEditorSimulationVisualization&, const ::Chaos::FClothVisualization& Visualization, FPrimitiveDrawInterface* PDI)
+		FClothVisualizationDebugDraw::CreateLambda([](const FClothEditorSimulationVisualization&, const ::Chaos::FClothVisualizationNoGC& Visualization, FPrimitiveDrawInterface* PDI)
 		{
 			Visualization.DrawPhysMeshShaded(PDI);
 		}),
@@ -164,14 +164,14 @@ const FVisualizationOption FVisualizationOption::OptionData[] =
 		LOCTEXT("ChaosVisName_PhysMeshShaded_ToolTip", "Draws the current physical result as a doubled sided flat shaded mesh"),
 		/*bDisablesSimulation =*/false, /*bHidesClothSections=*/true),
 	FVisualizationOption(
-		FClothVisualizationDebugDraw::CreateLambda([](const FClothEditorSimulationVisualization&, const ::Chaos::FClothVisualization& Visualization, FPrimitiveDrawInterface* PDI)
+		FClothVisualizationDebugDraw::CreateLambda([](const FClothEditorSimulationVisualization&, const ::Chaos::FClothVisualizationNoGC& Visualization, FPrimitiveDrawInterface* PDI)
 		{
 			Visualization.DrawPhysMeshWired(PDI);
 		}), 
 		LOCTEXT("ChaosVisName_PhysMeshWire", "Physical Mesh (Wireframe)"), 
 		LOCTEXT("ChaosVisName_PhysMeshWired_ToolTip", "Draws the current physical mesh result in wireframe")),
 	FVisualizationOption(
-		FClothVisualizationDebugDraw::CreateLambda([](const FClothEditorSimulationVisualization&, const ::Chaos::FClothVisualization& Visualization, FPrimitiveDrawInterface* PDI)
+		FClothVisualizationDebugDraw::CreateLambda([](const FClothEditorSimulationVisualization&, const ::Chaos::FClothVisualizationNoGC& Visualization, FPrimitiveDrawInterface* PDI)
 		{
 			Visualization.DrawAnimMeshWired(PDI);
 		}),
@@ -182,21 +182,21 @@ const FVisualizationOption FVisualizationOption::OptionData[] =
 		LOCTEXT("ChaosVisName_HideRenderMesh_ToolTip", "Hide the render mesh."), 
 		/*bDisablesSimulation =*/false, /*bHidesClothSections=*/true),
 	FVisualizationOption(
-		FClothVisualizationDebugDrawTexts::CreateLambda([](const FClothEditorSimulationVisualization&, const ::Chaos::FClothVisualization& Visualization, FCanvas* Canvas, const FSceneView* SceneView)
+		FClothVisualizationDebugDrawTexts::CreateLambda([](const FClothEditorSimulationVisualization&, const ::Chaos::FClothVisualizationNoGC& Visualization, FCanvas* Canvas, const FSceneView* SceneView)
 		{
 			Visualization.DrawParticleIndices(Canvas, SceneView);
 		}),
 		LOCTEXT("ChaosVisName_ParticleIndices", "Particle Indices"), 
 		LOCTEXT("ChaosVisName_ParticleIndices_ToolTip", "Draws the particle indices as instantiated by the solver")),
 	FVisualizationOption(
-		FClothVisualizationDebugDrawTexts::CreateLambda([](const FClothEditorSimulationVisualization&, const ::Chaos::FClothVisualization& Visualization, FCanvas* Canvas, const FSceneView* SceneView)
+		FClothVisualizationDebugDrawTexts::CreateLambda([](const FClothEditorSimulationVisualization&, const ::Chaos::FClothVisualizationNoGC& Visualization, FCanvas* Canvas, const FSceneView* SceneView)
 		{
 			Visualization.DrawElementIndices(Canvas, SceneView);
 		}),
 		LOCTEXT("ChaosVisName_ElementIndices", "Element Indices"), 
 		LOCTEXT("ChaosVisName_ElementIndices_ToolTip", "Draws the element's (triangle or other) indices as instantiated by the solver")),
 	FVisualizationOption(
-		FClothVisualizationDebugDraw::CreateLambda([](const FClothEditorSimulationVisualization& EditorVisualization, const ::Chaos::FClothVisualization& Visualization, FPrimitiveDrawInterface* PDI)
+		FClothVisualizationDebugDraw::CreateLambda([](const FClothEditorSimulationVisualization& EditorVisualization, const ::Chaos::FClothVisualizationNoGC& Visualization, FPrimitiveDrawInterface* PDI)
 		{
 			Visualization.DrawPointNormals(PDI, EditorVisualization.GetPointNormalLength());
 		}),
@@ -208,14 +208,14 @@ const FVisualizationOption FVisualizationOption::OptionData[] =
 			EditorVisualization.ExtendViewportShowMenuPointNormalsLength(MenuBuilder);
 		})),
 	FVisualizationOption(
-		FClothVisualizationDebugDraw::CreateLambda([](const FClothEditorSimulationVisualization&, const ::Chaos::FClothVisualization& Visualization, FPrimitiveDrawInterface* PDI)
+		FClothVisualizationDebugDraw::CreateLambda([](const FClothEditorSimulationVisualization&, const ::Chaos::FClothVisualizationNoGC& Visualization, FPrimitiveDrawInterface* PDI)
 		{
 			Visualization.DrawPointVelocities(PDI);
 		}),
 		LOCTEXT("ChaosVisName_PointVelocities", "Point Velocities"), 
 		LOCTEXT("ChaosVisName_PointVelocities_ToolTip", "Draws the current point velocities for the simulation mesh")),
 	FVisualizationOption(
-		FClothVisualizationDebugDraw::CreateLambda([](const FClothEditorSimulationVisualization& EditorVisualization, const ::Chaos::FClothVisualization& Visualization, FPrimitiveDrawInterface* PDI)
+		FClothVisualizationDebugDraw::CreateLambda([](const FClothEditorSimulationVisualization& EditorVisualization, const ::Chaos::FClothVisualizationNoGC& Visualization, FPrimitiveDrawInterface* PDI)
 		{
 			Visualization.DrawAnimNormals(PDI, EditorVisualization.GetAnimatedNormalLength());
 		}),
@@ -227,21 +227,21 @@ const FVisualizationOption FVisualizationOption::OptionData[] =
 			EditorVisualization.ExtendViewportShowMenuAnimatedNormalsLength(MenuBuilder);
 		})),
 	FVisualizationOption(
-		FClothVisualizationDebugDraw::CreateLambda([](const FClothEditorSimulationVisualization&, const ::Chaos::FClothVisualization& Visualization, FPrimitiveDrawInterface* PDI)
+		FClothVisualizationDebugDraw::CreateLambda([](const FClothEditorSimulationVisualization&, const ::Chaos::FClothVisualizationNoGC& Visualization, FPrimitiveDrawInterface* PDI)
 		{
 			Visualization.DrawCollision(PDI);
 		}),
 		LOCTEXT("ChaosVisName_Collision", "Collisions"), 
 		LOCTEXT("ChaosVisName_Collision_ToolTip", "Draws the collision bodies the simulation is currently using")),
 	FVisualizationOption(
-		FClothVisualizationDebugDraw::CreateLambda([](const FClothEditorSimulationVisualization&, const ::Chaos::FClothVisualization& Visualization, FPrimitiveDrawInterface* PDI)
+		FClothVisualizationDebugDraw::CreateLambda([](const FClothEditorSimulationVisualization&, const ::Chaos::FClothVisualizationNoGC& Visualization, FPrimitiveDrawInterface* PDI)
 		{
 			Visualization.DrawBackstops(PDI);
 		}),
 		LOCTEXT("ChaosVisName_Backstop", "Backstops"), 
 		LOCTEXT("ChaosVisName_Backstop_ToolTip", "Draws the backstop radius and position for each simulation particle")),
 	FVisualizationOption(
-		FClothVisualizationDebugDraw::CreateLambda([](const FClothEditorSimulationVisualization&, const ::Chaos::FClothVisualization& Visualization, FPrimitiveDrawInterface* PDI)
+		FClothVisualizationDebugDraw::CreateLambda([](const FClothEditorSimulationVisualization&, const ::Chaos::FClothVisualizationNoGC& Visualization, FPrimitiveDrawInterface* PDI)
 		{
 			Visualization.DrawBackstopDistances(PDI);
 		}),
@@ -249,91 +249,91 @@ const FVisualizationOption FVisualizationOption::OptionData[] =
 		LOCTEXT("ChaosVisName_BackstopDistance_ToolTip", "Draws the backstop distance offset for each simulation particle"), 
 		/*bDisablesSimulation =*/true),
 	FVisualizationOption(
-		FClothVisualizationDebugDraw::CreateLambda([](const FClothEditorSimulationVisualization&, const ::Chaos::FClothVisualization& Visualization, FPrimitiveDrawInterface* PDI)
+		FClothVisualizationDebugDraw::CreateLambda([](const FClothEditorSimulationVisualization&, const ::Chaos::FClothVisualizationNoGC& Visualization, FPrimitiveDrawInterface* PDI)
 		{
 			Visualization.DrawMaxDistances(PDI);
 		}),
 		LOCTEXT("ChaosVisName_MaxDistance", "Max Distances"), 
 		LOCTEXT("ChaosVisName_MaxDistance_ToolTip", "Draws the current max distances for the sim particles as a line along its normal")),
 	FVisualizationOption(
-		FClothVisualizationDebugDrawTexts::CreateLambda([](const FClothEditorSimulationVisualization&, const ::Chaos::FClothVisualization& Visualization, FCanvas* Canvas, const FSceneView* SceneView)
+		FClothVisualizationDebugDrawTexts::CreateLambda([](const FClothEditorSimulationVisualization&, const ::Chaos::FClothVisualizationNoGC& Visualization, FCanvas* Canvas, const FSceneView* SceneView)
 		{
 			Visualization.DrawMaxDistanceValues(Canvas, SceneView);
 		}),
 		LOCTEXT("ChaosVisName_MaxDistanceValue", "Max Distances As Numbers"), 
 		LOCTEXT("ChaosVisName_MaxDistanceValue_ToolTip", "Draws the current max distances as numbers")),
 	FVisualizationOption(
-		FClothVisualizationDebugDraw::CreateLambda([](const FClothEditorSimulationVisualization&, const ::Chaos::FClothVisualization& Visualization, FPrimitiveDrawInterface* PDI)
+		FClothVisualizationDebugDraw::CreateLambda([](const FClothEditorSimulationVisualization&, const ::Chaos::FClothVisualizationNoGC& Visualization, FPrimitiveDrawInterface* PDI)
 		{
 			Visualization.DrawAnimDrive(PDI);
 		}),
 		LOCTEXT("ChaosVisName_AnimDrive", "Anim Drive"), 
 		LOCTEXT("ChaosVisName_AnimDrive_Tooltip", "Draws the current skinned reference mesh for the simulation which anim drive will attempt to reach if enabled")),
 	FVisualizationOption(
-		FClothVisualizationDebugDraw::CreateLambda([](const FClothEditorSimulationVisualization&, const ::Chaos::FClothVisualization& Visualization, FPrimitiveDrawInterface* PDI)
+		FClothVisualizationDebugDraw::CreateLambda([](const FClothEditorSimulationVisualization&, const ::Chaos::FClothVisualizationNoGC& Visualization, FPrimitiveDrawInterface* PDI)
 		{
 			Visualization.DrawEdgeConstraint(PDI);
 		}),
 		LOCTEXT("ChaosVisName_EdgeConstraint", "Edge Constraint"),
 		LOCTEXT("ChaosVisName_EdgeConstraint_Tooltip", "Draws the edge spring constraints")),
 	FVisualizationOption(
-		FClothVisualizationDebugDraw::CreateLambda([](const FClothEditorSimulationVisualization&, const ::Chaos::FClothVisualization& Visualization, FPrimitiveDrawInterface* PDI)
+		FClothVisualizationDebugDraw::CreateLambda([](const FClothEditorSimulationVisualization&, const ::Chaos::FClothVisualizationNoGC& Visualization, FPrimitiveDrawInterface* PDI)
 		{
 			Visualization.DrawBendingConstraint(PDI);
 		}),
 		LOCTEXT("ChaosVisName_BendingConstraint", "Bending Constraint"), 
 		LOCTEXT("ChaosVisName_BendingConstraint_Tooltip", "Draws the bending spring constraints")),
 	FVisualizationOption(
-		FClothVisualizationDebugDraw::CreateLambda([](const FClothEditorSimulationVisualization&, const ::Chaos::FClothVisualization& Visualization, FPrimitiveDrawInterface* PDI)
+		FClothVisualizationDebugDraw::CreateLambda([](const FClothEditorSimulationVisualization&, const ::Chaos::FClothVisualizationNoGC& Visualization, FPrimitiveDrawInterface* PDI)
 		{
 			Visualization.DrawLongRangeConstraint(PDI);
 		}),
 		LOCTEXT("ChaosVisName_LongRangeConstraint" , "Long Range Constraint"), 
 		LOCTEXT("ChaosVisName_LongRangeConstraint_Tooltip", "Draws the long range attachment constraint distances")),
 	FVisualizationOption(
-		FClothVisualizationDebugDraw::CreateLambda([](const FClothEditorSimulationVisualization&, const ::Chaos::FClothVisualization& Visualization, FPrimitiveDrawInterface* PDI)
+		FClothVisualizationDebugDraw::CreateLambda([](const FClothEditorSimulationVisualization&, const ::Chaos::FClothVisualizationNoGC& Visualization, FPrimitiveDrawInterface* PDI)
 		{
 			Visualization.DrawWindAndPressureForces(PDI);
 		}),
 		LOCTEXT("ChaosVisName_WindAndPressureForces", "Wind Aerodynamic And Pressure Forces"), 
 		LOCTEXT("ChaosVisName_WindAndPressure_Tooltip", "Draws the Wind drag and lift and pressure forces")),
 	FVisualizationOption(
-		FClothVisualizationDebugDraw::CreateLambda([](const FClothEditorSimulationVisualization&, const ::Chaos::FClothVisualization& Visualization, FPrimitiveDrawInterface* PDI)
+		FClothVisualizationDebugDraw::CreateLambda([](const FClothEditorSimulationVisualization&, const ::Chaos::FClothVisualizationNoGC& Visualization, FPrimitiveDrawInterface* PDI)
 		{
 			Visualization.DrawLocalSpace(PDI);
 		}),
 		LOCTEXT("ChaosVisName_LocalSpace", "Local Space Reference Bone"), 
 		LOCTEXT("ChaosVisName_LocalSpace_Tooltip", "Draws the local space reference bone")),
 	FVisualizationOption(
-		FClothVisualizationDebugDraw::CreateLambda([](const FClothEditorSimulationVisualization&, const ::Chaos::FClothVisualization& Visualization, FPrimitiveDrawInterface* PDI)
+		FClothVisualizationDebugDraw::CreateLambda([](const FClothEditorSimulationVisualization&, const ::Chaos::FClothVisualizationNoGC& Visualization, FPrimitiveDrawInterface* PDI)
 		{
 			Visualization.DrawGravity(PDI);
 		}),
 		LOCTEXT("ChaosVisName_Gravity", "Gravity"), 
 		LOCTEXT("ChaosVisName_Gravity_Tooltip", "Draws gravity")),
 	FVisualizationOption(
-		FClothVisualizationDebugDraw::CreateLambda([](const FClothEditorSimulationVisualization&, const ::Chaos::FClothVisualization& Visualization, FPrimitiveDrawInterface* PDI)
+		FClothVisualizationDebugDraw::CreateLambda([](const FClothEditorSimulationVisualization&, const ::Chaos::FClothVisualizationNoGC& Visualization, FPrimitiveDrawInterface* PDI)
 		{
 			Visualization.DrawFictitiousAngularForces(PDI);
 		}),
 		LOCTEXT("ChaosVisName_FictitiousAngularForces", "Fictitious Angular Forces"),
 		LOCTEXT("ChaosVisName_Gravity_FictitiousAngularForces", "Draws fictitious angular forces (force based solver only)")),
 	FVisualizationOption(
-		FClothVisualizationDebugDraw::CreateLambda([](const FClothEditorSimulationVisualization&, const ::Chaos::FClothVisualization& Visualization, FPrimitiveDrawInterface* PDI)
+		FClothVisualizationDebugDraw::CreateLambda([](const FClothEditorSimulationVisualization&, const ::Chaos::FClothVisualizationNoGC& Visualization, FPrimitiveDrawInterface* PDI)
 		{
 			Visualization.DrawSelfCollision(PDI);
 		}),
 		LOCTEXT("ChaosVisName_SelfCollision", "Self Collision"), 
 		LOCTEXT("ChaosVisName_SelfCollision_Tooltip", "Draws the self collision thickness/debugging information")),
 	FVisualizationOption(
-		FClothVisualizationDebugDraw::CreateLambda([](const FClothEditorSimulationVisualization&, const ::Chaos::FClothVisualization& Visualization, FPrimitiveDrawInterface* PDI)
+		FClothVisualizationDebugDraw::CreateLambda([](const FClothEditorSimulationVisualization&, const ::Chaos::FClothVisualizationNoGC& Visualization, FPrimitiveDrawInterface* PDI)
 		{
 			Visualization.DrawSelfIntersection(PDI);
 		}),
 		LOCTEXT("ChaosVisName_SelfIntersection", "Self Intersection"), 
 		LOCTEXT("ChaosVisName_SelfIntersection_Tooltip", "Draws the self intersection contour/region information")),
 	FVisualizationOption(
-		FClothVisualizationDebugDraw::CreateLambda([](const FClothEditorSimulationVisualization&, const ::Chaos::FClothVisualization& Visualization, FPrimitiveDrawInterface* PDI)
+		FClothVisualizationDebugDraw::CreateLambda([](const FClothEditorSimulationVisualization&, const ::Chaos::FClothVisualizationNoGC& Visualization, FPrimitiveDrawInterface* PDI)
 		{
 			Visualization.DrawSelfCollisionLayers(PDI);
 		}),
@@ -341,28 +341,28 @@ const FVisualizationOption FVisualizationOption::OptionData[] =
 		LOCTEXT("ChaosVisName_SelfCollisionLayers_Tooltip", "Draws the self collision layers"),
 		/*bDisablesSimulation =*/false, /*bHidesClothSections=*/true),
 	FVisualizationOption(
-		FClothVisualizationDebugDraw::CreateLambda([](const FClothEditorSimulationVisualization&, const ::Chaos::FClothVisualization& Visualization, FPrimitiveDrawInterface* PDI)
+		FClothVisualizationDebugDraw::CreateLambda([](const FClothEditorSimulationVisualization&, const ::Chaos::FClothVisualizationNoGC& Visualization, FPrimitiveDrawInterface* PDI)
 		{
 			Visualization.DrawSelfCollisionThickness(PDI);
 		}),
 		LOCTEXT("ChaosVisName_SelfCollisionThickness", "Self Collision Thickness"), 
 		LOCTEXT("ChaosVisName_SelfCollisionThickness_Tooltip", "Draws the self collision Thickness")),
 	FVisualizationOption(
-		FClothVisualizationDebugDraw::CreateLambda([](const FClothEditorSimulationVisualization&, const ::Chaos::FClothVisualization& Visualization, FPrimitiveDrawInterface* PDI)
+		FClothVisualizationDebugDraw::CreateLambda([](const FClothEditorSimulationVisualization&, const ::Chaos::FClothVisualizationNoGC& Visualization, FPrimitiveDrawInterface* PDI)
 		{
 			Visualization.DrawKinematicColliderShaded(PDI);
 		}),
 		LOCTEXT("ChaosVisName_DrawKinematicColliderShaded", "Draw Kinematic Colliders (Shaded)"), 
 		LOCTEXT("ChaosVisName_DrawKinematicColliderShaded_Tooltip", "Draw kinematic cloth colliders with flat shading.")),
 	FVisualizationOption(
-		FClothVisualizationDebugDraw::CreateLambda([](const FClothEditorSimulationVisualization&, const ::Chaos::FClothVisualization& Visualization, FPrimitiveDrawInterface* PDI)
+		FClothVisualizationDebugDraw::CreateLambda([](const FClothEditorSimulationVisualization&, const ::Chaos::FClothVisualizationNoGC& Visualization, FPrimitiveDrawInterface* PDI)
 		{
 			Visualization.DrawKinematicColliderWired(PDI);
 		}),
 		LOCTEXT("ChaosVisName_DrawKinematicColliderWired", "Draw Kinematic Colliders (Wireframe)"), 
 		LOCTEXT("ChaosVisName_DrawKinematicColliderWired_Tooltip", "Draw kinematic cloth colliders in wireframe.")),
 	FVisualizationOption(
-		FClothVisualizationDebugDraw::CreateLambda([](const FClothEditorSimulationVisualization&, const ::Chaos::FClothVisualization& Visualization, FPrimitiveDrawInterface* PDI)
+		FClothVisualizationDebugDraw::CreateLambda([](const FClothEditorSimulationVisualization&, const ::Chaos::FClothVisualizationNoGC& Visualization, FPrimitiveDrawInterface* PDI)
 		{
 			Visualization.DrawMultiResConstraint(PDI);
 		}),
@@ -376,7 +376,7 @@ const FVisualizationOption FVisualizationOption::OptionData[] =
 		LOCTEXT("ChaosVisName_SimulationStatistics", "Simulation Statistics"), 
 		LOCTEXT("ChaosVisName_SimulationStatistics_Tooltip", "Displays simulation statistics")),
 	FVisualizationOption(
-		FClothVisualizationDebugDraw::CreateLambda([](const FClothEditorSimulationVisualization& EditorVisualization, const ::Chaos::FClothVisualization& Visualization, FPrimitiveDrawInterface* PDI)
+		FClothVisualizationDebugDraw::CreateLambda([](const FClothEditorSimulationVisualization& EditorVisualization, const ::Chaos::FClothVisualizationNoGC& Visualization, FPrimitiveDrawInterface* PDI)
 		{
 			const FString* const WeightMap = EditorVisualization.GetCurrentlySelectedWeightMap();
 			Visualization.DrawWeightMapWithName(PDI, WeightMap ? *WeightMap : FString());
@@ -389,7 +389,7 @@ const FVisualizationOption FVisualizationOption::OptionData[] =
 			EditorVisualization.ExtendViewportShowMenuWeightMapSelector(MenuBuilder, ViewportClient);
 		})),
 	FVisualizationOption(
-		FClothVisualizationDebugDraw::CreateLambda([](const FClothEditorSimulationVisualization&, const ::Chaos::FClothVisualization& Visualization, FPrimitiveDrawInterface* PDI)
+		FClothVisualizationDebugDraw::CreateLambda([](const FClothEditorSimulationVisualization&, const ::Chaos::FClothVisualizationNoGC& Visualization, FPrimitiveDrawInterface* PDI)
 		{
 			Visualization.DrawInpaintWeightsMatched(PDI);
 		}),
@@ -486,7 +486,7 @@ void FClothEditorSimulationVisualization::ExtendViewportShowMenuSpinBox(FMenuBui
 void FClothEditorSimulationVisualization::RefreshMenusForClothComponent(const UChaosClothComponent* ClothComponent)
 {
 	TArray<FString> WeightMaps;
-	if (const ::Chaos::FClothVisualization* const Visualization = ClothComponent && ClothComponent->GetClothSimulationProxy() ? ClothComponent->GetClothSimulationProxy()->GetVisualization() : nullptr)
+	if (const ::Chaos::FClothVisualizationNoGC* const Visualization = ClothComponent && ClothComponent->GetClothSimulationProxy() ? ClothComponent->GetClothSimulationProxy()->GetClothVisualization() : nullptr)
 	{
 		WeightMaps = Visualization->GetAllWeightMapNames();
 	}
@@ -522,7 +522,7 @@ void FClothEditorSimulationVisualization::RefreshMenusForClothComponent(const UC
 
 void FClothEditorSimulationVisualization::DebugDrawSimulation(const UChaosClothComponent* ClothComponent, FPrimitiveDrawInterface* PDI)
 {
-	const ::Chaos::FClothVisualization* const Visualization = ClothComponent && ClothComponent->GetClothSimulationProxy() ? ClothComponent->GetClothSimulationProxy()->GetVisualization() : nullptr;
+	const ::Chaos::FClothVisualizationNoGC* const Visualization = ClothComponent && ClothComponent->GetClothSimulationProxy() ? ClothComponent->GetClothSimulationProxy()->GetClothVisualization() : nullptr;
 	if (!Visualization)
 	{
 		return;
@@ -539,7 +539,7 @@ void FClothEditorSimulationVisualization::DebugDrawSimulation(const UChaosClothC
 
 void FClothEditorSimulationVisualization::DebugDrawSimulationTexts(const UChaosClothComponent* ClothComponent, FCanvas* Canvas, const FSceneView* SceneView)
 {
-	const ::Chaos::FClothVisualization* const Visualization = ClothComponent && ClothComponent->GetClothSimulationProxy() ? ClothComponent->GetClothSimulationProxy()->GetVisualization() : nullptr;
+	const ::Chaos::FClothVisualizationNoGC* const Visualization = ClothComponent && ClothComponent->GetClothSimulationProxy() ? ClothComponent->GetClothSimulationProxy()->GetClothVisualization() : nullptr;
 	if (!Visualization)
 	{
 		return;
