Chaos Cloth: Add the ability to specify Velocity Scale and clamps in different simulation spaces (defaults to Reference Bone space, which was the previous behavior).

diff --git a/Engine/Plugins/ChaosCloth/Source/ChaosCloth/Private/ChaosCloth/ChaosClothVisualization.cpp b/Engine/Plugins/ChaosCloth/Source/ChaosCloth/Private/ChaosCloth/ChaosClothVisualization.cpp
--- a/Engine/Plugins/ChaosCloth/Source/ChaosCloth/Private/ChaosCloth/ChaosClothVisualization.cpp
+++ b/Engine/Plugins/ChaosCloth/Source/ChaosCloth/Private/ChaosCloth/ChaosClothVisualization.cpp
@@ -2307,6 +2307,7 @@ namespace Chaos
 
 		// Draw local space
 		DrawCoordinateSystem(PDI, FQuat::Identity, Solver->GetLocalSpaceLocation(), 5.f, 0.25f);
+		DrawCoordinateSystem(PDI, Solver->GetLocalSpaceRotation(), Solver->GetLocalSpaceLocation(), 7.f, 0.5f);
 
 		// Draw reference spaces
 		for (const FClothingSimulationCloth* const Cloth : Solver->GetCloths())
diff --git a/Engine/Plugins/ChaosCloth/Source/ChaosCloth/Private/ChaosCloth/ChaosClothingSimulationCloth.cpp b/Engine/Plugins/ChaosCloth/Source/ChaosCloth/Private/ChaosCloth/ChaosClothingSimulationCloth.cpp
--- a/Engine/Plugins/ChaosCloth/Source/ChaosCloth/Private/ChaosCloth/ChaosClothingSimulationCloth.cpp
+++ b/Engine/Plugins/ChaosCloth/Source/ChaosCloth/Private/ChaosCloth/ChaosClothingSimulationCloth.cpp
@@ -12,6 +12,7 @@
 #include "Chaos/PBDFlatWeightMap.h"
 #include "Chaos/PBDSoftBodyCollisionConstraint.h"
 #include "Chaos/SoftsExternalForces.h"
+#include "Chaos/SoftsSimulationSpace.h"
 #include "Containers/ArrayView.h"
 #include "GeometryCollection/ManagedArrayCollection.h"
 #include "HAL/IConsoleManager.h"
@@ -35,6 +36,7 @@ namespace ClothingSimulationClothDefault
 	constexpr float AirDensity = 1.225f;  // Air density in kg/m^3
 	constexpr float GravityScale = Softs::FExternalForces::DefaultGravityScale; // 1.f;
 	constexpr float GravityZOverride = Softs::FExternalForces::DefaultGravityZOverride; // -980.665f;
+	constexpr EChaosSoftsSimulationSpace VelocityScaleSpace = EChaosSoftsSimulationSpace::ReferenceBoneSpace;
 	constexpr float VelocityScale = 0.75f;
 	constexpr float MaxVelocityScale = 1.f;
 	constexpr float MaxVelocity = TNumericLimits<float>::Max();
@@ -367,7 +369,7 @@ void FClothingSimulationCloth::FLODData::Update(FClothingSimulationSolver* Solve
 	check(Cloth->Config);
 	const Softs::FSolverReal MeshScale = Cloth->Mesh->GetScale();
 	const Softs::FSolverReal MaxDistancesScale = (Softs::FSolverReal)Cloth->MaxDistancesMultiplier;
-	ClothConstraints.Update(Cloth->Config->GetProperties(SolverDatum.LODIndex), WeightMaps, VertexSets, FaceSets, FaceIntMaps, MeshScale, MaxDistancesScale);
+	ClothConstraints.Update(Cloth->Config->GetProperties(SolverDatum.LODIndex), WeightMaps, VertexSets, FaceSets, FaceIntMaps, MeshScale, MaxDistancesScale, Solver->GetLocalSpaceRotation(), Cloth->ReferenceSpaceTransform.GetRotation());
 }
 
 void FClothingSimulationCloth::FLODData::Enable(FClothingSimulationSolver* Solver, bool bEnable) const
@@ -894,6 +896,7 @@ void FClothingSimulationCloth::Update(FClothingSimulationSolver* Solver)
 		// TODO: Move all groupID updates out of the cloth update to allow to use of the same GroupId with different cloths
 
 		// Set the reference input velocity and deal with teleport & reset; external forces depends on these values, so they must be initialized before then
+		EChaosSoftsSimulationSpace VelocityScaleSpace = ClothingSimulationClothDefault::VelocityScaleSpace;
 		FVec3f OutLinearVelocityScale;
 		FRealSingle OutAngularVelocityScale;
 		FRealSingle OutMaxVelocityScale;
@@ -939,6 +942,7 @@ void FClothingSimulationCloth::Update(FClothingSimulationSolver* Solver)
 		else
 		{
 			// Use the cloth config parameters
+			VelocityScaleSpace = (EChaosSoftsSimulationSpace)ConfigProperties.GetValue<int32>(TEXT("VelocityScaleSpace"), (int32)ClothingSimulationClothDefault::VelocityScaleSpace);
 			OutLinearVelocityScale = ConfigProperties.GetValue<FVector3f>(TEXT("LinearVelocityScale"), FVector3f(ClothingSimulationClothDefault::VelocityScale));
 			OutAngularVelocityScale = ConfigProperties.GetValue<float>(TEXT("AngularVelocityScale"), ClothingSimulationClothDefault::VelocityScale);
 			OutMaxVelocityScale = ConfigProperties.GetValue<float>(TEXT("MaxVelocityScale"), ClothingSimulationClothDefault::MaxVelocityScale);
@@ -956,6 +960,7 @@ void FClothingSimulationCloth::Update(FClothingSimulationSolver* Solver)
 			ReferenceSpaceTransform,
 			ReferenceSpaceVelocity,
 			ReferenceSpaceAngularVelocity,
+			VelocityScaleSpace,
 			OutLinearVelocityScale,
 			MaxLinearVelocity,
 			MaxLinearAcceleration,
diff --git a/Engine/Plugins/ChaosCloth/Source/ChaosCloth/Private/ChaosCloth/ChaosClothingSimulationSolver.cpp b/Engine/Plugins/ChaosCloth/Source/ChaosCloth/Private/ChaosCloth/ChaosClothingSimulationSolver.cpp
--- a/Engine/Plugins/ChaosCloth/Source/ChaosCloth/Private/ChaosCloth/ChaosClothingSimulationSolver.cpp
+++ b/Engine/Plugins/ChaosCloth/Source/ChaosCloth/Private/ChaosCloth/ChaosClothingSimulationSolver.cpp
@@ -1084,6 +1084,7 @@ void FClothingSimulationSolver::SetReferenceVelocityScale(uint32 GroupId,
 	const FRigidTransform3& ReferenceSpaceTransform,
 	TVec3<FReal>& InOutReferenceVelocity, // Old reference velocity is passed in. New reference velocity is returned.
 	TVec3<FReal>& InOutReferenceAngularVelocity, // Old reference velocity is passed in. New reference velocity is returned.
+	const EChaosSoftsSimulationSpace VelocityScaleSpace, // the space the following linear velocity properties are in.
 	const TVec3<FRealSingle>& LinearVelocityScale,
 	const TVec3<FRealSingle>& MaxLinearVelocity,
 	const TVec3<FRealSingle>& MaxLinearAcceleration,
@@ -1127,8 +1128,25 @@ void FClothingSimulationSolver::SetReferenceVelocityScale(uint32 GroupId,
 	}	
 	const FVec3 FullDeltaRotation = FullDeltaAngle * FullAxis;
 
-	// Do all linear velocity changes in old reference space.
-	const FVec3 ReferenceSpaceDeltaTranslation= OldReferenceSpaceRotationInverse * LocalDeltaTranslation;
+	// Choose which space to do linear velocity scale and clamp calculations based on VelocityScaleSpace
+	FRotation3 VelocityScaleSpaceToOldRefSpace;
+	FVec3 VelocityScaleSpaceDeltaTranslation;
+	switch (VelocityScaleSpace)
+	{
+	case EChaosSoftsSimulationSpace::WorldSpace:
+		VelocityScaleSpaceToOldRefSpace = OldReferenceSpaceRotationInverse;
+		VelocityScaleSpaceDeltaTranslation = LocalDeltaTranslation;
+		break;
+	case EChaosSoftsSimulationSpace::ComponentSpace:
+		VelocityScaleSpaceToOldRefSpace = OldReferenceSpaceRotationInverse * LocalSpaceRotation;
+		VelocityScaleSpaceDeltaTranslation = LocalSpaceRotation.UnrotateVector(LocalDeltaTranslation);
+		break;
+	case EChaosSoftsSimulationSpace::ReferenceBoneSpace:
+	default:
+		VelocityScaleSpaceToOldRefSpace = FRotation3::Identity;
+		VelocityScaleSpaceDeltaTranslation = OldReferenceSpaceRotationInverse * LocalDeltaTranslation;
+		break;
+	}
 
 	// Apply linear velocity scale.
 	const FVec3 LinearRatio = FVec3(
@@ -1136,7 +1154,7 @@ void FClothingSimulationSolver::SetReferenceVelocityScale(uint32 GroupId,
 		CalculateClampedVelocityScale(LinearVelocityScale[1]),
 		CalculateClampedVelocityScale(LinearVelocityScale[2]));
 
-	FVec3 AppliedDeltaTranslation = LinearRatio * ReferenceSpaceDeltaTranslation;
+	FVec3 AppliedDeltaTranslation = LinearRatio * VelocityScaleSpaceDeltaTranslation;
 
 	// Apply angular velocity scale
 	const FReal AngularRatio = CalculateClampedVelocityScale(AngularVelocityScale);
@@ -1194,7 +1212,7 @@ void FClothingSimulationSolver::SetReferenceVelocityScale(uint32 GroupId,
 
 		AppliedDeltaRotation = AppliedAngularVelocity * DeltaTime;
 	}
-	const FVec3 DeltaPosition = ReferenceSpaceDeltaTranslation - AppliedDeltaTranslation;
+	const FVec3 DeltaPosition = VelocityScaleSpaceToOldRefSpace * (VelocityScaleSpaceDeltaTranslation - AppliedDeltaTranslation);
 	const FVec3 DeltaRotation = FullDeltaRotation - AppliedDeltaRotation;
 	const FReal DeltaAngle = DeltaRotation.Length();
 	const FVec3 Axis = DeltaAngle > UE_KINDA_SMALL_NUMBER ? DeltaRotation / DeltaAngle : FVec3::XAxisVector;
diff --git a/Engine/Plugins/ChaosCloth/Source/ChaosCloth/Public/ChaosCloth/ChaosClothingSimulationSolver.h b/Engine/Plugins/ChaosCloth/Source/ChaosCloth/Public/ChaosCloth/ChaosClothingSimulationSolver.h
--- a/Engine/Plugins/ChaosCloth/Source/ChaosCloth/Public/ChaosCloth/ChaosClothingSimulationSolver.h
+++ b/Engine/Plugins/ChaosCloth/Source/ChaosCloth/Public/ChaosCloth/ChaosClothingSimulationSolver.h
@@ -4,6 +4,7 @@
 #include "ChaosCloth/ChaosClothConstraints.h"
 #include "Chaos/PBDSoftsEvolutionFwd.h"
 #include "Chaos/ArrayCollectionArray.h"
+#include "Chaos/SoftsSimulationSpace.h"
 #include "Chaos/Transform.h"
 #include "Chaos/Framework/PhysicsSolverBase.h"
 #include "Chaos/ImplicitObject.h"
@@ -136,6 +137,7 @@ namespace Chaos
 			const FRigidTransform3& ReferenceSpaceTransform,
 			TVec3<FReal>& InOutReferenceVelocity, // Old reference velocity is passed in. New reference velocity is returned.
 			TVec3<FReal>& InOutReferenceAngularVelocity, // Old reference velocity is passed in. New reference velocity is returned.
+			const EChaosSoftsSimulationSpace VelocityScaleSpace, // the space the following linear velocity properties are in.
 			const TVec3<FRealSingle>& LinearVelocityScale,
 			const TVec3<FRealSingle>& MaxLinearVelocity,
 			const TVec3<FRealSingle>& MaxLinearAcceleration,
@@ -158,8 +160,9 @@ namespace Chaos
 			// Acceleration clamps are disabled, so it doesn't matter what the old velocities were.
 			TVec3<FReal> ReferenceVelocity(0.);
 			TVec3<FReal> ReferenceAngularVelocity(0.);
-			SetReferenceVelocityScale(GroupId, OldReferenceSpaceTransform, ReferenceSpaceTransform, ReferenceVelocity, ReferenceAngularVelocity, LinearVelocityScale,
-				TVec3<FRealSingle>(TNumericLimits<FRealSingle>::Max()), TVec3<FRealSingle>(TNumericLimits<FRealSingle>::Max()), AngularVelocityScale, FictitiousAngularScale,
+			SetReferenceVelocityScale(GroupId, OldReferenceSpaceTransform, ReferenceSpaceTransform, ReferenceVelocity, ReferenceAngularVelocity,
+				EChaosSoftsSimulationSpace::ReferenceBoneSpace, LinearVelocityScale, TVec3<FRealSingle>(TNumericLimits<FRealSingle>::Max()),
+				TVec3<FRealSingle>(TNumericLimits<FRealSingle>::Max()), AngularVelocityScale, FictitiousAngularScale,
 				MaxVelocityScale, TNumericLimits<FRealSingle>::Max(), TNumericLimits<FRealSingle>::Max(), bDisableFictitiousForces);
 		}
 
diff --git a/Engine/Plugins/ChaosClothAssetEditor/Source/ChaosClothAssetDataflowNodes/Private/ChaosClothAsset/SimulationVelocityScaleConfigNode.cpp b/Engine/Plugins/ChaosClothAssetEditor/Source/ChaosClothAssetDataflowNodes/Private/ChaosClothAsset/SimulationVelocityScaleConfigNode.cpp
--- a/Engine/Plugins/ChaosClothAssetEditor/Source/ChaosClothAssetDataflowNodes/Private/ChaosClothAsset/SimulationVelocityScaleConfigNode.cpp
+++ b/Engine/Plugins/ChaosClothAssetEditor/Source/ChaosClothAssetDataflowNodes/Private/ChaosClothAsset/SimulationVelocityScaleConfigNode.cpp
@@ -14,6 +14,7 @@ FChaosClothAssetSimulationVelocityScaleConfigNode::FChaosClothAssetSimulationVel
 void FChaosClothAssetSimulationVelocityScaleConfigNode::AddProperties(FPropertyHelper& PropertyHelper) const
 {
 	static const FVector3f DefaultLinearClamp = { TNumericLimits<float>::Max(), TNumericLimits<float>::Max(), TNumericLimits<float>::Max() };
+	PropertyHelper.SetPropertyEnum(this, &VelocityScaleSpace, {}, ECollectionPropertyFlags::None);
 	PropertyHelper.SetProperty(this, &LinearVelocityScale);
 	PropertyHelper.SetProperty(TEXT("MaxLinearVelocity"), bEnableLinearVelocityClamping ? MaxLinearVelocity : DefaultLinearClamp);
 	PropertyHelper.SetProperty(TEXT("MaxLinearAcceleration"), bEnableLinearAccelerationClamping ? MaxLinearAcceleration : DefaultLinearClamp);
diff --git a/Engine/Plugins/ChaosClothAssetEditor/Source/ChaosClothAssetDataflowNodes/Public/ChaosClothAsset/SimulationVelocityScaleConfigNode.h b/Engine/Plugins/ChaosClothAssetEditor/Source/ChaosClothAssetDataflowNodes/Public/ChaosClothAsset/SimulationVelocityScaleConfigNode.h
--- a/Engine/Plugins/ChaosClothAssetEditor/Source/ChaosClothAssetDataflowNodes/Public/ChaosClothAsset/SimulationVelocityScaleConfigNode.h
+++ b/Engine/Plugins/ChaosClothAssetEditor/Source/ChaosClothAssetDataflowNodes/Public/ChaosClothAsset/SimulationVelocityScaleConfigNode.h
@@ -4,6 +4,7 @@
 
 #include "ChaosClothAsset/SimulationBaseConfigNode.h"
 #include "ChaosClothAsset/WeightedValue.h"
+#include "Chaos/SoftsSimulationSpace.h"
 #include "SimulationVelocityScaleConfigNode.generated.h"
 
 /** Velocity scale properties configuration node. */
@@ -15,6 +16,13 @@ struct FChaosClothAssetSimulationVelocityScaleConfigNode : public FChaosClothAss
 
 public:
 
+	/**
+	 * All vector properties on this node (e.g., Linear Velocity Scale, Max Linear Acceleration)
+	 * will be evaluated in this space. 
+	 */
+	UPROPERTY(EditAnywhere, Category = "Velocity Scale Properties")
+	EChaosSoftsSimulationSpace VelocityScaleSpace = EChaosSoftsSimulationSpace::ReferenceBoneSpace;
+
 	/**
 	 * The amount of linear velocities sent to the local cloth space from the reference bone
 	 * (the closest bone to the root on which the cloth section has been skinned, or the root itself if the cloth isn't skinned).
