Chaos Cloth: Ensure AnimatedNormals and normals are initialized and updated correctly when LOD switching.

diff --git a/Engine/Plugins/ChaosCloth/Source/ChaosCloth/Private/ChaosCloth/ChaosClothingSimulationMesh.cpp b/Engine/Plugins/ChaosCloth/Source/ChaosCloth/Private/ChaosCloth/ChaosClothingSimulationMesh.cpp
--- a/Engine/Plugins/ChaosCloth/Source/ChaosCloth/Private/ChaosCloth/ChaosClothingSimulationMesh.cpp
+++ b/Engine/Plugins/ChaosCloth/Source/ChaosCloth/Private/ChaosCloth/ChaosClothingSimulationMesh.cpp
@@ -132,6 +132,52 @@ bool FClothingSimulationMesh::WrapDeformLOD(
 	return true;
 }
 
+bool FClothingSimulationMesh::WrapDeformLOD(
+	int32 PrevLODIndex,
+	int32 LODIndex,
+	const TConstArrayView<Softs::FSolverVec3>& Positions,
+	const TConstArrayView<Softs::FSolverVec3>& Normals,
+	TArrayView<Softs::FSolverVec3>& OutPositions,
+	TArrayView<Softs::FSolverVec3>& OutNormals) const
+{
+
+	TRACE_CPUPROFILER_EVENT_SCOPE(FClothingSimulationMesh_WrapDeformLOD);
+	SCOPE_CYCLE_COUNTER(STAT_ChaosClothWrapDeformMesh);
+
+	const int32 NumLODsPassed = FMath::Abs(LODIndex - PrevLODIndex);
+	if (NumLODsPassed != 1 || !IsValidLODIndex(PrevLODIndex) || !IsValidLODIndex(LODIndex))
+	{
+		return false;
+	}
+
+	const int32 NumPoints = OutPositions.Num();
+
+	const TConstArrayView<FMeshToMeshVertData> SkinData = (PrevLODIndex < LODIndex) ?
+		GetTransitionUpSkinData(LODIndex) :
+		GetTransitionDownSkinData(LODIndex);
+
+	for (int32 Index = 0; Index < NumPoints; ++Index)  // TODO: Profile for parallel for
+	{
+		const FMeshToMeshVertData& VertData = SkinData[Index];
+
+		const int32 VertIndex0 = (int32)VertData.SourceMeshVertIndices[0];  // Note: The source is uint16. Watch out for large mesh sections!
+		const int32 VertIndex1 = (int32)VertData.SourceMeshVertIndices[1];
+		const int32 VertIndex2 = (int32)VertData.SourceMeshVertIndices[2];
+
+		OutPositions[Index] =
+			Positions[VertIndex0] * (FSolverReal)VertData.PositionBaryCoordsAndDist.X + Normals[VertIndex0] * (FSolverReal)VertData.PositionBaryCoordsAndDist.W +
+			Positions[VertIndex1] * (FSolverReal)VertData.PositionBaryCoordsAndDist.Y + Normals[VertIndex1] * (FSolverReal)VertData.PositionBaryCoordsAndDist.W +
+			Positions[VertIndex2] * (FSolverReal)VertData.PositionBaryCoordsAndDist.Z + Normals[VertIndex2] * (FSolverReal)VertData.PositionBaryCoordsAndDist.W;
+
+		OutNormals[Index] =
+			(Normals[VertIndex0] * (FSolverReal)VertData.PositionBaryCoordsAndDist.X +
+			Normals[VertIndex1] * (FSolverReal)VertData.PositionBaryCoordsAndDist.Y +
+			Normals[VertIndex2] * (FSolverReal)VertData.PositionBaryCoordsAndDist.Z).GetSafeNormal();
+	}
+
+	return true;
+}
+
 // Inline function used to force the unrolling of the skinning loop, LWC: note skinning is all done in float to match the asset data type
 FORCEINLINE static void AddInfluence(FVector3f& OutPosition, FVector3f& OutNormal, const FVector3f& RefParticle, const FVector3f& RefNormal, const FMatrix44f& BoneMatrix, const float Weight)
 {
@@ -139,7 +185,7 @@ FORCEINLINE static void AddInfluence(FVector3f& OutPosition, FVector3f& OutNorma
 	OutNormal += BoneMatrix.TransformVector(RefNormal) * Weight;
 }
 
-void FClothingSimulationMesh::SkinPhysicsMesh(int32 LODIndex, const FVec3& LocalSpaceLocation, FSolverVec3* OutPositions, FSolverVec3* OutNormals) const
+void FClothingSimulationMesh::SkinPhysicsMesh(int32 LODIndex, const FVec3& LocalSpaceLocation, TArrayView<Softs::FSolverVec3>& OutPositions, TArrayView<Softs::FSolverVec3>& OutNormals) const
 {
 	TRACE_CPUPROFILER_EVENT_SCOPE(FClothingSimulationMesh_SkinPhysicsMesh);
 	SCOPE_CYCLE_COUNTER(STAT_ChaosClothSkinPhysicsMesh);
@@ -154,6 +200,8 @@ void FClothingSimulationMesh::SkinPhysicsMesh(int32 LODIndex, const FVec3& Local
 
 	check(IsValidLODIndex(LODIndex));
 	const uint32 NumPoints = GetNumPoints(LODIndex);
+	check(NumPoints == OutPositions.Num());
+	check(NumPoints == OutNormals.Num());
 	const TConstArrayView<FClothVertBoneData> BoneData = GetBoneData(LODIndex);
 	const TConstArrayView<FVector3f> Positions = GetPositions(LODIndex);
 	const TConstArrayView<FVector3f> Normals = GetNormals(LODIndex);
@@ -162,8 +210,8 @@ void FClothingSimulationMesh::SkinPhysicsMesh(int32 LODIndex, const FVec3& Local
 	if (bChaos_SkinPhysicsMesh_ISPC_Enabled)
 	{
 		ispc::SkinPhysicsMesh(
-			(ispc::FVector3f*)OutPositions,
-			(ispc::FVector3f*)OutNormals,
+			(ispc::FVector3f*)OutPositions.GetData(),
+			(ispc::FVector3f*)OutNormals.GetData(),
 			(ispc::FVector3f*)Positions.GetData(),
 			(ispc::FVector3f*)Normals.GetData(),
 			(ispc::FClothVertBoneData*)BoneData.GetData(),
@@ -233,8 +281,8 @@ void FClothingSimulationMesh::Update(
 
 	// Skin current LOD positions
 	const FVec3& LocalSpaceLocation = Solver->GetLocalSpaceLocation();
-	FSolverVec3* const OutPositions = Solver->GetAnimationPositions(ParticleRangeId);
-	FSolverVec3* const OutNormals = Solver->GetAnimationNormals(ParticleRangeId);
+	TArrayView<FSolverVec3> OutPositions = Solver->GetAnimationPositionsView(ParticleRangeId);
+	TArrayView<FSolverVec3> OutNormals = Solver->GetAnimationNormalsView(ParticleRangeId);
 	
 	SkinPhysicsMesh(LODIndex, LocalSpaceLocation, OutPositions, OutNormals);
 
@@ -242,20 +290,19 @@ void FClothingSimulationMesh::Update(
 	if (LODIndex != PrevLODIndex)
 	{
 		// TODO: Using the more accurate skinning method here would require double buffering the context at the skeletal mesh level
-		const FSolverVec3* const SrcWrapNormals = Solver->GetAnimationNormals(PrevParticleRangeId);  // No need to keep an old normals array around, since the LOD has just changed
-		const FSolverVec3* const SrcWrapPositions = Solver->GetOldAnimationPositions(PrevParticleRangeId);
-		FSolverVec3* const OutOldPositions = Solver->GetOldAnimationPositions(ParticleRangeId);
+		const TConstArrayView<FSolverVec3> SrcWrapPositions = Solver->GetOldAnimationPositionsView(PrevParticleRangeId);
+		const TConstArrayView<FSolverVec3> SrcWrapNormals = Solver->GetOldAnimationNormalsView(PrevParticleRangeId);
+		TArrayView<FSolverVec3> OutOldPositions = Solver->GetOldAnimationPositionsView(ParticleRangeId);
+		TArrayView<FSolverVec3> OutOldNormals = Solver->GetOldAnimationNormalsView(ParticleRangeId);
 
-		const bool bValidWrap = WrapDeformLOD(PrevLODIndex, LODIndex, SrcWrapNormals, SrcWrapPositions, OutOldPositions);
-	
+		const bool bValidWrap = WrapDeformLOD(PrevLODIndex, LODIndex, SrcWrapPositions, SrcWrapNormals, OutOldPositions, OutOldNormals);
 		if (!bValidWrap)
 		{
 			// The previous LOD is invalid, reset old positions with the new LOD
-			const int32 NumPoints = GetNumPoints(LODIndex);
-
-			for (int32 Index = 0; Index < NumPoints; ++Index)
+			for (int32 Index = 0; Index < OutOldPositions.Num(); ++Index)
 			{
 				OutOldPositions[Index] = OutPositions[Index];
+				OutOldNormals[Index] = OutNormals[Index];
 			}
 		}
 	}
diff --git a/Engine/Plugins/ChaosCloth/Source/ChaosCloth/Private/ChaosCloth/ChaosClothingSimulationSolver.cpp b/Engine/Plugins/ChaosCloth/Source/ChaosCloth/Private/ChaosCloth/ChaosClothingSimulationSolver.cpp
--- a/Engine/Plugins/ChaosCloth/Source/ChaosCloth/Private/ChaosCloth/ChaosClothingSimulationSolver.cpp
+++ b/Engine/Plugins/ChaosCloth/Source/ChaosCloth/Private/ChaosCloth/ChaosClothingSimulationSolver.cpp
@@ -620,13 +620,17 @@ void FClothingSimulationSolver::ResetStartPose(int32 ParticleRangeId, int32 NumP
 	Softs::FSolverVec3* const Xs = GetParticleXs(ParticleRangeId);
 	Softs::FSolverVec3* const Vs = GetParticleVs(ParticleRangeId);
 	const Softs::FSolverVec3* const Positions = GetAnimationPositions(ParticleRangeId);
+	const Softs::FSolverVec3* const AnimNormals = GetAnimationNormals(ParticleRangeId);
 	Softs::FSolverVec3* const OldPositions = GetOldAnimationPositions(ParticleRangeId);
 	Softs::FSolverVec3* const InterpolatedPositions = GetInterpolatedAnimationPositions(ParticleRangeId);
 	Softs::FSolverVec3* const AnimationVs = GetAnimationVelocities(ParticleRangeId);
+	Softs::FSolverVec3* const OldAnimNormals = GetOldAnimationNormals(ParticleRangeId);
+	Softs::FSolverVec3* const InterpolatedNormals = GetInterpolatedAnimationNormals(ParticleRangeId);
 
 	for (int32 Index = 0; Index < NumParticles; ++Index)
 	{
 		PandInvMs[Index].P = Xs[Index] = OldPositions[Index] = InterpolatedPositions[Index] = Positions[Index];
+		OldAnimNormals[Index] = InterpolatedNormals[Index] = AnimNormals[Index];
 		Vs[Index] = AnimationVelocities[Index] = Softs::FSolverVec3(0.);
 	}
 }
diff --git a/Engine/Plugins/ChaosCloth/Source/ChaosCloth/Public/ChaosCloth/ChaosClothingSimulationMesh.h b/Engine/Plugins/ChaosCloth/Source/ChaosCloth/Public/ChaosCloth/ChaosClothingSimulationMesh.h
--- a/Engine/Plugins/ChaosCloth/Source/ChaosCloth/Public/ChaosCloth/ChaosClothingSimulationMesh.h
+++ b/Engine/Plugins/ChaosCloth/Source/ChaosCloth/Public/ChaosCloth/ChaosClothingSimulationMesh.h
@@ -166,8 +166,16 @@ namespace Chaos
 		void SkinPhysicsMesh(
 			int32 LODIndex,
 			const FVec3& LocalSpaceLocation,
-			Softs::FSolverVec3* OutPositions,
-			Softs::FSolverVec3* OutNormals) const;
+			TArrayView<Softs::FSolverVec3>& OutPositions,
+			TArrayView<Softs::FSolverVec3>& OutNormals) const;
+
+		bool WrapDeformLOD(
+			int32 PrevLODIndex,
+			int32 LODIndex,
+			const TConstArrayView<Softs::FSolverVec3>& Positions,
+			const TConstArrayView<Softs::FSolverVec3>& Normals,
+			TArrayView<Softs::FSolverVec3>& OutPositions,
+			TArrayView<Softs::FSolverVec3>& OutNormals) const;
 
 #if !UE_BUILD_SHIPPING
 		/** Debug name of the source component. */
