Manually redoing this change (CL 35811937) in Main since Main and Release-31.xx have diverged too much in a couple of files.

Fixed an incorrect scale factor when applying Acceleration clamps. This feature is only used by the Chaos Cloth Asset, so it shouldn't affect existing SKM assets.
-------------
Chaos Cloth: Added ClothGeometryScale to SkeletalMeshComponent and Cloth Component.

The cloth solver reads this value and will scale the solver geometry internally by this scale factor.

The value is currently only read when creating cloth actors, but the intention is to make this animatable in a future CL.

Added ability to set this value in the Cloth Editor from the Preview Scene Description for testing.

diff --git a/Engine/Plugins/ChaosCloth/Source/ChaosCloth/Private/ChaosCloth/ChaosClothVisualization.cpp b/Engine/Plugins/ChaosCloth/Source/ChaosCloth/Private/ChaosCloth/ChaosClothVisualization.cpp
--- a/Engine/Plugins/ChaosCloth/Source/ChaosCloth/Private/ChaosCloth/ChaosClothVisualization.cpp
+++ b/Engine/Plugins/ChaosCloth/Source/ChaosCloth/Private/ChaosCloth/ChaosClothVisualization.cpp
@@ -183,6 +183,25 @@ namespace Chaos
 		Solver = InSolver;
 	}
 
+	static FMatrix GetLocalSpaceToWorld(const ::Chaos::FClothingSimulationSolver& Solver)
+	{
+		return FMatrix(
+			FVector(Solver.GetLocalSpaceScale(), 0., 0.),
+			FVector(0., Solver.GetLocalSpaceScale(), 0.),
+			FVector(0., 0., Solver.GetLocalSpaceScale()),
+			Solver.GetLocalSpaceLocation());
+	}
+
+	static FVector GetWorldPosition(const ::Chaos::FClothingSimulationSolver& Solver, const ::Chaos::Softs::FSolverVec3& SolverPosition)
+	{
+		return FVector(SolverPosition) * Solver.GetLocalSpaceScale() + Solver.GetLocalSpaceLocation();
+	}
+
+	static FVector GetWorldVector(const ::Chaos::FClothingSimulationSolver& Solver, const ::Chaos::Softs::FSolverVec3& SolverVector)
+	{
+		return FVector(SolverVector) * Solver.GetLocalSpaceScale();
+	}
+
 	static void DrawText(FCanvas* Canvas, const FSceneView* SceneView, const FVector& Pos, const FText& Text, const FLinearColor& Color, const float Scale = 1.f)
 	{
 #if WITH_EDITOR
@@ -252,8 +271,7 @@ namespace Chaos
 			}
 		}
 
-		FMatrix LocalSimSpaceToWorld(FMatrix::Identity);
-		LocalSimSpaceToWorld.SetOrigin(Solver->GetLocalSpaceLocation());
+		const FMatrix LocalSimSpaceToWorld = GetLocalSpaceToWorld(*Solver);
 		MeshBuilder.Draw(PDI, LocalSimSpaceToWorld, FMaterials::GetInstance().GetClothMaterial()->GetRenderProxy(), SDPG_World, false, false);
 	}
 
@@ -312,9 +330,8 @@ namespace Chaos
 				MeshBuilder.AddTriangle(VertexIndex, VertexIndex + 1, VertexIndex + 2);
 			}
 		}
-		
-		FMatrix LocalSimSpaceToWorld(FMatrix::Identity);
-		LocalSimSpaceToWorld.SetOrigin(Solver->GetLocalSpaceLocation());
+
+		const FMatrix LocalSimSpaceToWorld = GetLocalSpaceToWorld(*Solver);
 		MeshBuilder.Draw(PDI, LocalSimSpaceToWorld, FMaterials::GetInstance().GetClothMaterialColor()->GetRenderProxy(), SDPG_World, false, false);
 	}
 
@@ -396,8 +413,6 @@ namespace Chaos
 		static const FLinearColor DynamicColor = FColor::White;
 		static const FLinearColor KinematicColor = FColor::Purple;
 
-		const FVec3& LocalSpaceLocation = Solver->GetLocalSpaceLocation();
-
 		for (const FClothingSimulationCloth* const Cloth : Solver->GetCloths())
 		{
 			const int32 ParticleRangeId = Cloth->GetParticleRangeId(Solver);
@@ -413,7 +428,7 @@ namespace Chaos
 
 			for (int32 Index = 0; Index < Positions.Num(); ++Index)
 			{
-				const FVector Position = LocalSpaceLocation + FVector(Positions[Index]);
+				const FVector Position = GetWorldPosition(*Solver, Positions[Index]);
 
 				const FText Text = FText::AsNumber(Offset + Index);
 				DrawText(Canvas, SceneView, Position, Text, InvMasses[Index] == (Softs::FSolverReal)0. ? KinematicColor : DynamicColor);
@@ -431,8 +446,6 @@ namespace Chaos
 		static const FLinearColor DynamicColor = FColor::White;
 		static const FLinearColor KinematicColor = FColor::Purple;
 
-		const FVec3& LocalSpaceLocation = Solver->GetLocalSpaceLocation();
-
 		for (const FClothingSimulationCloth* const Cloth : Solver->GetCloths())
 		{
 			const int32 ParticleRangeId = Cloth->GetParticleRangeId(Solver);
@@ -451,10 +464,10 @@ namespace Chaos
 			for (int32 Index = 0; Index < Elements.Num(); ++Index)
 			{
 				const TVec3<int32>& Element = Elements[Index];
-				const FVector Position = LocalSpaceLocation + (
-					FVector(Positions[Element[0] - Offset]) +
-					FVector(Positions[Element[1] - Offset]) +
-					FVector(Positions[Element[2] - Offset])) / (FReal)3.;
+				const FVector Position = GetWorldPosition(*Solver,
+					(Positions[Element[0] - Offset] +
+					Positions[Element[1] - Offset] +
+					Positions[Element[2] - Offset]) / (Softs::FSolverReal)3.f);
 
 				const bool bIsKinematic0 = (InvMasses[Element.X - Offset] == (Softs::FSolverReal)0.);
 				const bool bIsKinematic1 = (InvMasses[Element.Y - Offset] == (Softs::FSolverReal)0.);
@@ -485,8 +498,6 @@ namespace Chaos
 		NumberFormattingOptions.MinimumFractionalDigits = 2;
 		NumberFormattingOptions.MaximumFractionalDigits = 2;
 
-		const FVec3& LocalSpaceLocation = Solver->GetLocalSpaceLocation();
-
 		for (const FClothingSimulationCloth* const Cloth : Solver->GetCloths())
 		{
 			if (Cloth->GetParticleRangeId(Solver) == INDEX_NONE)
@@ -508,7 +519,7 @@ namespace Chaos
 			for (int32 Index = 0; Index < MaxDistances.Num(); ++Index)
 			{
 				const FReal MaxDistance = (FReal)MaxDistances.GetValue(Index);
-				const FVector Position = LocalSpaceLocation + FVector(Positions[Index]);
+				const FVector Position = GetWorldPosition(*Solver, Positions[Index]);
 
 				const FText Text = FText::AsNumber(MaxDistance, &NumberFormattingOptions);
 				DrawText(Canvas, SceneView, Position, Text, InvMasses[Index] == (Softs::FSolverReal)0. ? KinematicColor : DynamicColor);
@@ -566,40 +577,40 @@ namespace Chaos
 		}
 	}
 
-	static void DrawSphere(FPrimitiveDrawInterface* PDI, const ::Chaos::TSphere<::Chaos::FReal, 3>& Sphere, const FQuat& Rotation, const FVector& Position, const FLinearColor& Color)
+	static void DrawSphere(FPrimitiveDrawInterface* PDI, const ::Chaos::TSphere<::Chaos::FReal, 3>& Sphere, const FQuat& Rotation, const FVector& Position, const FReal Scale, const FLinearColor& Color)
 	{
-		const FReal Radius = Sphere.GetRadius();
-		const FVec3 Center = Position + Rotation.RotateVector(Sphere.GetCenter());
+		const FVec3 Center = Position + Rotation.RotateVector(Scale * Sphere.GetCenter());
 		if (!PDI)
 		{
+			const FReal Radius = Sphere.GetRadius() * Scale;
 			FDebugDrawQueue::GetInstance().DrawDebugSphere(Center, Radius, 12, Color.ToFColor(true), false, KINDA_SMALL_NUMBER, Chaos::Private::GetDepthPriority(), 0.f);
 			return;
 		}
 #if WITH_EDITOR
-		const FTransform Transform(Rotation, Center);
-		DrawWireSphere(PDI, Transform, Color, Radius, 12, SDPG_World, 0.0f, 0.001f, false);
+		const FTransform Transform(Rotation, Center, FVector(Scale));
+		DrawWireSphere(PDI, Transform, Color, Sphere.GetRadius(), 12, SDPG_World, 0.0f, 0.001f, false);
 #endif
 	}
 
-	static void DrawBox(FPrimitiveDrawInterface* PDI, const ::Chaos::FAABB3& Box, const FQuat& Rotation, const FVector& Position, const FLinearColor& Color)
+	static void DrawBox(FPrimitiveDrawInterface* PDI, const ::Chaos::FAABB3& Box, const FQuat& Rotation, const FVector& Position, FReal Scale,  const FLinearColor& Color)
 	{
 		if (!PDI)
 		{
-			const FVec3 Center = Position + Rotation.RotateVector(Box.GetCenter());
-			FDebugDrawQueue::GetInstance().DrawDebugBox(Center, Box.Extents() * 0.5f, Rotation, Color.ToFColor(true), false, KINDA_SMALL_NUMBER, Chaos::Private::GetDepthPriority(), 0.f);
+			const FVec3 Center = Position + Rotation.RotateVector(Scale * Box.GetCenter());
+			FDebugDrawQueue::GetInstance().DrawDebugBox(Center, Scale * Box.Extents() * 0.5f, Rotation, Color.ToFColor(true), false, KINDA_SMALL_NUMBER, Chaos::Private::GetDepthPriority(), 0.f);
 			return;
 		}
 #if WITH_EDITOR
-		const FMatrix BoxToWorld = FTransform(Rotation, Position).ToMatrixNoScale();
+		const FMatrix BoxToWorld = FTransform(Rotation, Position, FVector(Scale)).ToMatrixWithScale();
 		DrawWireBox(PDI, BoxToWorld, FBox(Box.Min(), Box.Max()), Color, SDPG_World, 0.0f, 0.001f, false);
 #endif
 	}
 
-	static void DrawCapsule(FPrimitiveDrawInterface* PDI, const ::Chaos::FCapsule& Capsule, const FQuat& Rotation, const FVector& Position, const FLinearColor& Color)
+	static void DrawCapsule(FPrimitiveDrawInterface* PDI, const ::Chaos::FCapsule& Capsule, const FQuat& Rotation, const FVector& Position, FReal Scale, const FLinearColor& Color)
 	{
-		const FReal Radius = Capsule.GetRadius();
-		const FReal HalfHeight = Capsule.GetHeight() * 0.5f + Radius;
-		const FVec3 Center = Position + Rotation.RotateVector(Capsule.GetCenter());
+		const FReal Radius = Scale * Capsule.GetRadius();
+		const FReal HalfHeight = Scale * Capsule.GetHeight() * 0.5f + Radius;
+		const FVec3 Center = Position + Rotation.RotateVector(Scale * Capsule.GetCenter());
 		if (!PDI)
 		{
 			const FQuat Orientation = FQuat::FindBetweenNormals(FVec3::UpVector, Capsule.GetAxis());
@@ -674,16 +685,16 @@ namespace Chaos
 		}
 	}
 
-	static void DrawTaperedCylinder(FPrimitiveDrawInterface* PDI, const ::Chaos::FTaperedCylinder& TaperedCylinder, const FQuat& Rotation, const FVector& Position, const FLinearColor& Color)
+	static void DrawTaperedCylinder(FPrimitiveDrawInterface* PDI, const ::Chaos::FTaperedCylinder& TaperedCylinder, const FQuat& Rotation, const FVector& Position, const FReal Scale, const FLinearColor& Color)
 	{
-		const FReal Radius1 = TaperedCylinder.GetRadius1();
-		const FReal Radius2 = TaperedCylinder.GetRadius2();
-		const FVector Position1 = Position + Rotation.RotateVector(TaperedCylinder.GetX1());
-		const FVector Position2 = Position + Rotation.RotateVector(TaperedCylinder.GetX2());
+		const FReal Radius1 = Scale * TaperedCylinder.GetRadius1();
+		const FReal Radius2 = Scale * TaperedCylinder.GetRadius2();
+		const FVector Position1 = Position + Rotation.RotateVector(Scale * TaperedCylinder.GetX1());
+		const FVector Position2 = Position + Rotation.RotateVector(Scale * TaperedCylinder.GetX2());
 		DrawTaperedCylinder(PDI, Position1, Position2, Radius1, Radius2, 12, Color);
 	}
 
-	static void DrawConvex(FPrimitiveDrawInterface* PDI, const ::Chaos::FConvex& Convex, const FQuat& Rotation, const FVector& Position, const FLinearColor& Color)
+	static void DrawConvex(FPrimitiveDrawInterface* PDI, const ::Chaos::FConvex& Convex, const FQuat& Rotation, const FVector& Position, const FReal Scale, const FLinearColor& Color)
 	{
 		const TArray<FConvex::FPlaneType>& Planes = Convex.GetFaces();
 		for (int32 PlaneIndex1 = 0; PlaneIndex1 < Planes.Num(); ++PlaneIndex1)
@@ -709,8 +720,8 @@ namespace Chaos
 						{
 							const FVector X1(Vertices[ParticleIndex1]);
 							const FVector X2(X);
-							const FVector Position1 = Position + Rotation.RotateVector(X1);
-							const FVector Position2 = Position + Rotation.RotateVector(X2);
+							const FVector Position1 = Position + Rotation.RotateVector(Scale * X1);
+							const FVector Position2 = Position + Rotation.RotateVector(Scale * X2);
 							DrawLine(PDI, Position1, Position2, Color);
 							break;
 						}
@@ -760,7 +771,7 @@ namespace Chaos
 		}
 	}
 
-	static void DrawSkinnedLevelSet(FPrimitiveDrawInterface* PDI, const TWeightedLatticeImplicitObject<FLevelSet>& SkinnedLevelSet, const FQuat& Rotation, const FVector& Position, const FMaterialRenderProxy* MaterialRenderProxy)
+	static void DrawSkinnedLevelSet(FPrimitiveDrawInterface* PDI, const TWeightedLatticeImplicitObject<FLevelSet>& SkinnedLevelSet, const FQuat& Rotation, const FVector& Position, const FReal Scale, const FMaterialRenderProxy* MaterialRenderProxy)
 	{
 #if WITH_EDITOR
 		if (PDI && MaterialRenderProxy)
@@ -780,7 +791,8 @@ namespace Chaos
 				MeshBuilder.AddTriangle(T[0], T[1], T[2]);
 			}
 
-			MeshBuilder.Draw(PDI, FTransform(Rotation, Position).ToMatrixWithScale(), MaterialRenderProxy, SDPG_World, false, false);
+			const FTransform LocalToWorld(Rotation, Position, FVector(Scale));
+			MeshBuilder.Draw(PDI, LocalToWorld.ToMatrixWithScale(), MaterialRenderProxy, SDPG_World, false, false);
 
 			if (Private::DrawSkinnedLattice)
 			{
@@ -791,7 +803,6 @@ namespace Chaos
 				const FColor EmptyLatticeColor = FColor::White;
 				const Chaos::TVec3<int32> CellCounts = LatticeGrid.Counts();
 
-				FTransform LocalToWorld(Rotation, Position);
 				for (int32 I = 0; I < CellCounts.X; ++I)
 				{
 					for (int32 J = 0; J < CellCounts.Y; ++J)
@@ -850,8 +861,8 @@ namespace Chaos
 		const FBoxSphereBounds Bounds = Solver->CalculateBounds();
 
 		// Draw bounds
-		DrawBox(PDI, FAABB3(-Bounds.BoxExtent, Bounds.BoxExtent), FQuat::Identity, Bounds.Origin, FLinearColor(FColor::Purple));
-		DrawSphere(PDI, TSphere<FReal, 3>(FVector::ZeroVector, Bounds.SphereRadius), FQuat::Identity, Bounds.Origin, FLinearColor(FColor::Orange));
+		DrawBox(PDI, FAABB3(-Bounds.BoxExtent, Bounds.BoxExtent), FQuat::Identity, Bounds.Origin, 1., FLinearColor(FColor::Purple));
+		DrawSphere(PDI, TSphere<FReal, 3>(FVector::ZeroVector, Bounds.SphereRadius), FQuat::Identity, Bounds.Origin, 1., FLinearColor(FColor::Orange));
 
 		// Draw individual cloth bounds
 		static const FLinearColor Color = FLinearColor(FColor::Purple).Desaturate(0.5);
@@ -864,7 +875,7 @@ namespace Chaos
 			}
 
 			const FAABB3 BoundingBox = Cloth->CalculateBoundingBox(Solver);
-			DrawBox(PDI, BoundingBox, FQuat::Identity, FVector::ZeroVector, Color);  // TODO: Express bounds in local coordinates for LWC
+			DrawBox(PDI, BoundingBox, FQuat::Identity, FVector::ZeroVector, 1.f, Color);  // TODO: Express bounds in local coordinates for LWC
 		}
 	}
 
@@ -890,12 +901,11 @@ namespace Chaos
 				check(!Solver->IsLegacySolver());
 				if (ExternalForces->HasPerParticleGravity())
 				{
-					const FVec3& LocalSpaceLocation = Solver->GetLocalSpaceLocation();
 					const TConstArrayView<Softs::FSolverVec3> Positions = Solver->GetParticleXsView(ParticleRangeId);
 					for (int32 ParticleIndex = 0; ParticleIndex < Positions.Num(); ++ParticleIndex)
 					{
-						const FVector Pos0 = LocalSpaceLocation + FVector(Positions[ParticleIndex]);
-						const FVector Pos1 = Pos0 + GravityVectorLengthMultiplier * FVector(ExternalForces->GetScaledGravity(ParticleIndex));
+						const FVector Pos0 = GetWorldPosition(*Solver, Positions[ParticleIndex]);
+						const FVector Pos1 = Pos0 + GravityVectorLengthMultiplier * GetWorldVector(*Solver, ExternalForces->GetScaledGravity(ParticleIndex));
 						DrawLine(PDI, Pos0, Pos1, FLinearColor::Red);
 					}
 				}
@@ -904,7 +914,7 @@ namespace Chaos
 					const FAABB3 Bounds = Cloth->CalculateBoundingBox(Solver);
 
 					const FVector Pos0 = Bounds.Center();
-					const FVector Pos1 = Pos0 + GravityVectorLengthMultiplier * FVector(ExternalForces->GetScaledGravity(0));
+					const FVector Pos1 = Pos0 + GravityVectorLengthMultiplier * GetWorldVector(*Solver, ExternalForces->GetScaledGravity(0));
 					DrawLine(PDI, Pos0, Pos1, FLinearColor::Red);
 				}
 			}
@@ -913,7 +923,7 @@ namespace Chaos
 				const FAABB3 Bounds = Cloth->CalculateBoundingBox(Solver);
 
 				const FVector Pos0 = Bounds.Center();
-				const FVector Pos1 = Pos0 + GravityVectorLengthMultiplier * FVector(Cloth->GetGravity(Solver));
+				const FVector Pos1 = Pos0 + GravityVectorLengthMultiplier * GetWorldVector(*Solver, Cloth->GetGravity(Solver));
 				DrawLine(PDI, Pos0, Pos1, FLinearColor::Red);
 			}
 		}
@@ -939,17 +949,16 @@ namespace Chaos
 			{
 				check(!Solver->IsLegacySolver());
 				{
-					const FVec3& LocalSpaceLocation = Solver->GetLocalSpaceLocation();
 					const TConstArrayView<Softs::FSolverVec3> Positions = Solver->GetParticleXsView(ParticleRangeId);
 					const Softs::FSolverVec3& FictitousAngularVelocity = ExternalForces->GetFictitiousAngularVelocity();
 					const Softs::FSolverVec3& ReferenceSpaceLocation = ExternalForces->GetReferenceSpaceLocation();
 
 					for (int32 ParticleIndex = 0; ParticleIndex < Positions.Num(); ++ParticleIndex)
 					{
-						const FVector Pos0 = LocalSpaceLocation + FVector(Positions[ParticleIndex]);
+						const FVector Pos0 = GetWorldPosition(*Solver, Positions[ParticleIndex]);
 						const Softs::FSolverVec3 CentrifugalAccel = -Softs::FSolverVec3::CrossProduct(FictitousAngularVelocity, Softs::FSolverVec3::CrossProduct(FictitousAngularVelocity, Positions[ParticleIndex] - ReferenceSpaceLocation));
 
-						const FVector Pos1 = Pos0 + FVector(CentrifugalAccel);
+						const FVector Pos1 = Pos0 + GetWorldVector(*Solver, CentrifugalAccel);
 						DrawLine(PDI, Pos0, Pos1, FLinearColor::Red);
 					}
 				}
@@ -967,8 +976,6 @@ namespace Chaos
 		static const FLinearColor DynamicColor = FColor::White;
 		static const FLinearColor KinematicColor = FColor::Purple;
 
-		const FVec3& LocalSpaceLocation = Solver->GetLocalSpaceLocation();
-
 		for (const FClothingSimulationCloth* const Cloth : Solver->GetCloths())
 		{
 			const int32 ParticleRangeId = Cloth->GetParticleRangeId(Solver);
@@ -987,9 +994,9 @@ namespace Chaos
 			{
 				const TVec3<int32>& Element = Elements[ElementIndex];
 
-				const FVector Pos0 = LocalSpaceLocation + FVector(Positions[Element.X - Offset]); // TODO: Triangle Mesh shouldn't really be solver dependent (ie not use an offset)
-				const FVector Pos1 = LocalSpaceLocation + FVector(Positions[Element.Y - Offset]);
-				const FVector Pos2 = LocalSpaceLocation + FVector(Positions[Element.Z - Offset]);
+				const FVector Pos0 = GetWorldPosition(*Solver, Positions[Element.X - Offset]); // TODO: Triangle Mesh shouldn't really be solver dependent (ie not use an offset)
+				const FVector Pos1 = GetWorldPosition(*Solver, Positions[Element.Y - Offset]);
+				const FVector Pos2 = GetWorldPosition(*Solver, Positions[Element.Z - Offset]);
 
 				const bool bIsKinematic0 = (InvMasses[Element.X - Offset] == (Softs::FSolverReal)0.);
 				const bool bIsKinematic1 = (InvMasses[Element.Y - Offset] == (Softs::FSolverReal)0.);
@@ -1011,8 +1018,6 @@ namespace Chaos
 
 		static const FLinearColor KinematicColor = FColor::Purple;
 
-		const FVec3& LocalSpaceLocation = Solver->GetLocalSpaceLocation();
-
 		for (const FClothingSimulationCloth* const Cloth : Solver->GetCloths())
 		{
 			const int32 ParticleRangeId = Cloth->GetParticleRangeId(Solver);
@@ -1030,9 +1035,9 @@ namespace Chaos
 			{
 				const TVec3<int32>& Element = Elements[ElementIndex];
 
-				const FVector Pos0 = LocalSpaceLocation + FVector(Positions[Element.X - Offset]); // TODO: Triangle Mesh shouldn't really be solver dependent (ie not use an offset)
-				const FVector Pos1 = LocalSpaceLocation + FVector(Positions[Element.Y - Offset]);
-				const FVector Pos2 = LocalSpaceLocation + FVector(Positions[Element.Z - Offset]);
+				const FVector Pos0 = GetWorldPosition(*Solver, Positions[Element.X - Offset]); // TODO: Triangle Mesh shouldn't really be solver dependent (ie not use an offset)
+				const FVector Pos1 = GetWorldPosition(*Solver, Positions[Element.Y - Offset]);
+				const FVector Pos2 = GetWorldPosition(*Solver, Positions[Element.Z - Offset]);
 
 				DrawLine(PDI, Pos0, Pos1, KinematicColor);
 				DrawLine(PDI, Pos1, Pos2, KinematicColor);
@@ -1077,8 +1082,6 @@ namespace Chaos
 		static const FLinearColor OpenedEdgeColor = FColor::Emerald;
 		static const FLinearColor ClosedEdgeColor = FColor::White;
 
-		const FVec3& LocalSpaceLocation = Solver->GetLocalSpaceLocation();
-
 		for (const FClothingSimulationCloth* const Cloth : Solver->GetCloths())
 		{
 			const int32 ParticleRangeId = Cloth->GetParticleRangeId(Solver);
@@ -1100,8 +1103,8 @@ namespace Chaos
 				const FUintVector2& Edge = EdgeToTriangles.Key;
 				const TArray<uint32>& Triangles = EdgeToTriangles.Value;
 
-				const FVector Pos0 = LocalSpaceLocation + FVector(Positions[Edge[0] - Offset]); // TODO: Triangle Mesh shouldn't really be solver dependent (ie not use an offset)
-				const FVector Pos1 = LocalSpaceLocation + FVector(Positions[Edge[1] - Offset]);
+				const FVector Pos0 = GetWorldPosition(*Solver, Positions[Edge[0] - Offset]); // TODO: Triangle Mesh shouldn't really be solver dependent (ie not use an offset)
+				const FVector Pos1 = GetWorldPosition(*Solver, Positions[Edge[1] - Offset]);
 				const FLinearColor& Color = (Triangles.Num() > 1) ? ClosedEdgeColor : OpenedEdgeColor;
 
 				DrawLine(PDI, Pos0, Pos1, Color);
@@ -1116,7 +1119,6 @@ namespace Chaos
 			return;
 		}
 
-		const FVec3& LocalSpaceLocation = Solver->GetLocalSpaceLocation();
 		for (const FClothingSimulationCloth* const Cloth : Solver->GetCloths())
 		{
 			const int32 ParticleRangeId = Cloth->GetParticleRangeId(Solver);
@@ -1138,9 +1140,9 @@ namespace Chaos
 				static const FLinearColor KinematicColor = FColor::Purple;
 				for (const TVec3<int32>& Element : CoarseMesh.GetElements())
 				{
-					const FVector Pos0 = LocalSpaceLocation + FVector(CoarsePositions[Element.X]);
-					const FVector Pos1 = LocalSpaceLocation + FVector(CoarsePositions[Element.Y]);
-					const FVector Pos2 = LocalSpaceLocation + FVector(CoarsePositions[Element.Z]);
+					const FVector Pos0 = GetWorldPosition(*Solver, CoarsePositions[Element.X]);
+					const FVector Pos1 = GetWorldPosition(*Solver, CoarsePositions[Element.Y]);
+					const FVector Pos2 = GetWorldPosition(*Solver, CoarsePositions[Element.Z]);
 					const bool bIsKinematic0 = (CoarseInvMasses[Element.X] == (Softs::FSolverReal)0.);
 					const bool bIsKinematic1 = (CoarseInvMasses[Element.Y] == (Softs::FSolverReal)0.);
 					const bool bIsKinematic2 = (CoarseInvMasses[Element.Z] == (Softs::FSolverReal)0.);
@@ -1164,8 +1166,8 @@ namespace Chaos
 				{
 					if (InvMasses[Index] != (FSolverReal)0.f && MultiResConstraints->IsConstraintActive(Index))
 					{
-						const FVector P1 = LocalSpaceLocation + FVector(Positions[Index]);
-						const FVector P2 = LocalSpaceLocation + FVector(TargetPositions[Index]);
+						const FVector P1 = GetWorldPosition(*Solver, Positions[Index]);
+						const FVector P2 = GetWorldPosition(*Solver, TargetPositions[Index]);
 
 						DrawPoint(PDI, P2, Red, nullptr, 2.f);
 						DrawLine(PDI, P1, P2, Brown);
@@ -1184,8 +1186,6 @@ namespace Chaos
 
 		static const FLinearColor KinematicColor = FColor::Magenta;
 
-		const FVec3& LocalSpaceLocation = Solver->GetLocalSpaceLocation();
-
 		for (const FClothingSimulationCloth* const Cloth : Solver->GetCloths())
 		{
 			if (Cloth->GetParticleRangeId(Solver) == INDEX_NONE)
@@ -1199,11 +1199,57 @@ namespace Chaos
 
 			for (int32 Index = 0; Index < Positions.Num(); ++Index)
 			{
-				const FVector Pos0 = LocalSpaceLocation + FVector(Positions[Index]);
-				const FVector Pos1 = Pos0 + FVector(Normals[Index]) * NormalLength;
+				const FVector Pos0 = GetWorldPosition(*Solver, Positions[Index]);
+				const FVector Pos1 = Pos0 + GetWorldVector(*Solver, Normals[Index]) * NormalLength;
+
+				DrawLine(PDI, Pos0, Pos1, KinematicColor);
+			}
+		}
+	}
+
+	void FClothVisualizationNoGC::DrawAnimVelocities(FPrimitiveDrawInterface* PDI) const
+	{
+		if (!Solver)
+		{
+			return;
+		}
+
+		static const FLinearColor KinematicColor = FLinearColor::Black;
+		static const FLinearColor KinematicMeshColor = (FLinearColor::Black + FLinearColor(FColor::Purple)) * .5;
+
+		for (const FClothingSimulationCloth* const Cloth : Solver->GetCloths())
+		{
+			const int32 ParticleRangeId = Cloth->GetParticleRangeId(Solver);
+			if (ParticleRangeId == INDEX_NONE)
+			{
+				continue;
+			}
 
+			const TConstArrayView<Softs::FSolverVec3> Positions = Cloth->GetAnimationPositions(Solver);
+			const TConstArrayView<Softs::FSolverVec3> OldPositions = Cloth->GetOldAnimationPositions(Solver);
+			checkSlow(OldPositions.Num() == Positions.Num());
+
+			for (int32 Index = 0; Index < Positions.Num(); ++Index)
+			{
+				const FVector Pos0 = GetWorldPosition(*Solver, Positions[Index]);
+				const FVector Pos1 = GetWorldPosition(*Solver, OldPositions[Index]);
 				DrawLine(PDI, Pos0, Pos1, KinematicColor);
 			}
+
+			const int32 Offset = Solver->IsLegacySolver() ? ParticleRangeId : 0;
+			const TConstArrayView<TVec3<int32>> Elements = Cloth->GetTriangleMesh(Solver).GetElements();
+			for (int32 ElementIndex = 0; ElementIndex < Elements.Num(); ++ElementIndex)
+			{
+				const TVec3<int32>& Element = Elements[ElementIndex];
+
+				const FVector Pos0 = GetWorldPosition(*Solver, OldPositions[Element.X - Offset]); // TODO: Triangle Mesh shouldn't really be solver dependent (ie not use an offset)
+				const FVector Pos1 = GetWorldPosition(*Solver, OldPositions[Element.Y - Offset]);
+				const FVector Pos2 = GetWorldPosition(*Solver, OldPositions[Element.Z - Offset]);
+
+				DrawLine(PDI, Pos0, Pos1, KinematicMeshColor);
+				DrawLine(PDI, Pos1, Pos2, KinematicMeshColor);
+				DrawLine(PDI, Pos2, Pos0, KinematicMeshColor);
+			}
 		}
 	}
 
@@ -1217,8 +1263,6 @@ namespace Chaos
 		static const FLinearColor DynamicColor = FColor::White;
 		static const FLinearColor KinematicColor = FColor::Purple;
 
-		const FVec3& LocalSpaceLocation = Solver->GetLocalSpaceLocation();
-
 		for (const FClothingSimulationCloth* const Cloth : Solver->GetCloths())
 		{
 			if (Cloth->GetParticleRangeId(Solver) == INDEX_NONE)
@@ -1235,8 +1279,8 @@ namespace Chaos
 			for (int32 Index = 0; Index < Positions.Num(); ++Index)
 			{
 				const bool bIsKinematic = (InvMasses[Index] == (Softs::FSolverReal)0.);
-				const FVector Pos0 = LocalSpaceLocation + FVector(Positions[Index]);
-				const FVector Pos1 = Pos0 + FVector(Normals[Index]) * NormalLength;
+				const FVector Pos0 = GetWorldPosition(*Solver, Positions[Index]);
+				const FVector Pos1 = Pos0 + GetWorldVector(*Solver, Normals[Index]) * NormalLength;
 
 				DrawLine(PDI, Pos0, Pos1, bIsKinematic ? KinematicColor : DynamicColor);
 			}
@@ -1250,8 +1294,6 @@ namespace Chaos
 			return;
 		}
 
-		const FVec3& LocalSpaceLocation = Solver->GetLocalSpaceLocation();
-
 		for (const FClothingSimulationCloth* const Cloth : Solver->GetCloths())
 		{
 			if (Cloth->GetParticleRangeId(Solver) == INDEX_NONE)
@@ -1271,8 +1313,8 @@ namespace Chaos
 				constexpr FReal DefaultFPS = (FReal)60.;   // TODO: A CVAR would be nice for this
 				const bool bIsKinematic = (InvMasses[Index] == 0.f);
 
-				const FVec3 Pos0 = LocalSpaceLocation + Positions[Index];
-				const FVec3 Pos1 = Pos0 + FVec3(Velocities[Index]) / DefaultFPS;  // Velocity per frame if running at DefaultFPS
+				const FVec3 Pos0 = GetWorldPosition(*Solver, Positions[Index]);
+				const FVec3 Pos1 = Pos0 + GetWorldVector(*Solver, Velocities[Index]) / DefaultFPS;  // Velocity per frame if running at DefaultFPS
 
 				DrawLine(PDI, Pos0, Pos1, bIsKinematic ? FLinearColor::Black : FLinearColor::Yellow);
 			}
@@ -1298,8 +1340,6 @@ namespace Chaos
 					(CollisionDataType == FClothingSimulationCollider::ECollisionDataType::LODless) ? GlobalColor :
 					(CollisionDataType == FClothingSimulationCollider::ECollisionDataType::External) ? DynamicColor : LODsColor;
 
-				const FVec3& LocalSpaceLocation = Solver->GetLocalSpaceLocation();
-
 				const TConstArrayView<FImplicitObjectPtr> CollisionGeometries = Collider->GetCollisionGeometry(Solver, Cloth, CollisionDataType);
 				const TConstArrayView<Softs::FSolverVec3> Translations = Collider->GetCollisionTranslations(Solver, Cloth, CollisionDataType);
 				const TConstArrayView<Softs::FSolverRotation3> Rotations = Collider->GetCollisionRotations(Solver, Cloth, CollisionDataType);
@@ -1312,21 +1352,21 @@ namespace Chaos
 					if (const FImplicitObject* const Object = CollisionGeometries[Index].GetReference())
 					{
 						const FLinearColor Color = CollisionStatus[Index] ? CollidedColor : TypeColor;
-						const FVec3 Position = LocalSpaceLocation + FVec3(Translations[Index]);
+						const FVec3 Position = GetWorldPosition(*Solver, Translations[Index]);
 						const FRotation3 Rotation(Rotations[Index]);
 
 						switch (Object->GetType())
 						{
 						case ImplicitObjectType::Sphere:
-							DrawSphere(PDI, Object->GetObjectChecked<TSphere<FReal, 3>>(), Rotation, Position, Color);
+							DrawSphere(PDI, Object->GetObjectChecked<TSphere<FReal, 3>>(), Rotation, Position, Solver->GetLocalSpaceScale(), Color);
 							break;
 
 						case ImplicitObjectType::Box:
-							DrawBox(PDI, Object->GetObjectChecked<TBox<FReal, 3>>().BoundingBox(), Rotation, Position, Color);
+							DrawBox(PDI, Object->GetObjectChecked<TBox<FReal, 3>>().BoundingBox(), Rotation, Position, Solver->GetLocalSpaceScale(), Color);
 							break;
 
 						case ImplicitObjectType::Capsule:
-							DrawCapsule(PDI, Object->GetObjectChecked<FCapsule>(), Rotation, Position, Color);
+							DrawCapsule(PDI, Object->GetObjectChecked<FCapsule>(), Rotation, Position, Solver->GetLocalSpaceScale(), Color);
 							break;
 
 						case ImplicitObjectType::Union:  // Union only used as old style tapered capsules
@@ -1337,11 +1377,11 @@ namespace Chaos
 									switch (SubObject->GetType())
 									{
 									case ImplicitObjectType::Sphere:
-										DrawSphere(PDI, SubObject->GetObjectChecked<TSphere<FReal, 3>>(), Rotation, Position, Color);
+										DrawSphere(PDI, SubObject->GetObjectChecked<TSphere<FReal, 3>>(), Rotation, Position, Solver->GetLocalSpaceScale(), Color);
 										break;
 
 									case ImplicitObjectType::TaperedCylinder:
-										DrawTaperedCylinder(PDI, SubObject->GetObjectChecked<FTaperedCylinder>(), Rotation, Position, Color);
+										DrawTaperedCylinder(PDI, SubObject->GetObjectChecked<FTaperedCylinder>(), Rotation, Position, Solver->GetLocalSpaceScale(), Color);
 										break;
 
 									default:
@@ -1358,21 +1398,21 @@ namespace Chaos
 								const FVec3 X2 = TaperedCapsule.GetX2();
 								const FReal Radius1 = TaperedCapsule.GetRadius1();
 								const FReal Radius2 = TaperedCapsule.GetRadius2();
-								DrawSphere(PDI, TSphere<FReal, 3>(X1, Radius1), Rotation, Position, Color);
-								DrawSphere(PDI, TSphere<FReal, 3>(X2, Radius2), Rotation, Position, Color);
-								DrawTaperedCylinder(PDI, FTaperedCylinder(X1, X2, Radius1, Radius2), Rotation, Position, Color);
+								DrawSphere(PDI, TSphere<FReal, 3>(X1, Radius1), Rotation, Position, Solver->GetLocalSpaceScale(), Color);
+								DrawSphere(PDI, TSphere<FReal, 3>(X2, Radius2), Rotation, Position, Solver->GetLocalSpaceScale(), Color);
+								DrawTaperedCylinder(PDI, FTaperedCylinder(X1, X2, Radius1, Radius2), Rotation, Position, Solver->GetLocalSpaceScale(), Color);
 							}
 							break;
 
 						case ImplicitObjectType::Convex:
-							DrawConvex(PDI, Object->GetObjectChecked<FConvex>(), Rotation, Position, Color);
+							DrawConvex(PDI, Object->GetObjectChecked<FConvex>(), Rotation, Position, Solver->GetLocalSpaceScale(), Color);
 							break;
 
 						case ImplicitObjectType::Transformed: // Transformed only used for levelsets
 							if (Object->GetObjectChecked<TImplicitObjectTransformed<FReal, 3>>().GetGeometry()->GetType() == ImplicitObjectType::LevelSet)
 							{
 								const TRigidTransform<FReal, 3>& Transform = Object->GetObjectChecked<TImplicitObjectTransformed<FReal, 3>>().GetTransform();
-								const FTransform CombinedTransform = Transform * FTransform(Rotation, Position);
+								const FTransform CombinedTransform = Transform * FTransform(Rotation, Position, FVector(Solver->GetLocalSpaceScale()));
 								const FLevelSet& LevelSet = Object->GetObjectChecked<TImplicitObjectTransformed<FReal, 3>>().GetGeometry()->GetObjectChecked<FLevelSet>();
 								const FMaterialRenderProxy* MaterialRenderProxy =
 #if WITH_EDITOR
@@ -1392,7 +1432,7 @@ namespace Chaos
 #else
 									nullptr;
 #endif
-								DrawSkinnedLevelSet(PDI, WeightedLevelset, Rotation, Position, MaterialRenderProxy);
+								DrawSkinnedLevelSet(PDI, WeightedLevelset, Rotation, Position, Solver->GetLocalSpaceScale(), MaterialRenderProxy);
 
 							}
 							break;
@@ -1420,11 +1460,10 @@ namespace Chaos
 		const bool bDrawPhis = Solver->GetCollisionContacts().Num() == Solver->GetCollisionPhis().Num();
 		constexpr FReal NormalLength = (FReal)10.;
 
-		const FVec3& LocalSpaceLocation = Solver->GetLocalSpaceLocation();
 		for (int32 i = 0; i < Solver->GetCollisionContacts().Num(); ++i)
 		{
-			const FVec3 Pos0 = LocalSpaceLocation + FVec3(Solver->GetCollisionContacts()[i]);
-			const FVec3 Normal = FVec3(Solver->GetCollisionNormals()[i]);
+			const FVec3 Pos0 = GetWorldPosition(*Solver, Solver->GetCollisionContacts()[i]);
+			const FVec3 Normal = GetWorldVector(*Solver, Solver->GetCollisionNormals()[i]);
 
 			// Draw contact
 			FVec3 TangentU, TangentV;
@@ -1467,8 +1506,6 @@ namespace Chaos
 				}
 			};
 
-		const FVec3& LocalSpaceLocation = Solver->GetLocalSpaceLocation();
-
 		uint8 ColorSeed = 0;
 
 		for (const FClothingSimulationCloth* const Cloth : Solver->GetCloths())
@@ -1494,15 +1531,15 @@ namespace Chaos
 					const FLinearColor ColorLight = FLinearColor::MakeFromHSV8(ColorSeed, 160, 128);
 					const FLinearColor ColorDark = FLinearColor::MakeFromHSV8(ColorSeed, 160, 64);
 
-					const FReal BackstopRadius = BackstopConstraint->GetBackstopRadius(Index) * BackstopConstraint->GetScale();
-					const FReal BackstopDistance = BackstopConstraint->GetBackstopDistance(Index) * BackstopConstraint->GetScale();
+					const FReal BackstopRadius = BackstopConstraint->GetBackstopRadius(Index) * BackstopConstraint->GetScale() * Solver->GetLocalSpaceScale();
+					const FReal BackstopDistance = BackstopConstraint->GetBackstopDistance(Index) * BackstopConstraint->GetScale() * Solver->GetLocalSpaceScale();
 
 					const FVector AnimationNormal(AnimationNormals[Index]);
 
 					// Draw a line to show the current distance to the sphere
-					const FVector Pos0 = LocalSpaceLocation + FVector(AnimationPositions[Index]);
+					const FVector Pos0 = GetWorldPosition(*Solver, AnimationPositions[Index]);
 					const FVector Pos1 = Pos0 - (bUseLegacyBackstop ? BackstopDistance - BackstopRadius : BackstopDistance) * AnimationNormal;
-					const FVector Pos2 = LocalSpaceLocation + FVector(ParticlePositions[Index]);
+					const FVector Pos2 = GetWorldPosition(*Solver, ParticlePositions[Index]);
 					DrawLine(PDI, Pos1, Pos2, ColorLight);
 
 					// Draw the sphere
@@ -1525,8 +1562,6 @@ namespace Chaos
 			return;
 		}
 
-		const FVec3& LocalSpaceLocation = Solver->GetLocalSpaceLocation();
-
 		uint8 ColorSeed = 0;
 
 		for (const FClothingSimulationCloth* const Cloth : Solver->GetCloths())
@@ -1550,13 +1585,13 @@ namespace Chaos
 					const FLinearColor ColorLight = FLinearColor::MakeFromHSV8(ColorSeed, 160, 128);
 					const FLinearColor ColorDark = FLinearColor::MakeFromHSV8(ColorSeed, 160, 64);
 
-					const FReal BackstopRadius = BackstopConstraint->GetBackstopRadius(Index) * BackstopConstraint->GetScale();
-					const FReal BackstopDistance = BackstopConstraint->GetBackstopDistance(Index) * BackstopConstraint->GetScale();
+					const FReal BackstopRadius = BackstopConstraint->GetBackstopRadius(Index) * BackstopConstraint->GetScale() * Solver->GetLocalSpaceScale();
+					const FReal BackstopDistance = BackstopConstraint->GetBackstopDistance(Index) * BackstopConstraint->GetScale() * Solver->GetLocalSpaceScale();
 
 					const FVector AnimationNormal(AnimationNormals[Index]);
 
 					// Draw a line to the sphere boundary
-					const FVector Pos0 = LocalSpaceLocation + FVector(AnimationPositions[Index]);
+					const FVector Pos0 = GetWorldPosition(*Solver, AnimationPositions[Index]);
 					const FVector Pos1 = Pos0 - (bUseLegacyBackstop ? BackstopDistance - BackstopRadius : BackstopDistance) * AnimationNormal;
 					DrawLine(PDI, Pos0, Pos1, ColorDark);
 				}
@@ -1572,8 +1607,6 @@ namespace Chaos
 		}
 
 		// Draw max distances
-		const FVec3& LocalSpaceLocation = Solver->GetLocalSpaceLocation();
-
 		for (const FClothingSimulationCloth* const Cloth : Solver->GetCloths())
 		{
 			const int32 ParticleRangeId = Cloth->GetParticleRangeId(Solver);
@@ -1598,7 +1631,7 @@ namespace Chaos
 			for (int32 Index = 0; Index < MaxDistances.Num(); ++Index)
 			{
 				const FReal MaxDistance = (FReal)MaxDistances.GetValue(Index);
-				const FVector Position = LocalSpaceLocation + FVector(Positions[Index]);
+				const FVector Position = GetWorldPosition(*Solver, Positions[Index]);
 				if (InvMasses[Index] == (Softs::FSolverReal)0.)
 				{
 #if WITH_EDITOR
@@ -1609,7 +1642,7 @@ namespace Chaos
 				}
 				else
 				{
-					DrawLine(PDI, Position, Position + FVector(Normals[Index]) * MaxDistance, FLinearColor::White);
+					DrawLine(PDI, Position, Position + GetWorldVector(*Solver, Normals[Index]) * MaxDistance, FLinearColor::White);
 				}
 			}
 		}
@@ -1622,8 +1655,6 @@ namespace Chaos
 			return;
 		}
 
-		const FVec3& LocalSpaceLocation = Solver->GetLocalSpaceLocation();
-
 		for (const FClothingSimulationCloth* const Cloth : Solver->GetCloths())
 		{
 			const int32 ParticleRangeId = Cloth->GetParticleRangeId(Solver);
@@ -1650,15 +1681,15 @@ namespace Chaos
 						StiffnessOffset + AnimDriveStiffnessMultipliers[Index] * StiffnessRange :
 						StiffnessOffset;
 
-					const FVector AnimationPosition = LocalSpaceLocation + FVector(AnimationPositions[Index]);
-					const FVector ParticlePosition = LocalSpaceLocation + FVector(ParticlePositions[Index]);
+					const FVector AnimationPosition = GetWorldPosition(*Solver, AnimationPositions[Index]);
+					const FVector ParticlePosition = GetWorldPosition(*Solver, ParticlePositions[Index]);
 					DrawLine(PDI, AnimationPosition, ParticlePosition, FLinearColor(FColor::Cyan) * Stiffness);
 				}
 			}
 		}
 	}
 	template<typename SpringConstraintType, typename GetEndPointsFuncType>
-	static void DrawSpringConstraintColors(FPrimitiveDrawInterface* PDI, const ::Chaos::FVec3& LocalSpaceLocation, const SpringConstraintType* const SpringConstraints, GetEndPointsFuncType GetEndPoints)
+	static void DrawSpringConstraintColors(FPrimitiveDrawInterface* PDI, const ::Chaos::FClothingSimulationSolver& Solver, const SpringConstraintType* const SpringConstraints, GetEndPointsFuncType GetEndPoints)
 	{
 		check(SpringConstraints);
 
@@ -1684,8 +1715,8 @@ namespace Chaos
 					Softs::FSolverVec3 P1, P2;
 					GetEndPoints(ConstraintIndex, P1, P2);
 
-					const FVec3 Pos0 = FVec3(P1) + LocalSpaceLocation;
-					const FVec3 Pos1 = FVec3(P2) + LocalSpaceLocation;
+					const FVec3 Pos0 = GetWorldPosition(Solver, P1);
+					const FVec3 Pos1 = GetWorldPosition(Solver, P2);
 					DrawLine(PDI, Pos0, Pos1, DrawColor);
 				}
 			}
@@ -1698,8 +1729,8 @@ namespace Chaos
 				Softs::FSolverVec3 P1, P2;
 				GetEndPoints(ConstraintIndex, P1, P2);
 
-				const FVec3 Pos0 = FVec3(P1) + LocalSpaceLocation;
-				const FVec3 Pos1 = FVec3(P2) + LocalSpaceLocation;
+				const FVec3 Pos0 = GetWorldPosition(Solver, P1);
+				const FVec3 Pos1 = GetWorldPosition(Solver, P2);
 
 				DrawLine(PDI, Pos0, Pos1, FLinearColor::Black);
 			}
@@ -1707,16 +1738,16 @@ namespace Chaos
 	}
 
 	template<typename SpringConstraintType>
-	static void DrawSpringConstraintColors(FPrimitiveDrawInterface* PDI, const TConstArrayView<::Chaos::Softs::FSolverVec3>& Positions, const ::Chaos::FVec3& LocalSpaceLocation, const SpringConstraintType* const SpringConstraints)
+	static void DrawSpringConstraintColors(FPrimitiveDrawInterface* PDI, const TConstArrayView<::Chaos::Softs::FSolverVec3>& Positions, const ::Chaos::FClothingSimulationSolver& Solver, const SpringConstraintType* const SpringConstraints)
 	{
-		DrawSpringConstraintColors(PDI, LocalSpaceLocation, SpringConstraints, [&SpringConstraints, &Positions](const int32 ConstraintIndex, Softs::FSolverVec3& P1, Softs::FSolverVec3& P2)
+		DrawSpringConstraintColors(PDI, Solver, SpringConstraints, [&SpringConstraints, &Positions](const int32 ConstraintIndex, Softs::FSolverVec3& P1, Softs::FSolverVec3& P2)
 		{
 			P1 = Positions[SpringConstraints->GetConstraints()[ConstraintIndex][0]];
 			P2 = Positions[SpringConstraints->GetConstraints()[ConstraintIndex][1]];
 		});
 	}
 
-	static void DrawStretchBiasConstraints_ParallelGraphColor(FPrimitiveDrawInterface* PDI, const TConstArrayView<::Chaos::Softs::FSolverVec3>& Positions, const ::Chaos::FVec3& LocalSpaceLocation, const FMaterialRenderProxy* MaterialRenderProxy, const Softs::FXPBDStretchBiasElementConstraints* const SpringConstraints)
+	static void DrawStretchBiasConstraints_ParallelGraphColor(FPrimitiveDrawInterface* PDI, const TConstArrayView<::Chaos::Softs::FSolverVec3>& Positions, const ::Chaos::FClothingSimulationSolver& Solver, const FMaterialRenderProxy* MaterialRenderProxy, const Softs::FXPBDStretchBiasElementConstraints* const SpringConstraints)
 	{
 #if WITH_EDITOR
 		if (PDI && MaterialRenderProxy)
@@ -1756,14 +1787,13 @@ namespace Chaos
 				}
 			}
 
-			FMatrix LocalSimSpaceToWorld(FMatrix::Identity);
-			LocalSimSpaceToWorld.SetOrigin(LocalSpaceLocation);
+			const FMatrix LocalSimSpaceToWorld = GetLocalSpaceToWorld(Solver);
 			MeshBuilder.Draw(PDI, LocalSimSpaceToWorld, MaterialRenderProxy, SDPG_World, false, false);
 		}
 #endif
 	}
 
-	static void DrawStretchBiasConstraints_WarpWeftStretch(FPrimitiveDrawInterface* PDI, const TConstArrayView<::Chaos::Softs::FSolverVec3>& Positions, const ::Chaos::FVec3& LocalSpaceLocation, const FMaterialRenderProxy* MaterialRenderProxy, const Softs::FXPBDStretchBiasElementConstraints* const SpringConstraints)
+	static void DrawStretchBiasConstraints_WarpWeftStretch(FPrimitiveDrawInterface* PDI, const TConstArrayView<::Chaos::Softs::FSolverVec3>& Positions, const::Chaos::FClothingSimulationSolver& Solver, const FMaterialRenderProxy* MaterialRenderProxy, const Softs::FXPBDStretchBiasElementConstraints* const SpringConstraints)
 	{
 #if WITH_EDITOR
 		if (PDI && MaterialRenderProxy)
@@ -1844,14 +1874,13 @@ namespace Chaos
 				VertexIndex += 3;
 			}
 
-			FMatrix LocalSimSpaceToWorld(FMatrix::Identity);
-			LocalSimSpaceToWorld.SetOrigin(LocalSpaceLocation);
+			const FMatrix LocalSimSpaceToWorld = GetLocalSpaceToWorld(Solver);
 			MeshBuilder.Draw(PDI, LocalSimSpaceToWorld, MaterialRenderProxy, SDPG_World, false, false);
 		}
 #endif
 	}
 
-	static void DrawStretchBiasConstraints_BiasStretch(FPrimitiveDrawInterface* PDI, const TConstArrayView<::Chaos::Softs::FSolverVec3>& Positions, const ::Chaos::FVec3& LocalSpaceLocation, const FMaterialRenderProxy* MaterialRenderProxy, const Softs::FXPBDStretchBiasElementConstraints* const SpringConstraints)
+	static void DrawStretchBiasConstraints_BiasStretch(FPrimitiveDrawInterface* PDI, const TConstArrayView<::Chaos::Softs::FSolverVec3>& Positions, const::Chaos::FClothingSimulationSolver& Solver, const FMaterialRenderProxy* MaterialRenderProxy, const Softs::FXPBDStretchBiasElementConstraints* const SpringConstraints)
 	{
 
 #if WITH_EDITOR
@@ -1896,15 +1925,14 @@ namespace Chaos
 				MeshBuilder.AddTriangle(VertexIndex, VertexIndex + 1, VertexIndex + 2);
 			}
 
-			FMatrix LocalSimSpaceToWorld(FMatrix::Identity);
-			LocalSimSpaceToWorld.SetOrigin(LocalSpaceLocation);
+			const FMatrix LocalSimSpaceToWorld = GetLocalSpaceToWorld(Solver);
 			MeshBuilder.Draw(PDI, LocalSimSpaceToWorld, MaterialRenderProxy, SDPG_World, false, false);
 		}
 #endif
 	}
 
 	template<typename ConstraintType>
-	static void DrawEdgeAnisotropy(FPrimitiveDrawInterface* PDI, const TConstArrayView<::Chaos::Softs::FSolverVec3>& Positions, const ::Chaos::FVec3& LocalSpaceLocation, const ConstraintType* const BendingConstraints)
+	static void DrawEdgeAnisotropy(FPrimitiveDrawInterface* PDI, const TConstArrayView<::Chaos::Softs::FSolverVec3>& Positions, const ::Chaos::FClothingSimulationSolver& Solver, const ConstraintType* const BendingConstraints)
 	{
 		const auto& Constraints = BendingConstraints->GetConstraints(); // auto because this could be Vec2 or Vec4, but we always just care about first two indices
 		const TArray<Softs::FSolverVec3>& WarpWeftBiasBaseMultipliers = BendingConstraints->GetWarpWeftBiasBaseMultipliers();
@@ -1914,14 +1942,14 @@ namespace Chaos
 			const Softs::FSolverVec3& P2 = Positions[Constraints[ConstraintIndex][1]];
 			const Softs::FSolverVec3& Multiplier = WarpWeftBiasBaseMultipliers[ConstraintIndex];
 
-			const FVec3 Pos0 = FVec3(P1) + LocalSpaceLocation;
-			const FVec3 Pos1 = FVec3(P2) + LocalSpaceLocation;
+			const FVec3 Pos0 = GetWorldPosition(Solver, P1);
+			const FVec3 Pos1 = GetWorldPosition(Solver, P2);
 			DrawLine(PDI, Pos0, Pos1, FLinearColor(Multiplier[0], Multiplier[1], Multiplier[2]));
 		}
 	}
 
 	template<typename ConstraintType>
-	static void DrawAxialSpringAnisotropy(FPrimitiveDrawInterface* PDI, const TConstArrayView<::Chaos::Softs::FSolverVec3>& Positions, const ::Chaos::FVec3& LocalSpaceLocation, const ConstraintType* const AxialConstraints)
+	static void DrawAxialSpringAnisotropy(FPrimitiveDrawInterface* PDI, const TConstArrayView<::Chaos::Softs::FSolverVec3>& Positions, const ::Chaos::FClothingSimulationSolver& Solver, const ConstraintType* const AxialConstraints)
 	{
 		const TArray<TVec3<int32>>& Constraints = AxialConstraints->GetConstraints();
 		const TArray<Softs::FSolverVec3>& WarpWeftBiasBaseMultipliers = AxialConstraints->GetWarpWeftBiasBaseMultipliers();
@@ -1934,16 +1962,16 @@ namespace Chaos
 			const Softs::FSolverVec3 P = Barys[ConstraintIndex] * P2 + ((FSolverReal)1. - Barys[ConstraintIndex]) * P3;
 			const Softs::FSolverVec3& Multiplier = WarpWeftBiasBaseMultipliers[ConstraintIndex];
 
-			const FVec3 Pos0 = FVec3(P1) + LocalSpaceLocation;
-			const FVec3 Pos1 = FVec3(P) + LocalSpaceLocation;
+			const FVec3 Pos0 = GetWorldPosition(Solver, P1);
+			const FVec3 Pos1 = GetWorldPosition(Solver, P);
 			DrawLine(PDI, Pos0, Pos1, FLinearColor(Multiplier[0], Multiplier[1], Multiplier[2]));
 		}
 	}
 
 	template<typename SpringConstraintType>
-	static void DrawAxialSpringConstraintColors(FPrimitiveDrawInterface* PDI, const TConstArrayView<::Chaos::Softs::FSolverVec3>& Positions, const ::Chaos::FVec3& LocalSpaceLocation, const SpringConstraintType* const SpringConstraints)
+	static void DrawAxialSpringConstraintColors(FPrimitiveDrawInterface* PDI, const TConstArrayView<::Chaos::Softs::FSolverVec3>& Positions, const ::Chaos::FClothingSimulationSolver& Solver, const SpringConstraintType* const SpringConstraints)
 	{
-		DrawSpringConstraintColors(PDI, LocalSpaceLocation, SpringConstraints, [&SpringConstraints, &Positions](const int32 ConstraintIndex, Softs::FSolverVec3& P1, Softs::FSolverVec3& P2)
+		DrawSpringConstraintColors(PDI, Solver, SpringConstraints, [&SpringConstraints, &Positions](const int32 ConstraintIndex, Softs::FSolverVec3& P1, Softs::FSolverVec3& P2)
 		{
 			P1 = Positions[SpringConstraints->GetConstraints()[ConstraintIndex][0]];
 			const Softs::FSolverReal Bary = SpringConstraints->GetBarys()[ConstraintIndex];
@@ -1958,9 +1986,6 @@ namespace Chaos
 		{
 			return;
 		}
-
-		const FVec3& LocalSpaceLocation = Solver->GetLocalSpaceLocation();
-
 		for (const FClothingSimulationCloth* const Cloth : Solver->GetCloths())
 		{
 			const int32 ParticleRangeId = Cloth->GetParticleRangeId(Solver);
@@ -1976,12 +2001,12 @@ namespace Chaos
 
 			if (const Softs::FPBDEdgeSpringConstraints* const EdgeConstraints = ClothConstraints.GetEdgeSpringConstraints().Get())
 			{
-				DrawSpringConstraintColors(PDI, Positions, LocalSpaceLocation, EdgeConstraints);
+				DrawSpringConstraintColors(PDI, Positions, *Solver, EdgeConstraints);
 			}
 
 			if (const Softs::FXPBDEdgeSpringConstraints* const EdgeConstraints = ClothConstraints.GetXEdgeSpringConstraints().Get())
 			{
-				DrawSpringConstraintColors(PDI, Positions, LocalSpaceLocation, EdgeConstraints);
+				DrawSpringConstraintColors(PDI, Positions, *Solver, EdgeConstraints);
 			}
 
 			if (const Softs::FXPBDAnisotropicSpringConstraints* const AnisoSpringConstraints = ClothConstraints.GetXAnisoSpringConstraints().Get())
@@ -1989,13 +2014,13 @@ namespace Chaos
 				switch ((Private::EAnisoSpringDrawMode)Private::AnisoSpringDrawMode)
 				{
 				case Private::EAnisoSpringDrawMode::Anisotropy:
-					DrawEdgeAnisotropy(PDI, Positions, LocalSpaceLocation, &AnisoSpringConstraints->GetEdgeConstraints());
-					DrawAxialSpringAnisotropy(PDI, Positions, LocalSpaceLocation, &AnisoSpringConstraints->GetAxialConstraints());
+					DrawEdgeAnisotropy(PDI, Positions, *Solver, &AnisoSpringConstraints->GetEdgeConstraints());
+					DrawAxialSpringAnisotropy(PDI, Positions, *Solver, &AnisoSpringConstraints->GetAxialConstraints());
 					break;
 				case Private::EAnisoSpringDrawMode::ParallelGraphColor: // fallthrough
 				default:
-					DrawSpringConstraintColors(PDI, Positions, LocalSpaceLocation, &AnisoSpringConstraints->GetEdgeConstraints());
-					DrawAxialSpringConstraintColors(PDI, Positions, LocalSpaceLocation, &AnisoSpringConstraints->GetAxialConstraints());
+					DrawSpringConstraintColors(PDI, Positions, *Solver, &AnisoSpringConstraints->GetEdgeConstraints());
+					DrawAxialSpringConstraintColors(PDI, Positions, *Solver, &AnisoSpringConstraints->GetAxialConstraints());
 					break;
 				}
 			}
@@ -2012,21 +2037,21 @@ namespace Chaos
 				{
 				case Private::EStretchBiasDrawMode::WarpStretch: // fallthrough
 				case Private::EStretchBiasDrawMode::WeftStretch:
-					DrawStretchBiasConstraints_WarpWeftStretch(PDI, Positions, LocalSpaceLocation, MaterialRenderProxy, StretchConstraints);
+					DrawStretchBiasConstraints_WarpWeftStretch(PDI, Positions, *Solver, MaterialRenderProxy, StretchConstraints);
 					break;
 				case Private::EStretchBiasDrawMode::BiasStretch:
-					DrawStretchBiasConstraints_BiasStretch(PDI, Positions, LocalSpaceLocation, MaterialRenderProxy, StretchConstraints);
+					DrawStretchBiasConstraints_BiasStretch(PDI, Positions, *Solver, MaterialRenderProxy, StretchConstraints);
 					break;
 				case Private::EStretchBiasDrawMode::ParallelGraphColor: // fallthrough
 				default:
-					DrawStretchBiasConstraints_ParallelGraphColor(PDI, Positions, LocalSpaceLocation, MaterialRenderProxy, StretchConstraints);
+					DrawStretchBiasConstraints_ParallelGraphColor(PDI, Positions, *Solver, MaterialRenderProxy, StretchConstraints);
 				}
 
 			}
 		}
 	}
 
-	static void DrawBendingElementBuckleStatus(FPrimitiveDrawInterface* PDI, const TConstArrayView<::Chaos::Softs::FSolverVec3>& Positions, const ::Chaos::FVec3& LocalSpaceLocation, const ::Chaos::Softs::FPBDBendingConstraintsBase* const BendingConstraints)
+	static void DrawBendingElementBuckleStatus(FPrimitiveDrawInterface* PDI, const TConstArrayView<::Chaos::Softs::FSolverVec3>& Positions, const ::Chaos::FClothingSimulationSolver& Solver, const ::Chaos::Softs::FPBDBendingConstraintsBase* const BendingConstraints)
 	{
 		const TArray<TVec4<int32>>& Constraints = BendingConstraints->GetConstraints();
 		const TArray<bool>& IsBuckled = BendingConstraints->GetIsBuckled();
@@ -2039,15 +2064,15 @@ namespace Chaos
 
 			const bool bIsBuckled = IsBuckled.IsValidIndex(ConstraintIndex) ? IsBuckled[ConstraintIndex] : false; // IsBuckled is empty if the simulation is paused.
 
-			const FVec3 Pos0 = FVec3(P1) + LocalSpaceLocation;
-			const FVec3 Pos1 = FVec3(P2) + LocalSpaceLocation;
+			const FVec3 Pos0 = GetWorldPosition(Solver, P1);
+			const FVec3 Pos1 = GetWorldPosition(Solver, P2);
 			DrawLine(PDI, Pos0, Pos1, bIsBuckled ? FLinearColor::Red : FLinearColor::Blue);
 		}
 	}
 
 
 
-	static void DrawBendingElementRestAngle(FPrimitiveDrawInterface* PDI, const TConstArrayView<::Chaos::Softs::FSolverVec3>& Positions, const ::Chaos::FVec3& LocalSpaceLocation, const ::Chaos::Softs::FPBDBendingConstraintsBase* const BendingConstraints)
+	static void DrawBendingElementRestAngle(FPrimitiveDrawInterface* PDI, const TConstArrayView<::Chaos::Softs::FSolverVec3>& Positions, const ::Chaos::FClothingSimulationSolver& Solver, const ::Chaos::Softs::FPBDBendingConstraintsBase* const BendingConstraints)
 	{
 		const TArray<TVec4<int32>>& Constraints = BendingConstraints->GetConstraints();
 		const TArray<Softs::FSolverReal>& RestAngles = BendingConstraints->GetRestAngles();
@@ -2058,8 +2083,8 @@ namespace Chaos
 			const Softs::FSolverReal RestAngle = RestAngles[ConstraintIndex];
 			const uint8 ColorSat = (uint8)FMath::Clamp(FMath::Abs(RestAngle) / UE_PI * 256, 0, 255);
 
-			const FVec3 Pos0 = FVec3(P1) + LocalSpaceLocation;
-			const FVec3 Pos1 = FVec3(P2) + LocalSpaceLocation;
+			const FVec3 Pos0 = GetWorldPosition(Solver, P1);
+			const FVec3 Pos1 = GetWorldPosition(Solver, P2);
 			DrawLine(PDI, Pos0, Pos1, FLinearColor::MakeFromHSV8(RestAngle > 0 ? 170 : 0, ColorSat, 255));
 		}
 	}
@@ -2071,8 +2096,6 @@ namespace Chaos
 			return;
 		}
 
-		const FVec3& LocalSpaceLocation = Solver->GetLocalSpaceLocation();
-
 		for (const FClothingSimulationCloth* const Cloth : Solver->GetCloths())
 		{
 			const int32 ParticleRangeId = Cloth->GetParticleRangeId(Solver);
@@ -2089,12 +2112,12 @@ namespace Chaos
 
 			if (const Softs::FPBDBendingSpringConstraints* const BendingConstraints = ClothConstraints.GetBendingSpringConstraints().Get())
 			{
-				DrawSpringConstraintColors(PDI, Positions, LocalSpaceLocation, BendingConstraints);
+				DrawSpringConstraintColors(PDI, Positions, *Solver, BendingConstraints);
 			}
 
 			if (const Softs::FXPBDBendingSpringConstraints* const BendingConstraints = ClothConstraints.GetXBendingSpringConstraints().Get())
 			{
-				DrawSpringConstraintColors(PDI, Positions, LocalSpaceLocation, BendingConstraints);
+				DrawSpringConstraintColors(PDI, Positions, *Solver, BendingConstraints);
 			}
 
 			if (const Softs::FPBDBendingConstraints* const BendingConstraints = ClothConstraints.GetBendingElementConstraints().Get())
@@ -2102,14 +2125,14 @@ namespace Chaos
 				switch ((Private::EBendingDrawMode)Private::BendingDrawMode)
 				{
 				case Private::EBendingDrawMode::RestAngle:
-					DrawBendingElementRestAngle(PDI, Positions, LocalSpaceLocation, BendingConstraints);
+					DrawBendingElementRestAngle(PDI, Positions, *Solver, BendingConstraints);
 					break;
 				case Private::EBendingDrawMode::ParallelGraphColor:
-					DrawSpringConstraintColors(PDI, Positions, LocalSpaceLocation, BendingConstraints);
+					DrawSpringConstraintColors(PDI, Positions, *Solver, BendingConstraints);
 					break;
 				case Private::EBendingDrawMode::BuckleStatus:
 				default:
-					DrawBendingElementBuckleStatus(PDI, Positions, LocalSpaceLocation, BendingConstraints);
+					DrawBendingElementBuckleStatus(PDI, Positions, *Solver, BendingConstraints);
 					break;
 				}
 			}
@@ -2119,14 +2142,14 @@ namespace Chaos
 				switch ((Private::EBendingDrawMode)Private::BendingDrawMode)
 				{
 				case Private::EBendingDrawMode::RestAngle:
-					DrawBendingElementRestAngle(PDI, Positions, LocalSpaceLocation, BendingConstraints);
+					DrawBendingElementRestAngle(PDI, Positions, *Solver, BendingConstraints);
 					break;
 				case Private::EBendingDrawMode::ParallelGraphColor:
-					DrawSpringConstraintColors(PDI, Positions, LocalSpaceLocation, BendingConstraints);
+					DrawSpringConstraintColors(PDI, Positions, *Solver, BendingConstraints);
 					break;
 				case Private::EBendingDrawMode::BuckleStatus:
 				default:
-					DrawBendingElementBuckleStatus(PDI, Positions, LocalSpaceLocation, BendingConstraints);
+					DrawBendingElementBuckleStatus(PDI, Positions, *Solver, BendingConstraints);
 					break;
 				}
 			}
@@ -2136,17 +2159,17 @@ namespace Chaos
 				switch ((Private::EBendingDrawMode)Private::BendingDrawMode)
 				{
 				case Private::EBendingDrawMode::RestAngle:
-					DrawBendingElementRestAngle(PDI, Positions, LocalSpaceLocation, BendingConstraints);
+					DrawBendingElementRestAngle(PDI, Positions, *Solver, BendingConstraints);
 					break;
 				case Private::EBendingDrawMode::ParallelGraphColor:
-					DrawSpringConstraintColors(PDI, Positions, LocalSpaceLocation, BendingConstraints);
+					DrawSpringConstraintColors(PDI, Positions, *Solver, BendingConstraints);
 					break;
 				case Private::EBendingDrawMode::Anisotropy:
-					DrawEdgeAnisotropy(PDI, Positions, LocalSpaceLocation, BendingConstraints);
+					DrawEdgeAnisotropy(PDI, Positions, *Solver, BendingConstraints);
 					break;
 				case Private::EBendingDrawMode::BuckleStatus:
 				default:
-					DrawBendingElementBuckleStatus(PDI, Positions, LocalSpaceLocation, BendingConstraints);
+					DrawBendingElementBuckleStatus(PDI, Positions, *Solver, BendingConstraints);
 					break;
 				}
 			}
@@ -2160,8 +2183,6 @@ namespace Chaos
 			return;
 		}
 
-		const FVec3& LocalSpaceLocation = Solver->GetLocalSpaceLocation();
-
 		auto PseudoRandomColor =
 			[](int32 NumColorRotations) -> FLinearColor
 			{
@@ -2213,8 +2234,11 @@ namespace Chaos
 						const int32 DynamicIndex = LongRangeConstraints->GetEndIndex(Tether);
 						const FReal TargetLength = LongRangeConstraints->GetTargetLength(Tether);
 
-						const FVec3 Pos0 = FVec3(Positions[KinematicIndex]) + LocalSpaceLocation;
-						const FVec3 Pos1 = FVec3(Positions[DynamicIndex]) + LocalSpaceLocation;
+						const Softs::FSolverVec3 KinematicPos = Positions[KinematicIndex];
+						const Softs::FSolverVec3 DynamicPos = Positions[DynamicIndex];
+
+						const FVec3 Pos0 = GetWorldPosition(*Solver, KinematicPos);
+						const FVec3 Pos1 = GetWorldPosition(*Solver, DynamicPos);
 
 						DrawLine(PDI, Pos0, Pos1, Color);
 #if WITH_EDITOR
@@ -2222,11 +2246,11 @@ namespace Chaos
 #else
 						DrawPoint(nullptr, Pos1, Color, nullptr);
 #endif
-						FVec3 Direction = Pos1 - Pos0;
+						Softs::FSolverVec3 Direction = DynamicPos - KinematicPos;
 						const float Length = Direction.SafeNormalize();
 						if (Length > SMALL_NUMBER)
 						{
-							const FVec3 Pos2 = Pos1 + Direction * (TargetLength - Length);
+							const FVec3 Pos2 = Pos1 + GetWorldVector(*Solver, Direction * (TargetLength - Length));
 							DrawLine(PDI, Pos1, Pos2, DarkenedColor);
 						}
 					}
@@ -2246,7 +2270,6 @@ namespace Chaos
 		}
 
 		constexpr FReal ForceLength = (FReal)10.;
-		const FVec3& LocalSpaceLocation = Solver->GetLocalSpaceLocation();
 
 		for (const FClothingSimulationCloth* const Cloth : Solver->GetCloths())
 		{
@@ -2282,17 +2305,17 @@ namespace Chaos
 			for (int32 ElementIndex = 0; ElementIndex < Elements.Num(); ++ElementIndex)
 			{
 				const TVec3<int32>& Element = Elements[ElementIndex];
-				const FVec3 Position = LocalSpaceLocation + (
-					FVec3(Positions[Element.X]) +
-					FVec3(Positions[Element.Y]) +
-					FVec3(Positions[Element.Z])) / (FReal)3.;
+				const FVec3 Position = GetWorldPosition(*Solver, 
+					(Positions[Element.X] +
+					Positions[Element.Y] +
+					Positions[Element.Z]) / (Softs::FSolverReal)3.f);
 
 				const bool bIsKinematic0 = !InvMasses[Element.X];
 				const bool bIsKinematic1 = !InvMasses[Element.Y];
 				const bool bIsKinematic2 = !InvMasses[Element.Z];
 				const bool bIsKinematic = bIsKinematic0 || bIsKinematic1 || bIsKinematic2;
 
-				const FVec3 Force = FVec3(Forces[ElementIndex]) * ForceLength;
+				const FVec3 Force = GetWorldVector(*Solver, FVec3(Forces[ElementIndex]) * ForceLength);
 				DrawLine(PDI, Position, Position + Force, bIsKinematic ? FColor::Cyan : FColor::Green);
 			}
 		}
@@ -2306,8 +2329,8 @@ namespace Chaos
 		}
 
 		// Draw local space
-		DrawCoordinateSystem(PDI, FQuat::Identity, Solver->GetLocalSpaceLocation(), 5.f, 0.25f);
-		DrawCoordinateSystem(PDI, Solver->GetLocalSpaceRotation(), Solver->GetLocalSpaceLocation(), 7.f, 0.5f);
+		DrawCoordinateSystem(PDI, FQuat::Identity, Solver->GetLocalSpaceLocation(), 5.f * Solver->GetLocalSpaceScale(), 0.25f);
+		DrawCoordinateSystem(PDI, Solver->GetLocalSpaceRotation(), Solver->GetLocalSpaceLocation(), 7.f * Solver->GetLocalSpaceScale(), 0.5f);
 
 		// Draw reference spaces
 		for (const FClothingSimulationCloth* const Cloth : Solver->GetCloths())
@@ -2318,7 +2341,7 @@ namespace Chaos
 				continue;
 			}
 			const FRigidTransform3& ReferenceSpaceTransform = Cloth->GetReferenceSpaceTransform();
-			DrawCoordinateSystem(PDI, ReferenceSpaceTransform.GetRotation(), ReferenceSpaceTransform.GetLocation());
+			DrawCoordinateSystem(PDI, ReferenceSpaceTransform.GetRotation(), ReferenceSpaceTransform.GetLocation(), 10.f * Solver->GetLocalSpaceScale());
 			check(Cloth->GetMesh());
 			DrawText(nullptr, nullptr, ReferenceSpaceTransform.GetLocation() + FVec3(2.), FText::FromName(Cloth->GetMesh()->GetReferenceBoneName()), FLinearColor::White, 2.f);
 		}
@@ -2331,8 +2354,6 @@ namespace Chaos
 			return;
 		}
 
-		const FVec3& LocalSpaceLocation = Solver->GetLocalSpaceLocation();
-
 		for (const FClothingSimulationCloth* const Cloth : Solver->GetCloths())
 		{
 			const int32 ParticleRangeId = Cloth->GetParticleRangeId(Solver);
@@ -2358,10 +2379,10 @@ namespace Chaos
 					const TVec4<int32>& Constraint = Constraints[Index];
 					const FVec3 Bary(Barys[Index]);
 
-					const FVector P = LocalSpaceLocation + FVector(Positions[Constraint[0]]);
-					const FVector P0 = LocalSpaceLocation + FVector(Positions[Constraint[1]]);
-					const FVector P1 = LocalSpaceLocation + FVector(Positions[Constraint[2]]);
-					const FVector P2 = LocalSpaceLocation + FVector(Positions[Constraint[3]]);
+					const FVector P = GetWorldPosition(*Solver, Positions[Constraint[0]]);
+					const FVector P0 = GetWorldPosition(*Solver, Positions[Constraint[1]]);
+					const FVector P1 = GetWorldPosition(*Solver, Positions[Constraint[2]]);
+					const FVector P2 = GetWorldPosition(*Solver, Positions[Constraint[3]]);
 
 					const FVector Pos0 = P0 * Bary[0] + P1 * Bary[1] + P2 * Bary[2];
 
@@ -2371,7 +2392,7 @@ namespace Chaos
 					const FVector Normal = FlipNormals[Index] ? -Triangle.GetNormal() : Triangle.GetNormal();
 
 					// Draw point to surface line (=normal)
-					const FVector Pos1 = Pos0 + Height * Normal;
+					const FVector Pos1 = Pos0 + GetWorldVector(*Solver, Height * Normal);
 					DrawPoint(PDI, Pos0, Brown, nullptr, 2.f);
 					DrawLine(PDI, Pos0, Pos1, FlipNormals[Index] ? Red : Brown);
 
@@ -2386,7 +2407,7 @@ namespace Chaos
 				const FTriangleMesh& TriangleMesh = SelfCollisionConstraints->GetTriangleMesh();
 				for (int32 Index1 : KinematicCollidingParticles)
 				{
-					const FVector P = LocalSpaceLocation + FVector(Positions[Index1]);
+					const FVector P = GetWorldPosition(*Solver, Positions[Index1]);
 
 					static const FLinearColor Orange(0.3f, 0.15f, 0.f);
 					DrawPoint(PDI, P, Orange, nullptr, 2.f);
@@ -2403,7 +2424,7 @@ namespace Chaos
 						const Softs::FSolverVec3& P3 = Positions[Index3];
 						const Softs::FSolverVec3& P4 = Positions[Index4];
 						Softs::FSolverVec3 Bary;
-						const FVector Pos1 = LocalSpaceLocation + FVector(FindClosestPointAndBaryOnTriangle(P2, P3, P4, P1, Bary));
+						const FVector Pos1 = GetWorldPosition(*Solver, FindClosestPointAndBaryOnTriangle(P2, P3, P4, P1, Bary));
 
 						static const FLinearColor LtRed(0.6f, 0.f, 0.f);
 						static const FLinearColor DkRed(0.3f, 0.f, 0.f);
@@ -2424,20 +2445,20 @@ namespace Chaos
 				for (int32 Index = 0; Index < Constraints.Num(); ++Index)
 				{
 					const TVec2<int32>& Constraint = Constraints[Index];
-					const FVector P0 = LocalSpaceLocation + FVector(Positions[Constraint[0]]);
-					const FVector P1 = LocalSpaceLocation + FVector(Positions[Constraint[1]]);
+					const FVector P0 = GetWorldPosition(*Solver, Positions[Constraint[0]]);
+					const FVector P1 = GetWorldPosition(*Solver, Positions[Constraint[1]]);
 					static const FLinearColor Brown(0.1f, 0.05f, 0.f);
 					DrawLine(PDI, P0, P1, Brown);
 				}
 
 				if (const TSet<int32>* const VertexSet = SelfCollisionSphereConstraints->GetVertexSet())
 				{
-					const FReal Radius = (FReal)SelfCollisionSphereConstraints->GetRadius();
+					const FReal Radius = (FReal)SelfCollisionSphereConstraints->GetRadius() * Solver->GetLocalSpaceScale();
 					for (const int32 Vertex : *VertexSet)
 					{
-						const FVector P0 = LocalSpaceLocation + FVector(Positions[Vertex + Offset]);
+						const FVector P0 = GetWorldPosition(*Solver, Positions[Vertex + Offset]);
 						DrawSphere(PDI, TSphere<FReal, 3>(FVector::ZeroVector, Radius), FQuat::Identity,
-							P0, FLinearColor(FColor::Orange));
+							P0, 1., FLinearColor(FColor::Orange));
 					}
 				}
 			}
@@ -2451,8 +2472,6 @@ namespace Chaos
 			return;
 		}
 
-		const FVec3& LocalSpaceLocation = Solver->GetLocalSpaceLocation();
-
 		for (const FClothingSimulationCloth* const Cloth : Solver->GetCloths())
 		{
 			const int32 ParticleRangeId = Cloth->GetParticleRangeId(Solver);
@@ -2496,15 +2515,15 @@ namespace Chaos
 					Black
 				};
 
-				auto DrawContour = [&LocalSpaceLocation, PDI, &Positions](const TArray<Softs::FPBDTriangleMeshCollisions::FBarycentricPoint>& Contour,
+				auto DrawContour = [this, PDI, &Positions](const TArray<Softs::FPBDTriangleMeshCollisions::FBarycentricPoint>& Contour,
 					const FLinearColor& ContourColor)
 				{
 					for (int32 PointIdx = 0; PointIdx < Contour.Num() - 1; ++PointIdx)
 					{
 						const Softs::FPBDTriangleMeshCollisions::FBarycentricPoint& Point0 = Contour[PointIdx];
-						const FVector EndPoint0 = LocalSpaceLocation + (1.f - Point0.Bary[0] - Point0.Bary[1]) * Positions[Point0.Vertices[0]] + Point0.Bary[0] * Positions[Point0.Vertices[1]] + Point0.Bary[1] * Positions[Point0.Vertices[2]];
+						const FVector EndPoint0 = GetWorldPosition(*Solver, (1.f - Point0.Bary[0] - Point0.Bary[1]) * Positions[Point0.Vertices[0]] + Point0.Bary[0] * Positions[Point0.Vertices[1]] + Point0.Bary[1] * Positions[Point0.Vertices[2]]);
 						const Softs::FPBDTriangleMeshCollisions::FBarycentricPoint& Point1 = Contour[PointIdx + 1];
-						const FVector EndPoint1 = LocalSpaceLocation + (1.f - Point1.Bary[0] - Point1.Bary[1]) * Positions[Point1.Vertices[0]] + Point1.Bary[0] * Positions[Point1.Vertices[1]] + Point1.Bary[1] * Positions[Point1.Vertices[2]];
+						const FVector EndPoint1 = GetWorldPosition(*Solver, (1.f - Point1.Bary[0] - Point1.Bary[1]) * Positions[Point1.Vertices[0]] + Point1.Bary[0] * Positions[Point1.Vertices[1]] + Point1.Bary[1] * Positions[Point1.Vertices[2]]);
 						DrawLine(PDI, EndPoint0, EndPoint1, ContourColor);
 						DrawPoint(PDI, EndPoint0, ContourColor, nullptr, 1.f);
 						DrawPoint(PDI, EndPoint1, ContourColor, nullptr, 1.f);
@@ -2540,7 +2559,7 @@ namespace Chaos
 							const bool bAnyBlack = (VertexGIAColors[ParticleIdx].ContourIndexBits & VertexGIAColors[ParticleIdx].ColorBits);
 							const FLinearColor& VertColor = bIsLoop ? Red : bIsBoundary ? Blue : (bAnyWhite && bAnyBlack) ? Gray : bAnyWhite ? White : Black;
 						
-							DrawPoint(PDI, LocalSpaceLocation + Positions[ParticleIdx], VertColor, nullptr, 5.f);
+							DrawPoint(PDI, GetWorldPosition(*Solver, Positions[ParticleIdx]), VertColor, nullptr, 5.f);
 						}
 					}
 				}
@@ -2556,9 +2575,9 @@ namespace Chaos
 							const bool bAnyWhite = (TriangleGIAColors[TriangleIdx].ContourIndexBits & ~TriangleGIAColors[TriangleIdx].ColorBits);
 							const bool bAnyBlack = (TriangleGIAColors[TriangleIdx].ContourIndexBits & TriangleGIAColors[TriangleIdx].ColorBits);
 							const FLinearColor& TriColor = bIsLoop ? Red : (bAnyWhite && bAnyBlack) ? Gray : bAnyWhite ? White : Black;
-							DrawLine(PDI, LocalSpaceLocation + Positions[Elements[TriangleIdx][0]], LocalSpaceLocation + Positions[Elements[TriangleIdx][1]], TriColor);
-							DrawLine(PDI, LocalSpaceLocation + Positions[Elements[TriangleIdx][1]], LocalSpaceLocation + Positions[Elements[TriangleIdx][2]], TriColor);
-							DrawLine(PDI, LocalSpaceLocation + Positions[Elements[TriangleIdx][0]], LocalSpaceLocation + Positions[Elements[TriangleIdx][2]], TriColor);
+							DrawLine(PDI, GetWorldPosition(*Solver, Positions[Elements[TriangleIdx][0]]), GetWorldPosition(*Solver, Positions[Elements[TriangleIdx][1]]), TriColor);
+							DrawLine(PDI, GetWorldPosition(*Solver, Positions[Elements[TriangleIdx][1]]), GetWorldPosition(*Solver, Positions[Elements[TriangleIdx][2]]), TriColor);
+							DrawLine(PDI, GetWorldPosition(*Solver, Positions[Elements[TriangleIdx][0]]), GetWorldPosition(*Solver, Positions[Elements[TriangleIdx][2]]), TriColor);
 						}
 					}
 				}
@@ -2574,8 +2593,8 @@ namespace Chaos
 					Intersection.GlobalGradientVector.ToDirectionAndLength(GradientDir, GradientLength);
 					const FVector Delta = FVector(GradientDir) * MaxDrawImpulse * GradientLength * FMath::InvSqrt(GradientLength * GradientLength + RegularizeEpsilonSq);
 
-					const FVector EdgeCenter = LocalSpaceLocation + .5 * (Positions[Intersection.EdgeVertices[0]] + Positions[Intersection.EdgeVertices[1]]);
-					const FVector TriCenter = LocalSpaceLocation + (Positions[Intersection.FaceVertices[0]] + Positions[Intersection.FaceVertices[1]] + Positions[Intersection.FaceVertices[2]]) / 3.;
+					const FVector EdgeCenter = GetWorldPosition(*Solver, .5 * (Positions[Intersection.EdgeVertices[0]] + Positions[Intersection.EdgeVertices[1]]));
+					const FVector TriCenter = GetWorldPosition(*Solver, (Positions[Intersection.FaceVertices[0]] + Positions[Intersection.FaceVertices[1]] + Positions[Intersection.FaceVertices[2]]) / 3.);
 
 					DrawPoint(PDI, EdgeCenter, Green, nullptr, 2.f);
 					DrawLine(PDI, EdgeCenter, EdgeCenter + Delta, Green);
@@ -2593,8 +2612,6 @@ namespace Chaos
 			return;
 		}
 
-		const FVec3& LocalSpaceLocation = Solver->GetLocalSpaceLocation();
-
 #if WITH_EDITOR
 		if (FMaterials::GetInstance().GetClothMaterialColor())
 		{
@@ -2643,8 +2660,7 @@ namespace Chaos
 					}
 				}
 			}
-			FMatrix LocalSimSpaceToWorld(FMatrix::Identity);
-			LocalSimSpaceToWorld.SetOrigin(Solver->GetLocalSpaceLocation());
+			const FMatrix LocalSimSpaceToWorld = GetLocalSpaceToWorld(*Solver);
 			MeshBuilder.Draw(PDI, LocalSimSpaceToWorld, FMaterials::GetInstance().GetClothMaterialColor()->GetRenderProxy(), SDPG_World, false, false);
 		}
 		else
@@ -2679,9 +2695,9 @@ namespace Chaos
 						const FVector Position1(Positions[Edge[0] - Offset]); // TODO: Triangle Mesh shouldn't really be solver dependent (ie not use an offset)
 						const FVector Position2(Positions[Edge[1] - Offset]);
 
-						const FReal Radius1 = (FReal)SelfCollisionConstraints->GetParticleThickness(Edge[0]);
-						const FReal Radius2 = (FReal)SelfCollisionConstraints->GetParticleThickness(Edge[1]);
-						DrawTaperedCylinder(PDI, Position1 + LocalSpaceLocation, Position2 + LocalSpaceLocation, Radius1, Radius2, 6, FLinearColor::Gray);
+						const FReal Radius1 = (FReal)SelfCollisionConstraints->GetParticleThickness(Edge[0]) * Solver->GetLocalSpaceScale();
+						const FReal Radius2 = (FReal)SelfCollisionConstraints->GetParticleThickness(Edge[1]) * Solver->GetLocalSpaceScale();
+						DrawTaperedCylinder(PDI, GetWorldPosition(*Solver, Position1), GetWorldPosition(*Solver, Position2), Radius1, Radius2, 6, FLinearColor::Gray);
 					}
 
 					for (int32 VertexIndex = 0; VertexIndex < Positions.Num(); ++VertexIndex)
@@ -2694,7 +2710,7 @@ namespace Chaos
 
 						const FVector Position1(Positions[VertexIndex]);
 						const FReal Radius1 = (FReal)SelfCollisionConstraints->GetParticleThickness(VertexIndex + Offset);
-						const FTransform Transform(Position1 + LocalSpaceLocation);
+						const FTransform Transform(FRotator::ZeroRotator, GetWorldPosition(*Solver, Position1), FVector(Solver->GetLocalSpaceScale()));
 						DrawWireSphere(PDI, Transform, FLinearColor::Gray, Radius1, 6, SDPG_World, 0.0f, 0.001f, false);
 					}
 				}
@@ -2710,8 +2726,6 @@ namespace Chaos
 		}
 
 		static const FLinearColor WireframeColor = FColor::Silver;
-		const FVec3& LocalSpaceLocation = Solver->GetLocalSpaceLocation();
-
 		for (const FClothingSimulationCloth* const Cloth : Solver->GetCloths())
 		{
 			const int32 ParticleRangeId = Cloth->GetParticleRangeId(Solver);
@@ -2732,9 +2746,9 @@ namespace Chaos
 				{
 					const TVec3<int32>& Element = Elements[ElementIndex];
 
-					const FVector Pos0 = LocalSpaceLocation + FVector(Positions[Element.X - Offset]); // TODO: Triangle Mesh shouldn't really be solver dependent (ie not use an offset)
-					const FVector Pos1 = LocalSpaceLocation + FVector(Positions[Element.Y - Offset]);
-					const FVector Pos2 = LocalSpaceLocation + FVector(Positions[Element.Z - Offset]);
+					const FVector Pos0 = GetWorldPosition(*Solver, Positions[Element.X - Offset]); // TODO: Triangle Mesh shouldn't really be solver dependent (ie not use an offset)
+					const FVector Pos1 = GetWorldPosition(*Solver, Positions[Element.Y - Offset]);
+					const FVector Pos2 = GetWorldPosition(*Solver, Positions[Element.Z - Offset]);
 
 					DrawLine(PDI, Pos0, Pos1, WireframeColor);
 					DrawLine(PDI, Pos1, Pos2, WireframeColor);
@@ -2799,9 +2813,7 @@ namespace Chaos
 			}
 		}
 
-		const FVec3& LocalSpaceLocation = Solver->GetLocalSpaceLocation();
-		FMatrix LocalSimSpaceToWorld(FMatrix::Identity);
-		LocalSimSpaceToWorld.SetOrigin(Solver->GetLocalSpaceLocation());
+		const FMatrix LocalSimSpaceToWorld = GetLocalSpaceToWorld(*Solver);
 		MeshBuilder.Draw(PDI, LocalSimSpaceToWorld, FMaterials::GetInstance().GetCollisionMaterial()->GetRenderProxy(), SDPG_World, false, false);
 	}
 #endif // #if WITH_EDITOR
diff --git a/Engine/Plugins/ChaosCloth/Source/ChaosCloth/Private/ChaosCloth/ChaosClothingSimulation.cpp b/Engine/Plugins/ChaosCloth/Source/ChaosCloth/Private/ChaosCloth/ChaosClothingSimulation.cpp
--- a/Engine/Plugins/ChaosCloth/Source/ChaosCloth/Private/ChaosCloth/ChaosClothingSimulation.cpp
+++ b/Engine/Plugins/ChaosCloth/Source/ChaosCloth/Private/ChaosCloth/ChaosClothingSimulation.cpp
@@ -71,6 +71,7 @@ namespace ClothingSimulationCVar
 	TAutoConsoleVariable<bool> DebugDrawPhysMeshWired       (TEXT("p.ChaosCloth.DebugDrawPhysMeshWired"       ), false, TEXT("Whether to debug draw the Chaos Cloth wireframe meshes"), ECVF_Cheat);
 	TAutoConsoleVariable<bool> DebugDrawAnimMeshWired       (TEXT("p.ChaosCloth.DebugDrawAnimMeshWired"       ), false, TEXT("Whether to debug draw the animated/kinematic Cloth wireframe meshes"), ECVF_Cheat);
 	TAutoConsoleVariable<bool> DebugDrawAnimNormals         (TEXT("p.ChaosCloth.DebugDrawAmimNormals"         ), false, TEXT("Whether to debug draw the animated/kinematic Cloth normals"), ECVF_Cheat);
+	TAutoConsoleVariable<bool> DebugDrawAnimVelocities      (TEXT("p.ChaosCloth.DebugDrawAnimVelocities"         ), false, TEXT("Whether to debug draw the animated/kinematic Cloth velocities"), ECVF_Cheat);
 	TAutoConsoleVariable<bool> DebugDrawPointNormals        (TEXT("p.ChaosCloth.DebugDrawPointNormals"        ), false, TEXT("Whether to debug draw the Chaos Cloth point normals"), ECVF_Cheat);
 	TAutoConsoleVariable<bool> DebugDrawPointVelocities     (TEXT("p.ChaosCloth.DebugDrawPointVelocities"     ), false, TEXT("Whether to debug draw the Chaos Cloth point velocities"), ECVF_Cheat);
 	TAutoConsoleVariable<bool> DebugDrawFaceNormals         (TEXT("p.ChaosCloth.DebugDrawFaceNormals"         ), false, TEXT("Whether to debug draw the Chaos Cloth face normals"), ECVF_Cheat);
@@ -345,6 +346,9 @@ void FClothingSimulation::CreateActor(USkeletalMeshComponent* InOwnerComponent,
 		const FClothingSimulationContext* const Context = static_cast<const FClothingSimulationContext*>(InOwnerComponent->GetClothingSimulationContext());
 		check(Context);
 		static const bool bReset = true;
+		const FReal WorldToSolverScale = bUseLocalSpaceSimulation ? Context->SolverGeometryScale : 1.;
+		const FReal LocalSpaceScale = 1. / FMath::Max(WorldToSolverScale, UE_SMALL_NUMBER);
+		Solver->SetLocalSpaceScale(LocalSpaceScale, bReset);
 		Solver->SetLocalSpaceLocation(bUseLocalSpaceSimulation ? (FVec3)Context->ComponentToWorld.GetLocation() : FVec3(0.), bReset);
 		Solver->SetLocalSpaceRotation(bUseLocalSpaceSimulation ? (FQuat)Context->ComponentToWorld.GetRotation() : FQuat::Identity);
 	}
@@ -577,6 +581,7 @@ void FClothingSimulation::Simulate(IClothingSimulationContext* InContext)
 	if (ClothingSimulationCVar::DebugDrawAnimMeshWired       .GetValueOnAnyThread()) { DebugDrawAnimMeshWired       (); }
 	if (ClothingSimulationCVar::DebugDrawPointVelocities     .GetValueOnAnyThread()) { DebugDrawPointVelocities     (); }
 	if (ClothingSimulationCVar::DebugDrawAnimNormals         .GetValueOnAnyThread()) { DebugDrawAnimNormals         (); }
+	if (ClothingSimulationCVar::DebugDrawAnimVelocities      .GetValueOnAnyThread()) { DebugDrawAnimVelocities      (); }
 	if (ClothingSimulationCVar::DebugDrawPointNormals        .GetValueOnAnyThread()) { DebugDrawPointNormals        (); }
 	if (ClothingSimulationCVar::DebugDrawCollision           .GetValueOnAnyThread()) { DebugDrawCollision           (); }
 	if (ClothingSimulationCVar::DebugDrawBackstops           .GetValueOnAnyThread()) { DebugDrawBackstops           (); }
@@ -615,6 +620,7 @@ void FClothingSimulation::GetSimulationData(
 
 	// Get the solver's local space
 	const FVec3& LocalSpaceLocation = Solver->GetLocalSpaceLocation(); // Note: Since the ReferenceSpaceTransform can be suspended with the simulation, it is important that the suspended local space location is used too in order to get the simulation data back into reference space
+	const FReal LocalSpaceScale = Solver->GetLocalSpaceScale();
 
 	// Retrieve the component transforms
 	const FClothingSimulationContext* const Context = static_cast<const FClothingSimulationContext*>(InOwnerComponent->GetClothingSimulationContext());
@@ -690,6 +696,7 @@ void FClothingSimulation::GetSimulationData(
 				(ispc::FVector3f*)Data.Normals.GetData(),
 				(ispc::FTransform&)ReferenceSpaceTransform,
 				(ispc::FVector&)LocalSpaceLocation,
+				LocalSpaceScale,
 				Data.Positions.Num());
 		}
 		else
@@ -699,7 +706,7 @@ void FClothingSimulation::GetSimulationData(
 			{
 				using FPositionsType = decltype(Data.Positions)::ElementType;
 				using FNormalsType = decltype(Data.Normals)::ElementType;
-				Data.Positions[Index] = FPositionsType(ReferenceSpaceTransform.InverseTransformPosition(FVec3(Data.Positions[Index]) + LocalSpaceLocation));  // Move into world space first
+				Data.Positions[Index] = FPositionsType(ReferenceSpaceTransform.InverseTransformPosition(LocalSpaceScale * FVec3(Data.Positions[Index]) + LocalSpaceLocation));  // Move into world space first
 				Data.Normals[Index] = FNormalsType(ReferenceSpaceTransform.InverseTransformVector(FVec3(-Data.Normals[Index])));  // Normals are inverted due to how barycentric coordinates are calculated (see GetPointBaryAndDist in ClothingMeshUtils.cpp)
 			}
 		}
@@ -728,6 +735,7 @@ FBoxSphereBounds FClothingSimulation::GetBounds(const USkeletalMeshComponent* In
 	{
 		// The component could be moving while the simulation is suspended so getting the bounds
 		// in world space isn't good enough and the bounds origin needs to be continuously updated
+		// Intentionally not applying LocalSpaceScale here because LocalSpaceScale does not match component scale.
 		Bounds = Bounds.TransformBy(FTransform((FQuat)Solver->GetLocalSpaceRotation(), (FVector)Solver->GetLocalSpaceLocation()).Inverse());
 	}
 	else if (InOwnerComponent)
diff --git a/Engine/Plugins/ChaosCloth/Source/ChaosCloth/Private/ChaosCloth/ChaosClothingSimulation.ispc b/Engine/Plugins/ChaosCloth/Source/ChaosCloth/Private/ChaosCloth/ChaosClothingSimulation.ispc
--- a/Engine/Plugins/ChaosCloth/Source/ChaosCloth/Private/ChaosCloth/ChaosClothingSimulation.ispc
+++ b/Engine/Plugins/ChaosCloth/Source/ChaosCloth/Private/ChaosCloth/ChaosClothingSimulation.ispc
@@ -10,6 +10,7 @@ export void GetClothingSimulationData(uniform FVector3f Positions[],
 									uniform FVector3f Normals[],
 									const uniform FTransform &ReferenceSpaceTransform,
 									const uniform FVector &LocalSpaceLocation,
+									const uniform double LocalSpaceScale,
 									const uniform int NumPositions)
 {
 	const uniform FVector4 RotationInverse = QuatInverse(ReferenceSpaceTransform.Rotation);
@@ -23,7 +24,7 @@ export void GetClothingSimulationData(uniform FVector3f Positions[],
 		const FVector3f NormalFloat = VectorLoad(&Normals[Index]);
 
 		const FVector Position = ConvertVector3fTo3Native(PositionFloat);
-		const FVector InputPositionVector = Position + LocalSpaceLocation;
+		const FVector InputPositionVector = LocalSpaceScale * Position + LocalSpaceLocation;
 
 		const FVector3f InputNormalVectorFloat = NormalFloat * -FLOAT_ONE;
 		const FVector InputNormalVector = ConvertVector3fTo3Native(InputNormalVectorFloat);
diff --git a/Engine/Plugins/ChaosCloth/Source/ChaosCloth/Private/ChaosCloth/ChaosClothingSimulationCloth.cpp b/Engine/Plugins/ChaosCloth/Source/ChaosCloth/Private/ChaosCloth/ChaosClothingSimulationCloth.cpp
--- a/Engine/Plugins/ChaosCloth/Source/ChaosCloth/Private/ChaosCloth/ChaosClothingSimulationCloth.cpp
+++ b/Engine/Plugins/ChaosCloth/Source/ChaosCloth/Private/ChaosCloth/ChaosClothingSimulationCloth.cpp
@@ -278,7 +278,10 @@ void FClothingSimulationCloth::FLODData::Add(FClothingSimulationSolver* Solver,
 	const int32 ParticleRangeId = SolverDatum.ParticleRangeId;
 
 	// Retrieve the component's scale
-	const Softs::FSolverReal MeshScale = Cloth->Mesh->GetScale();
+	const FReal LocalSpaceScale = Solver->GetLocalSpaceScale();
+	check(LocalSpaceScale > UE_SMALL_NUMBER);
+	const FReal LocalSpaceScaleInv = 1. / LocalSpaceScale;
+	const Softs::FSolverReal MeshScale = Cloth->Mesh->GetScale() * LocalSpaceScaleInv;
 
 	const FTriangleMesh& TriangleMesh = Solver->IsLegacySolver() ? SolverDatum.OffsetTriangleMesh : NoOffsetTriangleMesh;
 
@@ -367,7 +370,10 @@ void FClothingSimulationCloth::FLODData::Update(FClothingSimulationSolver* Solve
 	FClothConstraints& ClothConstraints = Solver->GetClothConstraints(ParticleRangeId);
 
 	check(Cloth->Config);
-	const Softs::FSolverReal MeshScale = Cloth->Mesh->GetScale();
+	const FReal LocalSpaceScale = Solver->GetLocalSpaceScale();
+	check(LocalSpaceScale > UE_SMALL_NUMBER);
+	const FReal LocalSpaceScaleInv = 1. / LocalSpaceScale;
+	const Softs::FSolverReal MeshScale = Cloth->Mesh->GetScale() * LocalSpaceScaleInv;
 	const Softs::FSolverReal MaxDistancesScale = (Softs::FSolverReal)Cloth->MaxDistancesMultiplier;
 	ClothConstraints.Update(Cloth->Config->GetProperties(SolverDatum.LODIndex), WeightMaps, VertexSets, FaceSets, FaceIntMaps, MeshScale, MaxDistancesScale, Solver->GetLocalSpaceRotation(), Cloth->ReferenceSpaceTransform.GetRotation());
 }
@@ -694,7 +700,7 @@ FAABB3 FClothingSimulationCloth::CalculateBoundingBox(const FClothingSimulationS
 	}
 
 	// Return world space bounding box
-	return FAABB3(BoundingBox).TransformedAABB(FRigidTransform3(Solver->GetLocalSpaceLocation(), FRotation3::Identity));
+	return FAABB3(BoundingBox).TransformedAABB(FTransform(FRotation3::Identity, Solver->GetLocalSpaceLocation(), FVector(Solver->GetLocalSpaceScale())));
 }
 
 int32 FClothingSimulationCloth::GetParticleRangeId(const FClothingSimulationSolver* Solver) const
@@ -926,7 +932,7 @@ void FClothingSimulationCloth::Update(FClothingSimulationSolver* Solver)
 			{
 				LODData[CoarseLODIndex]->ResetStartPose(Solver);
 			}
-			UE_LOG(LogChaosCloth, VeryVerbose, TEXT("Cloth in group Id %d Needs reset."), GroupId);
+			UE_LOG(LogChaosCloth, Verbose, TEXT("Cloth in group Id %d Needs reset."), GroupId);
 		}
 		else if (bNeedsTeleport)
 		{
@@ -937,7 +943,7 @@ void FClothingSimulationCloth::Update(FClothingSimulationSolver* Solver)
 			ReferenceSpaceAngularVelocity = FVec3(0.);
 			ReferenceSpaceVelocity = FVec3(0.);
 			bDisableFictitiousForces = true; // Disable fictitious forces. Otherwise they will be applied since AngularVelocityScale < 1.
-			UE_LOG(LogChaosCloth, VeryVerbose, TEXT("Cloth in group Id %d Needs teleport."), GroupId);
+			UE_LOG(LogChaosCloth, Verbose, TEXT("Cloth in group Id %d Needs teleport."), GroupId);
 		}
 		else
 		{
@@ -999,7 +1005,10 @@ void FClothingSimulationCloth::Update(FClothingSimulationSolver* Solver)
 			Solver->SetWindVelocity(GroupId, WindVelocity + Solver->GetWindVelocity());
 
 			// Update general solver properties
-			const Softs::FSolverReal MeshScale = Mesh->GetScale();
+			const FReal LocalSpaceScale = Solver->GetLocalSpaceScale();
+			check(LocalSpaceScale > UE_SMALL_NUMBER);
+			const FReal LocalSpaceScaleInv = 1. / LocalSpaceScale;
+			const Softs::FSolverReal MeshScale = Mesh->GetScale() * LocalSpaceScaleInv;
 
 			const FRealSingle DampingCoefficient = ConfigProperties.GetValue<float>(TEXT("DampingCoefficient"), ClothingSimulationClothDefault::DampingCoefficient);
 			const FRealSingle LocalDampingCoefficient = ConfigProperties.GetValue<float>(TEXT("LocalDampingCoefficient"));
@@ -1066,6 +1075,22 @@ TConstArrayView<Softs::FSolverVec3> FClothingSimulationCloth::GetAnimationPositi
 	return TConstArrayView<Softs::FSolverVec3>(Solver->GetAnimationPositions(GetParticleRangeId(Solver, LODIndex)), GetNumParticles(LODIndex));
 }
 
+TConstArrayView<Softs::FSolverVec3> FClothingSimulationCloth::GetOldAnimationPositions(const FClothingSimulationSolver* Solver) const
+{
+	check(Solver);
+	const int32 LODIndex = LODIndices.FindChecked(Solver);
+	check(GetParticleRangeId(Solver, LODIndex) != INDEX_NONE);
+	return TConstArrayView<Softs::FSolverVec3>(Solver->GetOldAnimationPositions(GetParticleRangeId(Solver, LODIndex)), GetNumParticles(LODIndex));
+}
+
+TConstArrayView<Softs::FSolverVec3> FClothingSimulationCloth::GetAnimationVelocities(const FClothingSimulationSolver* Solver) const
+{
+	check(Solver);
+	const int32 LODIndex = LODIndices.FindChecked(Solver);
+	check(GetParticleRangeId(Solver, LODIndex) != INDEX_NONE);
+	return TConstArrayView<Softs::FSolverVec3>(Solver->GetAnimationVelocities(GetParticleRangeId(Solver, LODIndex)), GetNumParticles(LODIndex));
+}
+
 TConstArrayView<Softs::FSolverVec3> FClothingSimulationCloth::GetAnimationNormals(const FClothingSimulationSolver* Solver) const
 {
 	check(Solver);
diff --git a/Engine/Plugins/ChaosCloth/Source/ChaosCloth/Private/ChaosCloth/ChaosClothingSimulationCollider.cpp b/Engine/Plugins/ChaosCloth/Source/ChaosCloth/Private/ChaosCloth/ChaosClothingSimulationCollider.cpp
--- a/Engine/Plugins/ChaosCloth/Source/ChaosCloth/Private/ChaosCloth/ChaosClothingSimulationCollider.cpp
+++ b/Engine/Plugins/ChaosCloth/Source/ChaosCloth/Private/ChaosCloth/ChaosClothingSimulationCollider.cpp
@@ -68,7 +68,7 @@ void FClothingSimulationCollider::FLODData::Add(
 	const FClothCollisionData& InClothCollisionData,
 	const TArray<FLevelSetCollisionData>& InLevelSetCollisionData,
 	const TArray<FSkinnedLevelSetCollisionData>& InSkinnedLevelSetCollisionData,
-	const FReal InScale,
+	const FReal InComponentScale,
 	const TArray<int32>& UsedBoneIndices)
 {
 	check(Solver);
@@ -77,6 +77,12 @@ void FClothingSimulationCollider::FLODData::Add(
 	// Keep a list of all collisions
 	ClothCollisionData = InClothCollisionData;
 
+	// Apply LocalScaleInv.
+	const FReal LocalSpaceScale = Solver->GetLocalSpaceScale();
+	check(LocalSpaceScale > UE_SMALL_NUMBER);
+	const FReal LocalSpaceScaleInv = 1. / LocalSpaceScale;
+	const FReal InScale = InComponentScale * LocalSpaceScaleInv;
+
 	// Calculate the number of geometries
 	int32 NumSpheres = ClothCollisionData.Spheres.Num();
 	TBitArray<TInlineAllocator<8>> CapsuleEnds(false, NumSpheres);  // Init on stack for 256 spheres, any number over this will be allocated
@@ -400,6 +406,11 @@ void FClothingSimulationCollider::FLODData::Update(
 
 		FTransform ComponentToLocalSpaceReal = ComponentTransform;
 		ComponentToLocalSpaceReal.AddToTranslation(-Solver->GetLocalSpaceLocation());
+		const FReal LocalSpaceScale = Solver->GetLocalSpaceScale();
+		check(LocalSpaceScale > UE_SMALL_NUMBER);
+		const FReal LocalSpaceScaleInv = 1. / LocalSpaceScale;
+		ComponentToLocalSpaceReal.MultiplyScale3D(FVector(LocalSpaceScaleInv));
+		ComponentToLocalSpaceReal.ScaleTranslation(LocalSpaceScaleInv);
 		const Softs::FSolverTransform3 ComponentToLocalSpace(ComponentToLocalSpaceReal);  // LWC, now in local space, therefore it is safe to use the solver transform type
 	
 		// Update the collision transforms
diff --git a/Engine/Plugins/ChaosCloth/Source/ChaosCloth/Private/ChaosCloth/ChaosClothingSimulationMesh.cpp b/Engine/Plugins/ChaosCloth/Source/ChaosCloth/Private/ChaosCloth/ChaosClothingSimulationMesh.cpp
--- a/Engine/Plugins/ChaosCloth/Source/ChaosCloth/Private/ChaosCloth/ChaosClothingSimulationMesh.cpp
+++ b/Engine/Plugins/ChaosCloth/Source/ChaosCloth/Private/ChaosCloth/ChaosClothingSimulationMesh.cpp
@@ -185,7 +185,7 @@ FORCEINLINE static void AddInfluence(FVector3f& OutPosition, FVector3f& OutNorma
 	OutNormal += BoneMatrix.TransformVector(RefNormal) * Weight;
 }
 
-void FClothingSimulationMesh::SkinPhysicsMesh(int32 LODIndex, const FVec3& LocalSpaceLocation, TArrayView<Softs::FSolverVec3>& OutPositions, TArrayView<Softs::FSolverVec3>& OutNormals) const
+void FClothingSimulationMesh::SkinPhysicsMesh(int32 LODIndex, const FReal LocalSpaceScale, const FVec3& LocalSpaceLocation, TArrayView<Softs::FSolverVec3>& OutPositions, TArrayView<Softs::FSolverVec3>& OutNormals) const
 {
 	TRACE_CPUPROFILER_EVENT_SCOPE(FClothingSimulationMesh_SkinPhysicsMesh);
 	SCOPE_CYCLE_COUNTER(STAT_ChaosClothSkinPhysicsMesh);
@@ -193,6 +193,10 @@ void FClothingSimulationMesh::SkinPhysicsMesh(int32 LODIndex, const FVec3& Local
 
 	FTransform ComponentToLocalSpaceReal = GetComponentToWorldTransform();
 	ComponentToLocalSpaceReal.AddToTranslation(-LocalSpaceLocation);
+	check(LocalSpaceScale > UE_SMALL_NUMBER);
+	const FReal LocalSpaceScaleInv = 1. / LocalSpaceScale;
+	ComponentToLocalSpaceReal.MultiplyScale3D(FVec3(LocalSpaceScaleInv));
+	ComponentToLocalSpaceReal.ScaleTranslation(LocalSpaceScaleInv);
 	const FTransform3f ComponentToLocalSpace(ComponentToLocalSpaceReal);  // LWC: Now in local space, therefore it is safe to use single precision which is the asset data format
 
 	const int32* const RESTRICT BoneMap = GetBoneMap().GetData();
@@ -280,11 +284,12 @@ void FClothingSimulationMesh::Update(
 	}
 
 	// Skin current LOD positions
+	const FReal LocalSpaceScale = Solver->GetLocalSpaceScale();
 	const FVec3& LocalSpaceLocation = Solver->GetLocalSpaceLocation();
 	TArrayView<FSolverVec3> OutPositions = Solver->GetAnimationPositionsView(ParticleRangeId);
 	TArrayView<FSolverVec3> OutNormals = Solver->GetAnimationNormalsView(ParticleRangeId);
 	
-	SkinPhysicsMesh(LODIndex, LocalSpaceLocation, OutPositions, OutNormals);
+	SkinPhysicsMesh(LODIndex, LocalSpaceScale, LocalSpaceLocation, OutPositions, OutNormals);
 
 	// Update old positions after LOD Switching
 	if (LODIndex != PrevLODIndex)
diff --git a/Engine/Plugins/ChaosCloth/Source/ChaosCloth/Private/ChaosCloth/ChaosClothingSimulationSolver.cpp b/Engine/Plugins/ChaosCloth/Source/ChaosCloth/Private/ChaosCloth/ChaosClothingSimulationSolver.cpp
--- a/Engine/Plugins/ChaosCloth/Source/ChaosCloth/Private/ChaosCloth/ChaosClothingSimulationSolver.cpp
+++ b/Engine/Plugins/ChaosCloth/Source/ChaosCloth/Private/ChaosCloth/ChaosClothingSimulationSolver.cpp
@@ -366,6 +366,15 @@ void FClothingSimulationSolver::SetLocalSpaceLocation(const FVec3& InLocalSpaceL
 	}
 }
 
+void FClothingSimulationSolver::SetLocalSpaceScale(FReal InLocalSpaceScale, bool bReset)
+{
+	LocalSpaceScale = InLocalSpaceScale;
+	if (bReset)
+	{
+		OldLocalSpaceScale = InLocalSpaceScale;
+	}
+}
+
 void FClothingSimulationSolver::SetCloths(TArray<FClothingSimulationCloth*>&& InCloths)
 {
 	// Remove old cloths
@@ -1102,9 +1111,12 @@ void FClothingSimulationSolver::SetReferenceVelocityScale(uint32 GroupId,
 {
 	FRigidTransform3 OldRootBoneLocalTransform = OldReferenceSpaceTransform;
 	OldRootBoneLocalTransform.AddToTranslation(-OldLocalSpaceLocation);
+	check(OldLocalSpaceScale > UE_SMALL_NUMBER);
+	const FReal OldLocalSpaceScaleInv = 1. / OldLocalSpaceScale;
+	OldRootBoneLocalTransform.ScaleTranslation(OldLocalSpaceScaleInv);
 
-	const FVec3 OldReferenceVelocity(InOutReferenceVelocity); // In local space
-	const FVec3 OldReferenceAngularVelocity(InOutReferenceAngularVelocity); // In local space
+	const FVec3 OldReferenceVelocity(InOutReferenceVelocity); // In world space
+	const FVec3 OldReferenceAngularVelocity(InOutReferenceAngularVelocity); // In world space
 
 	const FReal SolverVelocityScale = bClothSolverUseVelocityScale ? VelocityScale : (FReal)1.;
 
@@ -1122,7 +1134,7 @@ void FClothingSimulationSolver::SetReferenceVelocityScale(uint32 GroupId,
 
 	// Calculate deltas
 	const FRotation3 OldReferenceSpaceRotationInverse = OldReferenceSpaceTransform.GetRotation().Inverse();
-	const FVec3 LocalDeltaTranslation = FRigidTransform3::SubtractTranslations(ReferenceSpaceTransform, OldReferenceSpaceTransform);
+	const FVec3 WorldDeltaTranslation = FRigidTransform3::SubtractTranslations(ReferenceSpaceTransform, OldReferenceSpaceTransform);
 	const FRotation3 FullDeltaRotationQ = OldReferenceSpaceRotationInverse * ReferenceSpaceTransform.GetRotation();
 	FReal FullDeltaAngle = FullDeltaRotationQ.GetAngle();
 	const FVec3 FullAxis = FullDeltaRotationQ.GetRotationAxis();
@@ -1139,16 +1151,16 @@ void FClothingSimulationSolver::SetReferenceVelocityScale(uint32 GroupId,
 	{
 	case EChaosSoftsSimulationSpace::WorldSpace:
 		VelocityScaleSpaceToOldRefSpace = OldReferenceSpaceRotationInverse;
-		VelocityScaleSpaceDeltaTranslation = LocalDeltaTranslation;
+		VelocityScaleSpaceDeltaTranslation = WorldDeltaTranslation;
 		break;
 	case EChaosSoftsSimulationSpace::ComponentSpace:
 		VelocityScaleSpaceToOldRefSpace = OldReferenceSpaceRotationInverse * LocalSpaceRotation;
-		VelocityScaleSpaceDeltaTranslation = LocalSpaceRotation.UnrotateVector(LocalDeltaTranslation);
+		VelocityScaleSpaceDeltaTranslation = LocalSpaceRotation.UnrotateVector(WorldDeltaTranslation);
 		break;
 	case EChaosSoftsSimulationSpace::ReferenceBoneSpace:
 	default:
 		VelocityScaleSpaceToOldRefSpace = FRotation3::Identity;
-		VelocityScaleSpaceDeltaTranslation = OldReferenceSpaceRotationInverse * LocalDeltaTranslation;
+		VelocityScaleSpaceDeltaTranslation = OldReferenceSpaceRotationInverse * WorldDeltaTranslation;
 		break;
 	}
 
@@ -1168,7 +1180,7 @@ void FClothingSimulationSolver::SetReferenceVelocityScale(uint32 GroupId,
 	if (DeltaTime > 0.f)
 	{
 		// Only update old reference velocity when ticking.
-		InOutReferenceVelocity = LocalDeltaTranslation / DeltaTime;
+		InOutReferenceVelocity = WorldDeltaTranslation / DeltaTime;
 		InOutReferenceAngularVelocity = FullDeltaRotation / DeltaTime;
 	
 		FVec3 AppliedLinearVelocity = AppliedDeltaTranslation / DeltaTime;
@@ -1179,14 +1191,14 @@ void FClothingSimulationSolver::SetReferenceVelocityScale(uint32 GroupId,
 		for (int32 Index = 0; Index < 3; ++Index)
 		{
 			if (MaxLinearAcceleration[Index] != TNumericLimits<FRealSingle>::Max() && MaxLinearAcceleration[Index] >= 0.f
-				&& FMath::Abs(ScaledReferenceSpaceAccelerationTimesDt[Index]) > (FReal)(MaxLinearAcceleration[Index] * DeltaTime))
+				&& FMath::Abs(ScaledReferenceSpaceAccelerationTimesDt[Index]) > (FReal)(MaxLinearAcceleration[Index] * DeltaTime * LocalSpaceScale))
 			{
-				AppliedLinearVelocity[Index] = OldScaledReferenceSpaceVelocity[Index] + FMath::Sign(ScaledReferenceSpaceAccelerationTimesDt[Index]) * (FReal)(MaxLinearAcceleration[Index] * DeltaTime);
+				AppliedLinearVelocity[Index] = OldScaledReferenceSpaceVelocity[Index] + FMath::Sign(ScaledReferenceSpaceAccelerationTimesDt[Index]) * (FReal)(MaxLinearAcceleration[Index] * DeltaTime * LocalSpaceScale);
 			}
 			if (MaxLinearVelocity[Index] != TNumericLimits<FRealSingle>::Max() && MaxLinearVelocity[Index] >= 0.f &&
-				FMath::Abs(AppliedLinearVelocity[Index]) > (FReal)MaxLinearVelocity[Index])
+				FMath::Abs(AppliedLinearVelocity[Index]) > (FReal)MaxLinearVelocity[Index] * LocalSpaceScale)
 			{
-				AppliedLinearVelocity[Index] = FMath::Sign(AppliedLinearVelocity[Index]) * (FReal)MaxLinearVelocity[Index];
+				AppliedLinearVelocity[Index] = FMath::Sign(AppliedLinearVelocity[Index]) * (FReal)MaxLinearVelocity[Index] * LocalSpaceScale;
 			}
 
 		}
@@ -1216,7 +1228,10 @@ void FClothingSimulationSolver::SetReferenceVelocityScale(uint32 GroupId,
 
 		AppliedDeltaRotation = AppliedAngularVelocity * DeltaTime;
 	}
-	const FVec3 DeltaPosition = VelocityScaleSpaceToOldRefSpace * (VelocityScaleSpaceDeltaTranslation - AppliedDeltaTranslation);
+
+	check(LocalSpaceScale > UE_SMALL_NUMBER);
+	const FReal LocalSpaceScaleInv = 1. / LocalSpaceScale;
+	const FVec3 DeltaPosition = LocalSpaceScaleInv * VelocityScaleSpaceToOldRefSpace * (VelocityScaleSpaceDeltaTranslation - AppliedDeltaTranslation);
 	const FVec3 DeltaRotation = FullDeltaRotation - AppliedDeltaRotation;
 	const FReal DeltaAngle = DeltaRotation.Length();
 	const FVec3 Axis = DeltaAngle > UE_KINDA_SMALL_NUMBER ? DeltaRotation / DeltaAngle : FVec3::XAxisVector;
@@ -1236,7 +1251,7 @@ void FClothingSimulationSolver::SetReferenceVelocityScale(uint32 GroupId,
 	const FVec3 FictitiousAngularDisplacement = ReferenceSpaceTransform.TransformVector(DeltaRotation)
 		* AppliedFictitiousAngularScale;
 	FictitiousAngularVelocities[GroupId] = DeltaTime > (Softs::FSolverReal)0.f ? Softs::FSolverVec3(FictitiousAngularDisplacement) / DeltaTime : Softs::FSolverVec3(0.f);
-	ReferenceSpaceLocations[GroupId] = ReferenceSpaceTransform.GetLocation() - LocalSpaceLocation;
+	ReferenceSpaceLocations[GroupId] = LocalSpaceScaleInv * (ReferenceSpaceTransform.GetLocation() - LocalSpaceLocation);
 }
 
 Softs::FSolverReal FClothingSimulationSolver::SetParticleMassPerArea(Softs::FSolverParticlesRange& Particles, const FTriangleMesh& Mesh)
@@ -1579,7 +1594,8 @@ void FClothingSimulationSolver::AddExternalForces(uint32 GroupId, bool bUseLegac
 void FClothingSimulationSolver::ApplyPreSimulationTransforms()
 {
 	TRACE_CPUPROFILER_EVENT_SCOPE(FClothingSimulationSolver_ApplyPreSimulationTransforms);
-	const Softs::FSolverVec3 DeltaLocalSpaceLocation(LocalSpaceLocation - OldLocalSpaceLocation);  // LWC, note that LocalSpaceLocation is FReal based. but the delta can be stored in FSolverReal
+	check(FMath::IsNearlyEqual(OldLocalSpaceScale, LocalSpaceScale)); // Otherwise will need to apply a DeltaLocalScale.
+	const Softs::FSolverVec3 DeltaLocalSpaceLocation((LocalSpaceLocation - OldLocalSpaceLocation) / LocalSpaceScale);  // LWC, note that LocalSpaceLocation is FReal based. but the delta can be stored in FSolverReal
 
 	const FSolverReal MaxVelocitySquared = (ClothSolverMaxVelocity > 0.f) ? FMath::Square((FSolverReal)ClothSolverMaxVelocity) : TNumericLimits<FSolverReal>::Max();
 
@@ -1923,7 +1939,7 @@ void FClothingSimulationSolver::UpdateSolverField()
 					for (int32 LocalParticleIndex = 0; LocalParticleIndex < Particles.GetRangeSize(); ++LocalParticleIndex)
 					{
 						const int32 GlobalParticleIndex = LocalParticleIndex + Particles.GetOffset();
-						SamplePositions[GlobalParticleIndex] = FVector(Particles.X(LocalParticleIndex)) + LocalSpaceLocation;
+						SamplePositions[GlobalParticleIndex] = LocalSpaceScale * FVector(Particles.X(LocalParticleIndex)) + LocalSpaceLocation;
 						SampleIndices[SampleIndex] = FFieldContextIndex(GlobalParticleIndex, SampleIndex);
 						++SampleIndex;
 					}
@@ -1942,7 +1958,7 @@ void FClothingSimulationSolver::UpdateSolverField()
 			PBDEvolution->ParticlesActiveView().SequentialFor(
 				[this, &SamplePositions, &SampleIndices, &SampleIndex](Softs::FSolverParticles& Particles, int32 ParticleIndex)
 			{
-				SamplePositions[ParticleIndex] = FVector(Particles.X(ParticleIndex)) + LocalSpaceLocation;
+				SamplePositions[ParticleIndex] = LocalSpaceScale * FVector(Particles.X(ParticleIndex)) + LocalSpaceLocation;
 				SampleIndices[SampleIndex] = FFieldContextIndex(ParticleIndex, SampleIndex);
 				++SampleIndex;
 			});
@@ -2243,7 +2259,7 @@ FBoxSphereBounds FClothingSimulationSolver::CalculateBounds() const
 		}
 
 		// Update bounds with this cloth
-		return FBoxSphereBounds(LocalSpaceLocation + BoundingBox.Center(), FVector(BoundingBox.Extents() * 0.5f), FMath::Sqrt(SquaredRadius));
+		return FBoxSphereBounds(LocalSpaceLocation + LocalSpaceScale * BoundingBox.Center(), LocalSpaceScale * FVector(BoundingBox.Extents() * 0.5f), LocalSpaceScale * FMath::Sqrt(SquaredRadius));
 	}
 	else
 	{
@@ -2312,7 +2328,7 @@ FBoxSphereBounds FClothingSimulationSolver::CalculateBounds() const
 			}
 
 			// Update bounds with this cloth
-			return FBoxSphereBounds(LocalSpaceLocation + BoundingBox.Center(), FVector(BoundingBox.Extents() * 0.5f), FMath::Sqrt(SquaredRadius));
+			return FBoxSphereBounds(LocalSpaceLocation + LocalSpaceScale * BoundingBox.Center(), LocalSpaceScale * FVector(BoundingBox.Extents() * 0.5f), LocalSpaceScale * FMath::Sqrt(SquaredRadius));
 		}
 
 		return FBoxSphereBounds(LocalSpaceLocation, FVector(0.f), 0.f);
diff --git a/Engine/Plugins/ChaosCloth/Source/ChaosCloth/Public/ChaosCloth/ChaosClothVisualization.h b/Engine/Plugins/ChaosCloth/Source/ChaosCloth/Public/ChaosCloth/ChaosClothVisualization.h
--- a/Engine/Plugins/ChaosCloth/Source/ChaosCloth/Public/ChaosCloth/ChaosClothVisualization.h
+++ b/Engine/Plugins/ChaosCloth/Source/ChaosCloth/Public/ChaosCloth/ChaosClothVisualization.h
@@ -35,6 +35,7 @@ namespace Chaos
 
 		CHAOSCLOTH_API void DrawPhysMeshWired(FPrimitiveDrawInterface* PDI = nullptr) const;
 		CHAOSCLOTH_API void DrawAnimMeshWired(FPrimitiveDrawInterface* PDI = nullptr) const;
+		CHAOSCLOTH_API void DrawAnimVelocities(FPrimitiveDrawInterface* PDI = nullptr) const;
 		CHAOSCLOTH_API void DrawAnimNormals(FPrimitiveDrawInterface* PDI = nullptr, const FReal Length = 20.) const;
 		CHAOSCLOTH_API void DrawOpenEdges(FPrimitiveDrawInterface* PDI = nullptr) const;
 		CHAOSCLOTH_API void DrawPointNormals(FPrimitiveDrawInterface* PDI = nullptr, const FReal Length = 20.) const;
diff --git a/Engine/Plugins/ChaosCloth/Source/ChaosCloth/Public/ChaosCloth/ChaosClothingSimulation.h b/Engine/Plugins/ChaosCloth/Source/ChaosCloth/Public/ChaosCloth/ChaosClothingSimulation.h
--- a/Engine/Plugins/ChaosCloth/Source/ChaosCloth/Public/ChaosCloth/ChaosClothingSimulation.h
+++ b/Engine/Plugins/ChaosCloth/Source/ChaosCloth/Public/ChaosCloth/ChaosClothingSimulation.h
@@ -97,6 +97,7 @@ namespace Chaos
 		void DebugDrawPhysMeshWired(FPrimitiveDrawInterface* PDI = nullptr) const { Visualization.DrawPhysMeshWired(PDI); }
 		void DebugDrawAnimMeshWired(FPrimitiveDrawInterface* PDI = nullptr) const { Visualization.DrawAnimMeshWired(PDI); }
 		void DebugDrawAnimNormals(FPrimitiveDrawInterface* PDI = nullptr) const { Visualization.DrawAnimNormals(PDI); }
+		void DebugDrawAnimVelocities(FPrimitiveDrawInterface* PDI = nullptr) const { Visualization.DrawAnimVelocities(PDI); }
 		void DebugDrawPointNormals(FPrimitiveDrawInterface* PDI = nullptr) const { Visualization.DrawPointNormals(PDI); }
 		void DebugDrawPointVelocities(FPrimitiveDrawInterface* PDI = nullptr) const { Visualization.DrawPointVelocities(PDI); }
 		void DebugDrawCollision(FPrimitiveDrawInterface* PDI = nullptr) const { Visualization.DrawCollision(PDI); }
diff --git a/Engine/Plugins/ChaosCloth/Source/ChaosCloth/Public/ChaosCloth/ChaosClothingSimulationCloth.h b/Engine/Plugins/ChaosCloth/Source/ChaosCloth/Public/ChaosCloth/ChaosClothingSimulationCloth.h
--- a/Engine/Plugins/ChaosCloth/Source/ChaosCloth/Public/ChaosCloth/ChaosClothingSimulationCloth.h
+++ b/Engine/Plugins/ChaosCloth/Source/ChaosCloth/Public/ChaosCloth/ChaosClothingSimulationCloth.h
@@ -71,8 +71,12 @@ namespace Chaos
 		// ---- Debugging/visualization functions
 		// Return the solver's input positions for this cloth source current LOD, not thread safe, call must be done right after the solver update.
 		CHAOSCLOTH_API TConstArrayView<Softs::FSolverVec3> GetAnimationPositions(const FClothingSimulationSolver* Solver) const;
+		// Return the solver's input positions for this cloth source current LOD, not thread safe, call must be done right after the solver update.
+		CHAOSCLOTH_API TConstArrayView<Softs::FSolverVec3> GetOldAnimationPositions(const FClothingSimulationSolver* Solver) const;
 		// Return the solver's input normals for this cloth source current LOD, not thread safe, call must be done right after the solver update.
 		CHAOSCLOTH_API TConstArrayView<Softs::FSolverVec3> GetAnimationNormals(const FClothingSimulationSolver* Solver) const;
+		// Return the solver's input normals for this cloth source current LOD, not thread safe, call must be done right after the solver update.
+		CHAOSCLOTH_API TConstArrayView<Softs::FSolverVec3> GetAnimationVelocities(const FClothingSimulationSolver* Solver) const;
 		// Return the solver's positions for this cloth current LOD, not thread safe, call must be done right after the solver update.
 		CHAOSCLOTH_API TConstArrayView<Softs::FSolverVec3> GetParticlePositions(const FClothingSimulationSolver* Solver) const;
 		// Return the solver's velocities for this cloth current LOD, not thread safe, call must be done right after the solver update.
diff --git a/Engine/Plugins/ChaosCloth/Source/ChaosCloth/Public/ChaosCloth/ChaosClothingSimulationMesh.h b/Engine/Plugins/ChaosCloth/Source/ChaosCloth/Public/ChaosCloth/ChaosClothingSimulationMesh.h
--- a/Engine/Plugins/ChaosCloth/Source/ChaosCloth/Public/ChaosCloth/ChaosClothingSimulationMesh.h
+++ b/Engine/Plugins/ChaosCloth/Source/ChaosCloth/Public/ChaosCloth/ChaosClothingSimulationMesh.h
@@ -165,6 +165,7 @@ namespace Chaos
 	private:
 		void SkinPhysicsMesh(
 			int32 LODIndex,
+			const FReal LocalSpaceScale,
 			const FVec3& LocalSpaceLocation,
 			TArrayView<Softs::FSolverVec3>& OutPositions,
 			TArrayView<Softs::FSolverVec3>& OutNormals) const;
diff --git a/Engine/Plugins/ChaosCloth/Source/ChaosCloth/Public/ChaosCloth/ChaosClothingSimulationSolver.h b/Engine/Plugins/ChaosCloth/Source/ChaosCloth/Public/ChaosCloth/ChaosClothingSimulationSolver.h
--- a/Engine/Plugins/ChaosCloth/Source/ChaosCloth/Public/ChaosCloth/ChaosClothingSimulationSolver.h
+++ b/Engine/Plugins/ChaosCloth/Source/ChaosCloth/Public/ChaosCloth/ChaosClothingSimulationSolver.h
@@ -47,6 +47,8 @@ namespace Chaos
 		const FVec3& GetLocalSpaceLocation() const { return LocalSpaceLocation; }
 		void SetLocalSpaceRotation(const FQuat& InLocalSpaceRotation) { LocalSpaceRotation = InLocalSpaceRotation; }
 		const FRotation3& GetLocalSpaceRotation() const { return LocalSpaceRotation; }
+		CHAOSCLOTH_API void SetLocalSpaceScale(FReal InLocalSpaceScale, bool bReset = false);
+		FReal GetLocalSpaceScale() const { return LocalSpaceScale; }
 		void SetVelocityScale(FReal InVelocityScale) { VelocityScale = InVelocityScale; }
 		FReal GetVelocityScale() const { return VelocityScale; }
 
@@ -473,11 +475,14 @@ namespace Chaos
 		TMap<int32, TUniquePtr<FClothConstraints>> ClothsConstraints;
 
 		// Local space simulation
+		FReal OldLocalSpaceScale = 1.; // Multiply this to simulation space coordinates to get world space coordinates.
+		FReal LocalSpaceScale = 1.;
 		FVec3 OldLocalSpaceLocation;  // This is used to translate between world space and simulation space,
 		FVec3 LocalSpaceLocation;     // add this to simulation space coordinates to get world space coordinates, must keep FReal as underlying type for LWC
-		FRotation3 LocalSpaceRotation;
+		FRotation3 LocalSpaceRotation; // NOTE: Currently the simulation does not apply LocalSpaceRotation.
 		FReal VelocityScale;
 
+
 		// Time stepping
 		FSolverReal Time;
 		FSolverReal DeltaTime;
diff --git a/Engine/Plugins/ChaosClothAsset/Source/ChaosClothAssetEngine/Private/ChaosClothAsset/ClothSimulationContext.cpp b/Engine/Plugins/ChaosClothAsset/Source/ChaosClothAssetEngine/Private/ChaosClothAsset/ClothSimulationContext.cpp
--- a/Engine/Plugins/ChaosClothAsset/Source/ChaosClothAssetEngine/Private/ChaosClothAsset/ClothSimulationContext.cpp
+++ b/Engine/Plugins/ChaosClothAsset/Source/ChaosClothAssetEngine/Private/ChaosClothAsset/ClothSimulationContext.cpp
@@ -34,6 +34,9 @@ namespace UE::Chaos::ClothAsset
 		// Copy component transform
 		ComponentTransform = ClothComponent.GetComponentTransform();
 
+		// Copy solver geometry scale
+		SolverGeometryScale = ClothComponent.GetClothGeometryScale();;
+
 		// Update bone transforms
 		const UChaosClothAsset* const ClothAsset = ClothComponent.GetClothAsset();
 		const FReferenceSkeleton* const ReferenceSkeleton = ClothAsset ? &ClothAsset->GetRefSkeleton() : nullptr;
diff --git a/Engine/Plugins/ChaosClothAsset/Source/ChaosClothAssetEngine/Private/ChaosClothAsset/ClothSimulationContext.h b/Engine/Plugins/ChaosClothAsset/Source/ChaosClothAssetEngine/Private/ChaosClothAsset/ClothSimulationContext.h
--- a/Engine/Plugins/ChaosClothAsset/Source/ChaosClothAssetEngine/Private/ChaosClothAsset/ClothSimulationContext.h
+++ b/Engine/Plugins/ChaosClothAsset/Source/ChaosClothAssetEngine/Private/ChaosClothAsset/ClothSimulationContext.h
@@ -27,6 +27,9 @@ namespace UE::Chaos::ClothAsset
 		/** Component to world transform. */
 		FTransform ComponentTransform;
 
+		/** Scale applied to solver. */
+		float SolverGeometryScale;
+
 		/** Component space bone transforms of the owning component. */
 		TArray<FTransform> BoneTransforms;
 
diff --git a/Engine/Plugins/ChaosClothAsset/Source/ChaosClothAssetEngine/Private/ChaosClothAsset/ClothSimulationProxy.cpp b/Engine/Plugins/ChaosClothAsset/Source/ChaosClothAssetEngine/Private/ChaosClothAsset/ClothSimulationProxy.cpp
--- a/Engine/Plugins/ChaosClothAsset/Source/ChaosClothAssetEngine/Private/ChaosClothAsset/ClothSimulationProxy.cpp
+++ b/Engine/Plugins/ChaosClothAsset/Source/ChaosClothAssetEngine/Private/ChaosClothAsset/ClothSimulationProxy.cpp
@@ -122,6 +122,8 @@ PRAGMA_ENABLE_DEPRECATION_WARNINGS
 
 		// Setup startup transforms
 		constexpr bool bNeedsReset = true;
+		const FReal LocalSpaceScale = 1. / FMath::Max(ClothSimulationContext->SolverGeometryScale, UE_SMALL_NUMBER);
+		Solver->SetLocalSpaceScale(LocalSpaceScale, bNeedsReset);
 		Solver->SetLocalSpaceLocation((FVec3)ClothSimulationContext->ComponentTransform.GetLocation(), bNeedsReset);
 		Solver->SetLocalSpaceRotation((FQuat)ClothSimulationContext->ComponentTransform.GetRotation());
 
@@ -503,6 +505,7 @@ PRAGMA_ENABLE_DEPRECATION_WARNINGS
 
 		// Get the solver's local space
 		const FVec3& LocalSpaceLocation = Solver->GetLocalSpaceLocation(); // Note: Since the ReferenceSpaceTransform can be suspended with the simulation, it is important that the suspended local space location is used too in order to get the simulation data back into reference space
+		const FReal LocalSpaceScale = Solver->GetLocalSpaceScale();
 
 		// Retrieve the component's bones transforms
 		const TArray<FTransform>& ComponentSpaceTransforms = LeaderPoseComponent ? LeaderPoseComponent->GetComponentSpaceTransforms() : ClothComponent.GetComponentSpaceTransforms();
@@ -572,6 +575,7 @@ PRAGMA_ENABLE_DEPRECATION_WARNINGS
 					(ispc::FVector3f*)Data.Positions.GetData(),
 					(ispc::FVector3f*)Data.Normals.GetData(),
 					(ispc::FTransform&)ReferenceSpaceTransform,
+					LocalSpaceScale,
 					Data.Positions.Num());
 			}
 			else
@@ -579,7 +583,7 @@ PRAGMA_ENABLE_DEPRECATION_WARNINGS
 			{
 				for (int32 Index = 0; Index < Data.Positions.Num(); ++Index)
 				{
-					Data.Positions[Index] = FVec3f(ReferenceSpaceTransform.InverseTransformPosition(FVec3(Data.Positions[Index])));
+					Data.Positions[Index] = FVec3f(ReferenceSpaceTransform.InverseTransformPosition(LocalSpaceScale * FVec3(Data.Positions[Index])));
 					Data.Normals[Index] = FVec3f(ReferenceSpaceTransform.InverseTransformVector(FVec3(-Data.Normals[Index])));
 				}
 			}
@@ -612,7 +616,10 @@ PRAGMA_ENABLE_DEPRECATION_WARNINGS
 			FBoxSphereBounds Bounds = Solver->CalculateBounds();
 
 			// The component could be moving while the simulation is suspended so getting the bounds
-			// in world space isn't good enough and the bounds origin needs to be continuously updated
+			// in world space isn't good enough and the bounds origin needs to be continuously updated.
+			// 
+			// This converts the bounds back to component space. Do not apply LocalSpaceScale, which may not match component space.
+			// TODO: this will not apply the component's actual scale either.
 			Bounds = Bounds.TransformBy(FTransform((FQuat)Solver->GetLocalSpaceRotation(), (FVector)Solver->GetLocalSpaceLocation()).Inverse());
 
 			return Bounds;
diff --git a/Engine/Plugins/ChaosClothAsset/Source/ChaosClothAssetEngine/Private/ChaosClothAsset/ClothSimulationProxy.ispc b/Engine/Plugins/ChaosClothAsset/Source/ChaosClothAssetEngine/Private/ChaosClothAsset/ClothSimulationProxy.ispc
--- a/Engine/Plugins/ChaosClothAsset/Source/ChaosClothAssetEngine/Private/ChaosClothAsset/ClothSimulationProxy.ispc
+++ b/Engine/Plugins/ChaosClothAsset/Source/ChaosClothAssetEngine/Private/ChaosClothAsset/ClothSimulationProxy.ispc
@@ -9,6 +9,7 @@
 export void TransformClothSimulData(uniform FVector3f Positions[],
 									uniform FVector3f Normals[],
 									const uniform FTransform &ReferenceSpaceTransform,
+									const uniform double LocalSpaceScale,
 									const uniform int NumPositions)
 {
 	const uniform FVector4 RotationInverse = QuatInverse(ReferenceSpaceTransform.Rotation);
@@ -26,7 +27,8 @@ export void TransformClothSimulData(uniform FVector3f Positions[],
 		const FVector3f InputNormalVectorFloat = NormalFloat * -FLOAT_ONE;
 		const FVector InputNormalVector = ConvertVector3fTo3Native(InputNormalVectorFloat);
 
-		const FVector TranslatedPosition = InputPositionVector - Translation;
+		const FVector ScaledPosition = InputPositionVector * LocalSpaceScale;
+		const FVector TranslatedPosition = ScaledPosition - Translation;
 
 		const FVector VRPosition = VectorQuaternionRotateVector(RotationInverse, TranslatedPosition);
 		const FVector VRNormal = VectorQuaternionRotateVector(RotationInverse, InputNormalVector);
diff --git a/Engine/Plugins/ChaosClothAsset/Source/ChaosClothAssetEngine/Public/ChaosClothAsset/ClothComponent.h b/Engine/Plugins/ChaosClothAsset/Source/ChaosClothAssetEngine/Public/ChaosClothAsset/ClothComponent.h
--- a/Engine/Plugins/ChaosClothAsset/Source/ChaosClothAssetEngine/Public/ChaosClothAsset/ClothComponent.h
+++ b/Engine/Plugins/ChaosClothAsset/Source/ChaosClothAssetEngine/Public/ChaosClothAsset/ClothComponent.h
@@ -135,6 +135,12 @@ public:
 
 	const UE::Chaos::ClothAsset::FClothSimulationProxy* GetClothSimulationProxy() const { return ClothSimulationProxy.Get(); }
 
+	/** This scale is applied to all cloth geometry (e.g., cloth meshes and collisions) in order to simulate in a different scale space than world.This scale is not applied to distance-based simulation parameters such as MaxDistance.
+	* This property is currently only read by the cloth solver when creating cloth actors, but may become animatable in the future.
+	*/
+	float GetClothGeometryScale() const { return ClothGeometryScale; }
+	void SetClothGeometryScale(float Scale) { ClothGeometryScale = Scale; }
+
 #if WITH_EDITOR
 	/** Update config properties from the asset. Will only update existing values.*/
 	void UpdateConfigProperties();
@@ -253,6 +259,12 @@ private:
 	UPROPERTY(Interp, Category = ClothComponent)
 	float BlendWeight = 1.f;
 
+	/** This scale is applied to all cloth geometry (e.g., cloth meshes and collisions) in order to simulate in a different scale space than world.This scale is not applied to distance-based simulation parameters such as MaxDistance.
+	* This property is currently only read by the cloth solver when creating cloth actors, but may become animatable in the future.
+	*/
+	UPROPERTY(EditAnywhere, Category = ClothComponent, meta = (UIMin = 0.0, UIMax = 10.0, ClampMin = 0.0, ClampMax = 10000.0))
+	float ClothGeometryScale = 1.f;
+
 #if WITH_EDITOR
 	bool bTickOnceInEditor = false;
 #endif
diff --git a/Engine/Plugins/ChaosClothAssetEditor/Source/ChaosClothAssetEditor/Private/ChaosClothAsset/ClothEditorPreviewScene.cpp b/Engine/Plugins/ChaosClothAssetEditor/Source/ChaosClothAssetEditor/Private/ChaosClothAsset/ClothEditorPreviewScene.cpp
--- a/Engine/Plugins/ChaosClothAssetEditor/Source/ChaosClothAssetEditor/Private/ChaosClothAsset/ClothEditorPreviewScene.cpp
+++ b/Engine/Plugins/ChaosClothAssetEditor/Source/ChaosClothAssetEditor/Private/ChaosClothAsset/ClothEditorPreviewScene.cpp
@@ -292,6 +292,13 @@ void FChaosClothPreviewScene::SceneDescriptionPropertyChanged(const FName& Prope
 		}
 	}
 
+	if (PropertyName == GET_MEMBER_NAME_CHECKED(UChaosClothPreviewSceneDescription, SolverGeometryScale))
+	{
+		if (ClothComponent)
+		{
+			ClothComponent->SetClothGeometryScale(PreviewSceneDescription->SolverGeometryScale);
+		}
+	}
 }
 
 UChaosClothComponent* FChaosClothPreviewScene::GetClothComponent()
diff --git a/Engine/Plugins/ChaosClothAssetEditor/Source/ChaosClothAssetEditor/Private/ChaosClothAsset/ClothEditorPreviewScene.h b/Engine/Plugins/ChaosClothAssetEditor/Source/ChaosClothAssetEditor/Private/ChaosClothAsset/ClothEditorPreviewScene.h
--- a/Engine/Plugins/ChaosClothAssetEditor/Source/ChaosClothAssetEditor/Private/ChaosClothAsset/ClothEditorPreviewScene.h
+++ b/Engine/Plugins/ChaosClothAssetEditor/Source/ChaosClothAssetEditor/Private/ChaosClothAsset/ClothEditorPreviewScene.h
@@ -62,6 +62,9 @@ public:
 	UPROPERTY(EditAnywhere, Transient, Category = "Transform", Meta = (AllowPreserveRatio))
 	FVector3d Scale = FVector3d::OneVector;
 
+	UPROPERTY(EditAnywhere, Transient, Category = "ClothComponent", Meta = (UIMin = 0.0, UIMax = 10.0, ClampMin = 0.0, ClampMax = 10000.0))
+	float SolverGeometryScale = 1.f;
+
 	// TODO: We should be able to hook this boolean property up to the EditCondition meta tag for the properties above and toggle it
 	// on and off when the selection changes in the scene. However the EditCondition does not seem to propagate for some reason, 
 	// even if we manually call PostEditChangeProperty() after toggling it. It will take some more digging to figure out exactly
diff --git a/Engine/Source/Runtime/ClothingSystemRuntimeCommon/Private/ClothingSimulation.cpp b/Engine/Source/Runtime/ClothingSystemRuntimeCommon/Private/ClothingSimulation.cpp
--- a/Engine/Source/Runtime/ClothingSystemRuntimeCommon/Private/ClothingSimulation.cpp
+++ b/Engine/Source/Runtime/ClothingSystemRuntimeCommon/Private/ClothingSimulation.cpp
@@ -61,6 +61,7 @@ void FClothingSimulationContextCommon::Fill(const USkeletalMeshComponent* InComp
 	FillDeltaSeconds(InDeltaSeconds, InMaxPhysicsDelta);
 	FillTeleportMode(InComponent, InDeltaSeconds, InMaxPhysicsDelta);
 	FillMaxDistanceScale(InComponent);
+	FillSolverGeometryScale(InComponent);
 
 	PredictedLod = InComponent->GetPredictedLODLevel();
 }
@@ -185,6 +186,11 @@ void FClothingSimulationContextCommon::FillMaxDistanceScale(const USkeletalMeshC
 	MaxDistanceScale = InComponent->GetClothMaxDistanceScale();
 }
 
+void FClothingSimulationContextCommon::FillSolverGeometryScale(const USkeletalMeshComponent* InComponent)
+{
+	SolverGeometryScale = InComponent->ClothGeometryScale;
+}
+
 //==============================================================================
 // FClothingSimulationCommon
 //==============================================================================
diff --git a/Engine/Source/Runtime/ClothingSystemRuntimeCommon/Public/ClothingSimulation.h b/Engine/Source/Runtime/ClothingSystemRuntimeCommon/Public/ClothingSimulation.h
--- a/Engine/Source/Runtime/ClothingSystemRuntimeCommon/Public/ClothingSimulation.h
+++ b/Engine/Source/Runtime/ClothingSystemRuntimeCommon/Public/ClothingSimulation.h
@@ -40,6 +40,7 @@ protected:
 	CLOTHINGSYSTEMRUNTIMECOMMON_API virtual void FillDeltaSeconds(float InDeltaSeconds, float InMaxPhysicsDelta);
 	CLOTHINGSYSTEMRUNTIMECOMMON_API virtual void FillTeleportMode(const USkeletalMeshComponent* InComponent, float InDeltaSeconds, float InMaxPhysicsDelta);
 	CLOTHINGSYSTEMRUNTIMECOMMON_API virtual void FillMaxDistanceScale(const USkeletalMeshComponent* InComponent);
+	CLOTHINGSYSTEMRUNTIMECOMMON_API virtual void FillSolverGeometryScale(const USkeletalMeshComponent* InComponent);
 
 public:
 	// Component space bone transforms of the owning component
@@ -73,6 +74,9 @@ public:
 	// Scale for the max distance constraints of the simulation mesh
 	float MaxDistanceScale;
 
+	// Scale applied to all incoming geometry
+	float SolverGeometryScale = 1.f;
+
 	// The predicted LOD of the skeletal mesh component running the simulation
 	int32 PredictedLod;
 	
diff --git a/Engine/Source/Runtime/Engine/Classes/Components/SkeletalMeshComponent.h b/Engine/Source/Runtime/Engine/Classes/Components/SkeletalMeshComponent.h
--- a/Engine/Source/Runtime/Engine/Classes/Components/SkeletalMeshComponent.h
+++ b/Engine/Source/Runtime/Engine/Classes/Components/SkeletalMeshComponent.h
@@ -888,6 +888,12 @@ public:
 
 	UPROPERTY(Interp, BlueprintReadWrite, Category=Clothing, meta=(UIMin = 0.0, UIMax = 10.0, ClampMin = 0.0, ClampMax = 10000.0))
 	float ClothMaxDistanceScale;
+	
+	/** This scale is applied to all cloth geometry (e.g., cloth meshes and collisions) in order to simulate in a different scale space than world.This scale is not applied to distance-based simulation parameters such as MaxDistance. 
+	* This property is currently only read by the cloth solver when creating cloth actors, but may become animatable in the future.
+	*/
+	UPROPERTY(BlueprintReadWrite, Category = Clothing, meta = (UIMin = 0.0, UIMax = 10.0, ClampMin = 0.0, ClampMax = 10000.0))
+	float ClothGeometryScale = 1.f;
 
 	/** Notification when constraint is broken. */
 	UPROPERTY(BlueprintAssignable)
