Chaos Cloth: A couple of bugfixes to velocity scale after recent changes.

* Fix a merge bug when applying ClothGeometryScale where the local scale wasn't actually applied to the calculated delta position due to order of operations with a FRigidTransform3
* Fix a bug where Acceleration clamping wasn't fully respecting VelocityScaleSpace.

diff --git a/Engine/Plugins/ChaosCloth/Source/ChaosCloth/Private/ChaosCloth/ChaosClothingSimulationSolver.cpp b/Engine/Plugins/ChaosCloth/Source/ChaosCloth/Private/ChaosCloth/ChaosClothingSimulationSolver.cpp
--- a/Engine/Plugins/ChaosCloth/Source/ChaosCloth/Private/ChaosCloth/ChaosClothingSimulationSolver.cpp
+++ b/Engine/Plugins/ChaosCloth/Source/ChaosCloth/Private/ChaosCloth/ChaosClothingSimulationSolver.cpp
@@ -1186,29 +1186,54 @@ void FClothingSimulationSolver::SetReferenceVelocityScale(uint32 GroupId,
 		FVec3 AppliedLinearVelocity = AppliedDeltaTranslation / DeltaTime;
 		FVec3 AppliedAngularVelocity = AppliedDeltaRotation / DeltaTime;
 
-		const FVec3 OldScaledReferenceSpaceVelocity = LinearRatio * (OldReferenceSpaceRotationInverse * OldReferenceVelocity);
-		const FVec3 ScaledReferenceSpaceAccelerationTimesDt = AppliedLinearVelocity - OldScaledReferenceSpaceVelocity;
-		for (int32 Index = 0; Index < 3; ++Index)
+		// Apply linear acceleration clamping
+		if (MaxLinearAcceleration != TVec3<FRealSingle>(TNumericLimits<FRealSingle>::Max()))
 		{
-			if (MaxLinearAcceleration[Index] != TNumericLimits<FRealSingle>::Max() && MaxLinearAcceleration[Index] >= 0.f
-				&& FMath::Abs(ScaledReferenceSpaceAccelerationTimesDt[Index]) > (FReal)(MaxLinearAcceleration[Index] * DeltaTime * LocalSpaceScale))
+			FVec3 OldScaledVelocity;
+			switch (VelocityScaleSpace)
 			{
-				AppliedLinearVelocity[Index] = OldScaledReferenceSpaceVelocity[Index] + FMath::Sign(ScaledReferenceSpaceAccelerationTimesDt[Index]) * (FReal)(MaxLinearAcceleration[Index] * DeltaTime * LocalSpaceScale);
+			case EChaosSoftsSimulationSpace::WorldSpace:
+				OldScaledVelocity = LinearRatio * OldReferenceVelocity;
+				break;
+			case EChaosSoftsSimulationSpace::ComponentSpace:
+				OldScaledVelocity = LinearRatio * LocalSpaceRotation.UnrotateVector(OldReferenceVelocity);
+				break;
+			case EChaosSoftsSimulationSpace::ReferenceBoneSpace:
+			default:
+				OldScaledVelocity = LinearRatio * (OldReferenceSpaceRotationInverse * OldReferenceVelocity);
+				break;
 			}
-			if (MaxLinearVelocity[Index] != TNumericLimits<FRealSingle>::Max() && MaxLinearVelocity[Index] >= 0.f &&
-				FMath::Abs(AppliedLinearVelocity[Index]) > (FReal)MaxLinearVelocity[Index] * LocalSpaceScale)
+			
+			const FVec3 ScaledAccelerationTimesDt = AppliedLinearVelocity - OldScaledVelocity;
+			for (int32 Index = 0; Index < 3; ++Index)
 			{
-				AppliedLinearVelocity[Index] = FMath::Sign(AppliedLinearVelocity[Index]) * (FReal)MaxLinearVelocity[Index] * LocalSpaceScale;
+				if (MaxLinearAcceleration[Index] != TNumericLimits<FRealSingle>::Max() && MaxLinearAcceleration[Index] >= 0.f
+					&& FMath::Abs(ScaledAccelerationTimesDt[Index]) > (FReal)(MaxLinearAcceleration[Index] * DeltaTime * LocalSpaceScale))
+				{
+					AppliedLinearVelocity[Index] = OldScaledVelocity[Index] + FMath::Sign(ScaledAccelerationTimesDt[Index]) * (FReal)(MaxLinearAcceleration[Index] * DeltaTime * LocalSpaceScale);
+				}
 			}
+		}
 
+		// Apply linear velocity clamping
+		if (MaxLinearVelocity != TVec3<FRealSingle>(TNumericLimits<FRealSingle>::Max()))
+		{
+			for (int32 Index = 0; Index < 3; ++Index)
+			{
+				if (MaxLinearVelocity[Index] != TNumericLimits<FRealSingle>::Max() && MaxLinearVelocity[Index] >= 0.f &&
+					FMath::Abs(AppliedLinearVelocity[Index]) > (FReal)MaxLinearVelocity[Index] * LocalSpaceScale)
+				{
+					AppliedLinearVelocity[Index] = FMath::Sign(AppliedLinearVelocity[Index]) * (FReal)MaxLinearVelocity[Index] * LocalSpaceScale;
+				}
+			}
 		}
 
 		AppliedDeltaTranslation = AppliedLinearVelocity * DeltaTime;
 
-		const FVec3 OldScaledAngularVelocity = AngularRatio * OldReferenceAngularVelocity;
-		const FVec3 ScaledReferenceSpaceAngularAccelerationTimesDt = AppliedAngularVelocity - OldScaledAngularVelocity;
 		if (MaxAngularAcceleration != TNumericLimits<FRealSingle>::Max() && MaxAngularAcceleration >= 0.f)
 		{
+			const FVec3 OldScaledAngularVelocity = AngularRatio * OldReferenceAngularVelocity;
+			const FVec3 ScaledReferenceSpaceAngularAccelerationTimesDt = AppliedAngularVelocity - OldScaledAngularVelocity;
 			const FReal LenSq = ScaledReferenceSpaceAngularAccelerationTimesDt.SquaredLength();
 			if (LenSq > FMath::Square(MaxAngularAcceleration * DeltaTime))
 			{
@@ -1231,7 +1256,7 @@ void FClothingSimulationSolver::SetReferenceVelocityScale(uint32 GroupId,
 
 	check(LocalSpaceScale > UE_SMALL_NUMBER);
 	const FReal LocalSpaceScaleInv = 1. / LocalSpaceScale;
-	const FVec3 DeltaPosition = LocalSpaceScaleInv * VelocityScaleSpaceToOldRefSpace * (VelocityScaleSpaceDeltaTranslation - AppliedDeltaTranslation);
+	const FVec3 DeltaPosition = VelocityScaleSpaceToOldRefSpace * (LocalSpaceScaleInv * (VelocityScaleSpaceDeltaTranslation - AppliedDeltaTranslation));
 	const FVec3 DeltaRotation = FullDeltaRotation - AppliedDeltaRotation;
 	const FReal DeltaAngle = DeltaRotation.Length();
 	const FVec3 Axis = DeltaAngle > UE_KINDA_SMALL_NUMBER ? DeltaRotation / DeltaAngle : FVec3::XAxisVector;
