Chaos Cloth - Make debug names more accurate in mesh debug messages.

diff --git a/Engine/Plugins/ChaosCloth/Source/ChaosCloth/Private/ChaosCloth/ChaosClothingSimulationSkeletalMesh.cpp b/Engine/Plugins/ChaosCloth/Source/ChaosCloth/Private/ChaosCloth/ChaosClothingSimulationSkeletalMesh.cpp
--- a/Engine/Plugins/ChaosCloth/Source/ChaosCloth/Private/ChaosCloth/ChaosClothingSimulationSkeletalMesh.cpp
+++ b/Engine/Plugins/ChaosCloth/Source/ChaosCloth/Private/ChaosCloth/ChaosClothingSimulationSkeletalMesh.cpp
@@ -11,7 +11,7 @@ namespace Chaos
 {
 	FClothingSimulationSkeletalMesh::FClothingSimulationSkeletalMesh(const UClothingAssetCommon* InAsset, const USkeletalMeshComponent* InSkeletalMeshComponent)
 		: FClothingSimulationMesh(InSkeletalMeshComponent->GetOwner() ?
-				FString::Format(TEXT("{0}|{1}"), { InSkeletalMeshComponent->GetOwner()->GetName(), InSkeletalMeshComponent->GetName() }) :
+				FString::Format(TEXT("{0}|{1}"), { InSkeletalMeshComponent->GetOwner()->GetActorNameOrLabel(), InSkeletalMeshComponent->GetName() }) :
 				InSkeletalMeshComponent->GetName())
 		, Asset(InAsset)
 		, SkeletalMeshComponent(InSkeletalMeshComponent)
diff --git a/Engine/Plugins/ChaosClothAsset/Source/ChaosClothAssetEngine/Private/ChaosClothAsset/ClothSimulationProxy.cpp b/Engine/Plugins/ChaosClothAsset/Source/ChaosClothAssetEngine/Private/ChaosClothAsset/ClothSimulationProxy.cpp
--- a/Engine/Plugins/ChaosClothAsset/Source/ChaosClothAssetEngine/Private/ChaosClothAsset/ClothSimulationProxy.cpp
+++ b/Engine/Plugins/ChaosClothAsset/Source/ChaosClothAssetEngine/Private/ChaosClothAsset/ClothSimulationProxy.cpp
@@ -139,7 +139,7 @@ PRAGMA_ENABLE_DEPRECATION_WARNINGS
 		FString DebugName;
 #if !UE_BUILD_SHIPPING
 		DebugName = ClothComponent.GetOwner() ?
-			FString::Format(TEXT("{0}|{1}"), { ClothComponent.GetOwner()->GetName(), ClothComponent.GetName() }) :
+			FString::Format(TEXT("{0}|{1}"), { ClothComponent.GetOwner()->GetActorNameOrLabel(), ClothComponent.GetName() }) :
 			ClothComponent.GetName();
 #endif
 		const int32 MeshIndex = Meshes.Emplace(MakeUnique<FClothSimulationMesh>(*ClothSimulationModel, *ClothSimulationContext, DebugName));
