SKM Component: expose ClothTeleportMode to Sequencer. Added a new HardReset mode which rebuilds the cloths. The property is transient and hidden from the details panel.

diff --git a/Engine/Plugins/ChaosCloth/Source/ChaosCloth/Private/ChaosCloth/ChaosClothingSimulation.cpp b/Engine/Plugins/ChaosCloth/Source/ChaosCloth/Private/ChaosCloth/ChaosClothingSimulation.cpp
--- a/Engine/Plugins/ChaosCloth/Source/ChaosCloth/Private/ChaosCloth/ChaosClothingSimulation.cpp
+++ b/Engine/Plugins/ChaosCloth/Source/ChaosCloth/Private/ChaosCloth/ChaosClothingSimulation.cpp
@@ -511,9 +511,9 @@ void FClothingSimulation::Simulate(IClothingSimulationContext* InContext)
 
 #if !(UE_BUILD_SHIPPING || UE_BUILD_TEST)
 		const bool bNeedsReset = (ClothingSimulationConsole::Command && ClothingSimulationConsole::Command->MustReset(ResetCount)) ||
-			Context->TeleportMode == EClothingTeleportMode::TeleportAndReset || PrevSimulationTime == 0.f;
+			Context->TeleportMode >= EClothingTeleportMode::TeleportAndReset || PrevSimulationTime == 0.f;
 #else
-		const bool bNeedsReset = Context->TeleportMode == EClothingTeleportMode::TeleportAndReset || PrevSimulationTime == 0.f;
+		const bool bNeedsReset = Context->TeleportMode >= EClothingTeleportMode::TeleportAndReset || PrevSimulationTime == 0.f;
 #endif
 		const bool bNeedsTeleport = (Context->TeleportMode > EClothingTeleportMode::None);
 		bIsTeleported = bNeedsTeleport;
@@ -794,6 +794,9 @@ void FClothingSimulation::RefreshClothConfig(const IClothingSimulationContext* I
 	// Update new space location
 	const FClothingSimulationContext* const Context = static_cast<const FClothingSimulationContext*>(InContext);
 	static const bool bReset = true;
+	const FReal WorldToSolverScale = bUseLocalSpaceSimulation ? Context->SolverGeometryScale : 1.;
+	const FReal LocalSpaceScale = 1. / FMath::Max(WorldToSolverScale, UE_SMALL_NUMBER);
+	Solver->SetLocalSpaceScale(LocalSpaceScale, bReset);
 	Solver->SetLocalSpaceLocation(bUseLocalSpaceSimulation ? (FVec3)Context->ComponentToWorld.GetLocation() : FVec3(0), bReset);
 	Solver->SetLocalSpaceRotation(bUseLocalSpaceSimulation ? (FQuat)Context->ComponentToWorld.GetRotation() : FQuat::Identity);
 
diff --git a/Engine/Source/Editor/DetailCustomizations/Private/SkeletalMeshComponentDetails.cpp b/Engine/Source/Editor/DetailCustomizations/Private/SkeletalMeshComponentDetails.cpp
--- a/Engine/Source/Editor/DetailCustomizations/Private/SkeletalMeshComponentDetails.cpp
+++ b/Engine/Source/Editor/DetailCustomizations/Private/SkeletalMeshComponentDetails.cpp
@@ -241,6 +241,9 @@ EVisibility FSkeletalMeshComponentDetails::VisibilityForBlueprintMode() const
 
 void FSkeletalMeshComponentDetails::UpdatePhysicsCategory(IDetailLayoutBuilder& DetailBuilder)
 {
+	// Force hiding ClothTeleportMode. If you stop force hiding this, please update the comment in SkeletalMeshComponent.h where the property is declared.
+	const FName ClothTeleportModeFName(GET_MEMBER_NAME_CHECKED(USkeletalMeshComponent, ClothTeleportMode));
+	DetailBuilder.HideProperty(ClothTeleportModeFName);
 }
 
 EVisibility FSkeletalMeshComponentDetails::VisibilityForAnimationMode(EAnimationMode::Type AnimationMode) const
diff --git a/Engine/Source/Runtime/ClothingSystemRuntimeCommon/Private/ClothingSimulation.cpp b/Engine/Source/Runtime/ClothingSystemRuntimeCommon/Private/ClothingSimulation.cpp
--- a/Engine/Source/Runtime/ClothingSystemRuntimeCommon/Private/ClothingSimulation.cpp
+++ b/Engine/Source/Runtime/ClothingSystemRuntimeCommon/Private/ClothingSimulation.cpp
@@ -170,7 +170,7 @@ void FClothingSimulationContextCommon::FillDeltaSeconds(float InDeltaSeconds, fl
 
 void FClothingSimulationContextCommon::FillTeleportMode(const USkeletalMeshComponent* InComponent, float InDeltaSeconds, float InMaxPhysicsDelta)
 {
-	TeleportMode = (InDeltaSeconds > InMaxPhysicsDelta * GClothMaxDeltaTimeTeleportMultiplier.GetValueOnGameThread()) ?
+	TeleportMode = (InComponent->ClothTeleportMode < EClothingTeleportMode::Teleport && (InDeltaSeconds > InMaxPhysicsDelta * GClothMaxDeltaTimeTeleportMultiplier.GetValueOnGameThread())) ?
 		EClothingTeleportMode::Teleport :
 		InComponent->ClothTeleportMode;
 
diff --git a/Engine/Source/Runtime/ClothingSystemRuntimeInterface/Public/ClothingSystemRuntimeTypes.h b/Engine/Source/Runtime/ClothingSystemRuntimeInterface/Public/ClothingSystemRuntimeTypes.h
--- a/Engine/Source/Runtime/ClothingSystemRuntimeInterface/Public/ClothingSystemRuntimeTypes.h
+++ b/Engine/Source/Runtime/ClothingSystemRuntimeInterface/Public/ClothingSystemRuntimeTypes.h
@@ -7,6 +7,7 @@
 #include "Math/Transform.h"
 #include "Math/UnrealMathSSE.h"
 #include "Math/Vector.h"
+#include "ClothingSystemRuntimeTypes.generated.h"
 
 
 // Data produced by a clothing simulation
@@ -35,6 +36,7 @@ struct FClothSimulData
 	int32 LODIndex;
 };
 
+UENUM()
 enum class EClothingTeleportMode : uint8
 {
 	// No teleport, simulate as normal
@@ -42,6 +44,8 @@ enum class EClothingTeleportMode : uint8
 	// Teleport the simulation, causing no intertial effects but keep the sim mesh shape
 	Teleport,
 	// Teleport the simulation, causing no intertial effects and reset the sim mesh shape
-	TeleportAndReset
+	TeleportAndReset,
+	// Hard reset the simulation by refreshing the cloth config
+	HardReset
 };
 
diff --git a/Engine/Source/Runtime/Engine/Classes/Components/SkeletalMeshComponent.h b/Engine/Source/Runtime/Engine/Classes/Components/SkeletalMeshComponent.h
--- a/Engine/Source/Runtime/Engine/Classes/Components/SkeletalMeshComponent.h
+++ b/Engine/Source/Runtime/Engine/Classes/Components/SkeletalMeshComponent.h
@@ -575,6 +575,7 @@ public:
 	TEnumAsByte<EPhysicsTransformUpdateMode::Type> PhysicsTransformUpdateMode;
 
 	/** whether we need to teleport cloth. */
+	UPROPERTY(Interp, Transient, BlueprintReadOnly, VisibleAnywhere,  Category=Clothing) // This property is explicitly hidden from the details panel inside FSkeletalMeshComponentDetails::UpdatePhysicsCategory
 	EClothingTeleportMode ClothTeleportMode;
 
 protected:
@@ -893,7 +894,7 @@ public:
 	/** This scale is applied to all cloth geometry (e.g., cloth meshes and collisions) in order to simulate in a different scale space than world.This scale is not applied to distance-based simulation parameters such as MaxDistance. 
 	* This property is currently only read by the cloth solver when creating cloth actors, but may become animatable in the future.
 	*/
-	UPROPERTY(EditAnywhere, BlueprintReadWrite, Category = Clothing, meta = (UIMin = 0.0, UIMax = 10.0, ClampMin = 0.0, ClampMax = 10000.0))
+	UPROPERTY(EditAnywhere, BlueprintReadWrite, Interp, Category = Clothing, meta = (UIMin = 0.0, UIMax = 10.0, ClampMin = 0.0, ClampMax = 10000.0))
 	float ClothGeometryScale = 1.f;
 
 	/** Notification when constraint is broken. */
diff --git a/Engine/Source/Runtime/Engine/Private/Components/SkeletalMeshComponent.cpp b/Engine/Source/Runtime/Engine/Private/Components/SkeletalMeshComponent.cpp
--- a/Engine/Source/Runtime/Engine/Private/Components/SkeletalMeshComponent.cpp
+++ b/Engine/Source/Runtime/Engine/Private/Components/SkeletalMeshComponent.cpp
@@ -2317,16 +2317,17 @@ void USkeletalMeshComponent::UpdateClothSimulationContext(float InDeltaTime)
 
 	bool bMustUpdateClothTransform = bForceCollisionUpdate;
 
+	const bool bNeedsHardReset = ClothTeleportMode == EClothingTeleportMode::HardReset;
 	if (bPendingClothTransformUpdate)	//it's possible we want to update cloth collision based on a pending transform
 	{
 		bPendingClothTransformUpdate = false;
 		if (PendingTeleportType == ETeleportType::TeleportPhysics)	//If the pending transform came from a teleport, make sure to teleport the cloth in this upcoming simulation
 		{
-			ClothTeleportMode = (ClothTeleportMode == EClothingTeleportMode::TeleportAndReset) ? ClothTeleportMode : EClothingTeleportMode::Teleport;
+			ClothTeleportMode = (ClothTeleportMode > EClothingTeleportMode::Teleport) ? ClothTeleportMode : EClothingTeleportMode::Teleport;
 		}
 		else if (PendingTeleportType == ETeleportType::ResetPhysics)
 		{
-			ClothTeleportMode = EClothingTeleportMode::TeleportAndReset;
+			ClothTeleportMode = (ClothTeleportMode > EClothingTeleportMode::TeleportAndReset) ? ClothTeleportMode : EClothingTeleportMode::TeleportAndReset;
 		}
 		bMustUpdateClothTransform = true;
 	}
@@ -2343,6 +2344,11 @@ void USkeletalMeshComponent::UpdateClothSimulationContext(float InDeltaTime)
 
 		if(ClothingInteractor)
 		{
+			if (bNeedsHardReset)
+			{
+				ClothingInteractor->ClothConfigUpdated();
+			}
+
 			ClothingInteractor->Sync(ClothingSimulation, ClothingSimulationContext);
 		}
 	}
