Chaos Cloth complex collider refactor and add support for SkinnedTriangleMesh colliders.

Note that this disables the use of SkinnedLevelSets for SKM cloth by default. It can be re-enabled by setting the cvar "p.ChaosCloth.UseLegacyEvolution 0" which will enable the same underlying solver that panel cloth uses.

Adds a notion of a "Complex Collider" which is colliders driven by multiple bones, i.e., SkinnedLevelSets, MLLevelSets, and now SkinnedTriangleMeshes. Removes use of empty colliders to handle tracking bones. These bones are now tracked in the ChaosClothComplexColliders class, where all updates are handled for these complex colliders.

MLLevelSets introduced a concept of PlanarConstraints which are calculated once per substep using full collision detection, and then subsequent iterations use planar constraints to approximate the collisions for a vertex. There was the possibility of a race condition in this code, and it didn't support friction. This change introduces PBDPlanarConstraints which do support friction and they are generated with one constraint per dynamic vertex, so race conditions are eliminated. To additionally test this code, I added the option to use planar constraints with Simple (i.e., rigid primitives) colliders. By default, simple colliders do full collisions every iteration, and complex colliders use planar constraints.

The "KinematicCollider" code has been refactored into a PBDKinematicTriangleMeshCollisions class. Old cloth-based kinematic colliders and new SkinnedTriangleMesh colliders now use this shared code.

Note that we continue to not support complex "External" colliders with this change. Hopefully a subsequent change can add this support.

diff --git a/Engine/Plugins/ChaosCloth/Source/ChaosCloth/Private/ChaosCloth/ChaosClothConstraints.cpp b/Engine/Plugins/ChaosCloth/Source/ChaosCloth/Private/ChaosCloth/ChaosClothConstraints.cpp
--- a/Engine/Plugins/ChaosCloth/Source/ChaosCloth/Private/ChaosCloth/ChaosClothConstraints.cpp
+++ b/Engine/Plugins/ChaosCloth/Source/ChaosCloth/Private/ChaosCloth/ChaosClothConstraints.cpp
@@ -2,6 +2,7 @@
 #include "ChaosCloth/ChaosClothConstraints.h"
 #include "ChaosCloth/ChaosWeightMapTarget.h"
 #include "ChaosCloth/ChaosClothingPatternData.h"
+#include "ChaosCloth/ChaosClothComplexColliders.h"
 #include "Chaos/PBDSpringConstraints.h"
 #include "Chaos/XPBDSpringConstraints.h"
 #include "Chaos/XPBDStretchBiasElementConstraints.h"
@@ -19,6 +20,7 @@
 #include "Chaos/PBDShapeConstraints.h"
 #include "Chaos/PBDCollisionSpringConstraints.h"
 #include "Chaos/PBDSelfCollisionSphereConstraints.h"
+#include "Chaos/PBDKinematicTriangleMeshCollisions.h"
 #include "Chaos/PBDTriangleMeshCollisions.h"
 #include "Chaos/PBDTriangleMeshIntersections.h"
 #include "Chaos/PBDEvolution.h"
@@ -363,7 +365,6 @@ FClothConstraints::FClothConstraints()
 	, NumUpdateLinearSystemRules(0)
 	, NumUpdateLinearSystemCollisionsRules(0)
 
-	, OldAnimationPositions_Deprecated(nullptr)
 	, ConstraintInitOffset(INDEX_NONE)
 	, ConstraintRuleOffset(INDEX_NONE)
 	, PostCollisionConstraintRuleOffset(INDEX_NONE)
@@ -420,7 +421,6 @@ void FClothConstraints::Initialize(
 	Evolution = nullptr;
 	PBDEvolution = InEvolution;
 	AnimationPositions = &InInterpolatedAnimationPositions;
-	OldAnimationPositions_Deprecated = nullptr;
 	AnimationNormals = &InInterpolatedAnimationNormals;
 	AnimationVelocities = &InAnimationVelocities;
 	ParticleOffset = InParticleOffset;
@@ -452,6 +452,18 @@ void FClothConstraints::Enable(bool bEnable)
 	}
 }
 
+void FClothConstraints::OnCollisionRangeRemoved(int32 CollisionRangeId)
+{
+	if (CollisionConstraint)
+	{
+		CollisionConstraint->OnCollisionRangeRemoved(CollisionRangeId);
+	}
+	if (SkinnedTriangleCollisionsConstraint)
+	{
+		SkinnedTriangleCollisionsConstraint->OnCollisionRangeRemoved(CollisionRangeId);
+	}
+}
+
 void FClothConstraints::AddRules(
 	const Softs::FCollectionPropertyConstFacade& ConfigProperties,
 	const FTriangleMesh& TriangleMesh,
@@ -782,7 +794,8 @@ void FClothConstraints::AddRules(
 	Softs::FSolverReal MeshScale, bool bEnabled,
 	const FTriangleMesh* MultiResCoarseLODMesh,
 	const int32 MultiResCoarseLODParticleRangeId,
-	const TSharedPtr<Softs::FMultiResConstraints>& FineLODMultiResConstraint)
+	const TSharedPtr<Softs::FMultiResConstraints>& FineLODMultiResConstraint,
+	const TArray<const FClothComplexColliders*>& ComplexColliders)
 {
 	// Self collisions
 	CreateSelfCollisionConstraints(ConfigProperties, WeightMaps, VertexSets, FaceSets, FaceIntMaps, TriangleMesh);
@@ -823,7 +836,7 @@ void FClothConstraints::AddRules(
 		}
 
 		// Body collisions
-		CreateCollisionConstraint(ConfigProperties, MeshScale);
+		CreateCollisionConstraint(ConfigProperties, MeshScale, WeightMaps, ComplexColliders);
 
 		//Multires Springs
 		CreateMultiresConstraint(ConfigProperties, WeightMaps, TriangleMesh, MultiResCoarseLODMesh, MultiResCoarseLODParticleRangeId);
@@ -1384,12 +1397,24 @@ void FClothConstraints::CreateExternalForces(
 
 void FClothConstraints::CreateCollisionConstraint(
 	const Softs::FCollectionPropertyConstFacade& ConfigProperties,
-	Softs::FSolverReal MeshScale)
+	Softs::FSolverReal MeshScale,
+	const TMap<FString, TConstArrayView<FRealSingle>>& WeightMaps,
+	const TArray<const FClothComplexColliders*>& ComplexColliders)
 {
 	if (Evolution)
 	{
 		// Always create collision constraint
 		check(LastSubframeCollisionTransformsCCD);
+
+		TMap<Softs::FParticleRangeIndex, Softs::FPBDComplexColliderBoneData> ComplexBoneData;
+		for (const FClothComplexColliders* Colliders : ComplexColliders)
+		{
+			if (Colliders)
+			{
+				Colliders->ExtractComplexColliderBoneData(ComplexBoneData);
+			}
+		}
+
 		CollisionConstraint = MakeShared<Softs::FPBDSoftBodyCollisionConstraint>(
 			*LastSubframeCollisionTransformsCCD,
 			ConfigProperties,
@@ -1397,21 +1422,34 @@ void FClothConstraints::CreateCollisionConstraint(
 			CollisionParticleCollided,
 			CollisionContacts,
 			CollisionNormals,
-			CollisionPhis
+			CollisionPhis,
+			ComplexBoneData
 			);
 		++NumCollisionConstraintRules;
 		++NumUpdateLinearSystemCollisionsRules;
-		if (bDoPlanarCollisionConstraints) 
+
+		SkinnedTriangleCollisionsConstraint = MakeShared<Softs::FPBDSkinnedTriangleMeshCollisionConstraints>(
+			NumParticles,
+			WeightMaps,
+			ConfigProperties
+		);
+
+		for (const FClothComplexColliders* Colliders : ComplexColliders)
 		{
-			CollisionConstraint->SetDoCollisionPerTimestep(true);
-			XPlanarConstraints = MakeShared<Softs::TXPBDPlanarConstraints<FSolverReal>>();
-			XPlanarConstraints->SetTolerance(PlanarCollisionTol);
-			XPlanarConstraints->SetStiffness(PlanarCollisionStiffness);
-			++NumConstraintRules;
-			++NumConstraintInits;
-			Evolution->SetDoCollisionPerTimestep(true);
+			if (Colliders)
+			{
+				const int32 NumMeshes = Colliders->GetNumSkinnedTriangleMeshes();
+				for (int32 MeshIndex = 0; MeshIndex < NumMeshes; ++MeshIndex)
+				{
+					SkinnedTriangleCollisionsConstraint->AddSkinnedTriangleMesh(
+						Softs::FParticleRangeIndex(Colliders->GetCollisionRangeId(), Colliders->GetSkinnedTriangleMeshIndex(MeshIndex)),
+						FSkinnedTriangleMeshPtr(Colliders->GetSkinnedTriangleMesh(MeshIndex)),
+						Colliders->GetSkinnedTriangleMeshVelocities(MeshIndex));
+				}
+			}
 		}
-
+		++NumConstraintInits;
+		++NumPostCollisionConstraintRules;
 	}
 }
 
@@ -1684,10 +1722,39 @@ void FClothConstraints::CreateForceBasedRules(const TSharedPtr<Softs::FMultiResC
 
 	if (CollisionConstraint)
 	{
-		RuleCreator.AddPerIterationCollisionPBDConstraintRule_Apply(CollisionConstraint.Get());
+		RuleCreator.AddPerIterationCollisionPBDConstraintRule(
+			[this](Softs::FSolverParticlesRange& Particles, const Softs::FSolverReal Dt, const TArray<Softs::FSolverCollisionParticlesRange>& CollisionParticles)
+			{
+				const Softs::FEvolutionGroupContext& Context = Evolution->GetGroupContextForSoftBody(Particles.GetRangeId());
+
+				if (Context.NumPBDIterations == 1)
+				{
+					// When only doing one PBDIteration, there's no need for planar constraints.
+					CollisionConstraint->Apply(Particles, Dt, CollisionParticles);
+				}
+				else
+				{
+					CollisionConstraint->ApplyWithPlanarConstraints(Particles, Dt, CollisionParticles, Context.CurrentPBDIteration == 0);
+				}
+			});
 		RuleCreator.AddUpdateLinearSystemCollisionsRule_UpdateLinearSystem(CollisionConstraint.Get());
 	}
 
+	if (SkinnedTriangleCollisionsConstraint)
+	{
+		RuleCreator.AddPostInitialGuessParallelInitRule(
+			[this](const Softs::FSolverParticlesRange& Particles, const Softs::FSolverReal Dt, const Softs::ESolverMode SolverMode)
+			{
+				if (bSkipSelfCollisionInit && SkinnedTriangleCollisionsConstraint->GetUseSelfCollisionSubstepsForSkinnedTriangleMeshes())
+				{
+					return;
+				}
+				SkinnedTriangleCollisionsConstraint->Init(Particles, Dt);
+			});
+		constexpr bool bPostCollisions = true;
+		RuleCreator.AddPerIterationPBDConstraintRule_Apply<bPostCollisions>(SkinnedTriangleCollisionsConstraint.Get());
+	}
+
 	if (SelfCollisionInit && SelfCollisionConstraints)
 	{
 		RuleCreator.AddPostInitialGuessParallelInitRule(
@@ -1763,20 +1830,6 @@ void FClothConstraints::CreateForceBasedRules(const TSharedPtr<Softs::FMultiResC
 			}
 		});
 	}
-	if (bDoPlanarCollisionConstraints)
-	{
-		RuleCreator.AddPerIterationPBDConstraintRule(
-			[this](Softs::FSolverParticlesRange& Particles, const Softs::FSolverReal Dt)
-		{
-			XPlanarConstraints->Apply(Particles, Dt, CollisionConstraint->GetCollisionIndices(), CollisionConstraint->GetCollisionTargets(), CollisionConstraint->GetCollisionNormalsOutput());
-		});
-
-		RuleCreator.AddPostInitialGuessParallelInitRule(
-			[this](const Softs::FSolverParticlesRange& Particles, const Softs::FSolverReal Dt, const Softs::ESolverMode SolverMode)
-		{
-			XPlanarConstraints->Init(Particles, Dt, CollisionConstraint->GetCollisionIndices().Num());
-		});
-	}
 }
 
 
@@ -2284,13 +2337,20 @@ void FClothConstraints::Update(
 	{
 		AnimDriveConstraints->SetProperties(ConfigProperties, WeightMaps);
 	}
-	if (SelfCollisionConstraints)
-	{
-		SelfCollisionConstraints->SetProperties(ConfigProperties, WeightMaps, FaceIntMaps);
-	}
+	bool bSelfCollisionInitCollidableSubMeshWillRebuild = false;
 	if (SelfCollisionInit)
 	{
 		SelfCollisionInit->SetProperties(ConfigProperties, FaceSets);
+		bSelfCollisionInitCollidableSubMeshWillRebuild = SelfCollisionInit->GetCollidableSubMeshDirty();
+	}
+	if (SelfCollisionConstraints)
+	{
+		SelfCollisionConstraints->SetProperties(ConfigProperties, WeightMaps, FaceIntMaps);
+		if (bSelfCollisionInitCollidableSubMeshWillRebuild)
+		{
+			// Kinematic collider indices will change on rebuild. Reset any existing constraints.
+			SelfCollisionConstraints->ResetKinematicCollider();
+		}
 	}
 	if (SelfCollisionSphereConstraints)
 	{
@@ -2329,6 +2389,10 @@ void FClothConstraints::Update(
 		CollisionConstraint->SetProperties(ConfigProperties);
 		CollisionConstraint->SetWriteDebugContacts(bWriteCCDContacts);
 	}
+	if (SkinnedTriangleCollisionsConstraint)
+	{
+		SkinnedTriangleCollisionsConstraint->SetProperties(ConfigProperties, WeightMaps);
+	}
 }
 
 // Deprecated 5.4
diff --git a/Engine/Plugins/ChaosCloth/Source/ChaosCloth/Private/ChaosCloth/ChaosClothVisualization.cpp b/Engine/Plugins/ChaosCloth/Source/ChaosCloth/Private/ChaosCloth/ChaosClothVisualization.cpp
--- a/Engine/Plugins/ChaosCloth/Source/ChaosCloth/Private/ChaosCloth/ChaosClothVisualization.cpp
+++ b/Engine/Plugins/ChaosCloth/Source/ChaosCloth/Private/ChaosCloth/ChaosClothVisualization.cpp
@@ -25,6 +25,7 @@
 #include "Chaos/PBDSphericalConstraint.h"
 #include "Chaos/PBDSpringConstraints.h"
 #include "Chaos/PBDTriangleMeshCollisions.h"
+#include "Chaos/SkinnedTriangleMesh.h"
 #include "Chaos/SoftsExternalForces.h"
 #include "Chaos/SoftsMultiResConstraints.h"
 #include "Chaos/XPBDBendingConstraints.h"
@@ -968,6 +969,45 @@ namespace Chaos
 		}
 	}
 
+	static void DrawSkinnedTriangleMesh(FPrimitiveDrawInterface* PDI, const FMaterialRenderProxy* MaterialRenderProxy, const FClothingSimulationSolver* Solver, const FLinearColor& Color, const FSkinnedTriangleMesh& Mesh)
+	{
+#if WITH_EDITOR
+		if (PDI && MaterialRenderProxy)
+		{
+			const FColor ColorRGB = Color.ToFColorSRGB();
+			FDynamicMeshBuilder MeshBuilder(PDI->View->GetFeatureLevel());
+			const TArray<Softs::FSolverVec3>& Positions = Mesh.GetLocalPositions();
+			for (const Softs::FSolverVec3& Pos : Positions)
+			{
+				MeshBuilder.AddVertex(FDynamicMeshVertex(FVector3f(Pos), FVector2f(0.f), ColorRGB));
+			}
+			const TArray<TVec3<int32>>& Elements = Mesh.GetTriangleMesh().GetElements();
+			for (const TVec3<int32>& Element : Elements)
+			{
+				MeshBuilder.AddTriangle(Element[0], Element[1], Element[2]);
+			}
+
+			const FMatrix LocalSimSpaceToWorld = GetLocalSpaceToWorld(*Solver);
+			MeshBuilder.Draw(PDI, LocalSimSpaceToWorld, MaterialRenderProxy, SDPG_World, false, false);
+		}
+		else
+#endif
+		{
+			const TArray<Softs::FSolverVec3>& Positions = Mesh.GetLocalPositions();
+			const TArray<TVec3<int32>>& Elements = Mesh.GetTriangleMesh().GetElements();
+			for (const TVec3<int32>& Element : Elements)
+			{
+				const FVector Pos0 = GetWorldPosition(*Solver, Positions[Element.X]);
+				const FVector Pos1 = GetWorldPosition(*Solver, Positions[Element.Y]);
+				const FVector Pos2 = GetWorldPosition(*Solver, Positions[Element.Z]);
+
+				DrawLine(PDI, Pos0, Pos1, Color);
+				DrawLine(PDI, Pos1, Pos2, Color);
+				DrawLine(PDI, Pos2, Pos0, Color);
+			}
+		}
+	}
+
 	void FClothVisualizationNoGC::DrawBounds(FPrimitiveDrawInterface* PDI) const
 	{
 		if (!Solver)
@@ -1439,7 +1479,7 @@ namespace Chaos
 		}
 	}
 
-	void FClothVisualizationNoGC::DrawCollision(FPrimitiveDrawInterface* PDI) const
+	void FClothVisualizationNoGC::DrawCollision(FPrimitiveDrawInterface* PDI, bool bWireframe) const
 	{
 		if (!Solver)
 		{
@@ -1447,7 +1487,7 @@ namespace Chaos
 		}
 
 		auto DrawCollision =
-			[this, PDI](const FClothingSimulationCollider* Collider, const FClothingSimulationCloth* Cloth, FClothingSimulationCollider::ECollisionDataType CollisionDataType)
+			[this, PDI, bWireframe](const FClothingSimulationCollider* Collider, const FClothingSimulationCloth* Cloth, FClothingSimulationCollider::ECollisionDataType CollisionDataType)
 			{
 				static const FLinearColor GlobalColor(FColor::Cyan);
 				static const FLinearColor DynamicColor(FColor::Orange);
@@ -1567,6 +1607,18 @@ namespace Chaos
 								DrawMLLevelSet(PDI, CombinedTransform, MaterialRenderProxy, MLLevelSet);
 							}
 							break;
+						case ImplicitObjectType::SkinnedTriangleMesh:
+						{
+							const FSkinnedTriangleMesh& Mesh = Object->GetObjectChecked<FSkinnedTriangleMesh>();
+							const FMaterialRenderProxy* MaterialRenderProxy =
+#if WITH_EDITOR
+								!bWireframe && FMaterials::GetInstance().GetClothMaterialColor() ? FMaterials::GetInstance().GetClothMaterialColor()->GetRenderProxy() : nullptr;
+#else
+								nullptr;
+#endif
+							DrawSkinnedTriangleMesh(PDI, MaterialRenderProxy, Solver, Color, Mesh);
+
+						}
 						default:
 							DrawCoordinateSystem(PDI, Rotation, Position);  // Draw everything else as a coordinate for now
 							break;
diff --git a/Engine/Plugins/ChaosCloth/Source/ChaosCloth/Private/ChaosCloth/ChaosClothingSimulation.cpp b/Engine/Plugins/ChaosCloth/Source/ChaosCloth/Private/ChaosCloth/ChaosClothingSimulation.cpp
--- a/Engine/Plugins/ChaosCloth/Source/ChaosCloth/Private/ChaosCloth/ChaosClothingSimulation.cpp
+++ b/Engine/Plugins/ChaosCloth/Source/ChaosCloth/Private/ChaosCloth/ChaosClothingSimulation.cpp
@@ -62,6 +62,9 @@ namespace Chaos
 bool bClothUseTimeStepSmoothing = true;
 FAutoConsoleVariableRef CVarClothUseTimeStepSmoothing(TEXT("p.ChaosCloth.UseTimeStepSmoothing"), bClothUseTimeStepSmoothing, TEXT("Use time step smoothing to avoid jitter during drastic changes in time steps."));
 
+bool bClothUseLegacyEvolution = true;
+FAutoConsoleVariableRef CVarClothUseLegacyEvolution(TEXT("p.ChaosCloth.UseLegacyEvolution"), bClothUseLegacyEvolution, TEXT("Use legacy underlying evolution as opposed to the new evolution shared with the ChaosClothAsset."));
+
 #if CHAOS_DEBUG_DRAW
 namespace ClothingSimulationCVar
 {
@@ -136,6 +139,7 @@ namespace ClothingSimulationConsole
 					bChaos_XPBDSpring_ISPC_Enabled =
 					bChaos_XPBDBending_ISPC_Enabled =
 					bChaos_MultiRes_ISPC_Enabled =
+					bChaos_KinematicTriangleMesh_ISPC_Enabled =
 					bEnableISPC;
 				return;
 			}
@@ -291,7 +295,8 @@ PRAGMA_ENABLE_DEPRECATION_WARNINGS
 void FClothingSimulation::Initialize()
 {
 	// Create solver
-	Solver = MakeUnique<FClothingSimulationSolver>();
+	constexpr FClothingSimulationConfig* InConfig = nullptr;
+	Solver = MakeUnique<FClothingSimulationSolver>(InConfig, bClothUseLegacyEvolution);
 
 	// Assign the solver to the visualization
 	Visualization.SetSolver(Solver.Get());
diff --git a/Engine/Plugins/ChaosCloth/Source/ChaosCloth/Private/ChaosCloth/ChaosClothingSimulationCloth.cpp b/Engine/Plugins/ChaosCloth/Source/ChaosCloth/Private/ChaosCloth/ChaosClothingSimulationCloth.cpp
--- a/Engine/Plugins/ChaosCloth/Source/ChaosCloth/Private/ChaosCloth/ChaosClothingSimulationCloth.cpp
+++ b/Engine/Plugins/ChaosCloth/Source/ChaosCloth/Private/ChaosCloth/ChaosClothingSimulationCloth.cpp
@@ -331,9 +331,23 @@ void FClothingSimulationCloth::FLODData::Add(FClothingSimulationSolver* Solver,
 	// Setup solver constraints
 	FClothConstraints& ClothConstraints = Solver->GetClothConstraints(ParticleRangeId);
 
+	TArray<const FClothComplexColliders*> ComplexColliders;
+	for (const FClothingSimulationCollider* const Collider : Cloth->GetColliders())
+	{
+		check(Collider);
+		for (FClothingSimulationCollider::ECollisionDataType CollisionType : TEnumRange<FClothingSimulationCollider::ECollisionDataType>())
+		{
+			const int32 CollisionRangeId = Collider->GetCollisionRangeId(Solver, Cloth, CollisionType);
+			if (const FClothComplexColliders* const ComplexCollider = Solver->GetComplexColliders(CollisionRangeId))
+			{
+				ComplexColliders.Add(ComplexCollider);
+			}
+		}
+	}
+
 	// Create constraints
 	const bool bEnabled = false;  // Set constraint disabled by default
-	ClothConstraints.AddRules(ConfigProperties, TriangleMesh, &PatternData, WeightMaps, VertexSets, FaceSets, FaceIntMaps, Tethers, MeshScale, bEnabled, CoarseLODTriangleMesh, CoarseLODParticleRangeId, FineLODMultiResConstraint);
+	ClothConstraints.AddRules(ConfigProperties, TriangleMesh, &PatternData, WeightMaps, VertexSets, FaceSets, FaceIntMaps, Tethers, MeshScale, bEnabled, CoarseLODTriangleMesh, CoarseLODParticleRangeId, FineLODMultiResConstraint, ComplexColliders);
 
 	// Update LOD stats
 	const TConstArrayView<Softs::FSolverReal> InvMasses(Solver->GetParticleInvMasses(ParticleRangeId), NumParticles);
@@ -605,6 +619,12 @@ void FClothingSimulationCloth::Add(FClothingSimulationSolver* Solver)
 {
 	check(Solver);
 
+	// Add Colliders. Do this first because the colliders need to be there to create constraints inside LODData::Add
+	for (FClothingSimulationCollider* Collider : Colliders)
+	{
+		Collider->Add(Solver, this);
+	}
+
 	// Can't add a cloth twice to the same solver
 	check(!LODIndices.Find(Solver));
 
@@ -623,12 +643,6 @@ void FClothingSimulationCloth::Add(FClothingSimulationSolver* Solver)
 	{
 		LODData[Index]->Add(Solver, this, Index);
 	}
-
-	// Add Colliders
-	for (FClothingSimulationCollider* Collider : Colliders)
-	{
-		Collider->Add(Solver, this);
-	}
 }
 
 void FClothingSimulationCloth::Remove(FClothingSimulationSolver* Solver)
diff --git a/Engine/Plugins/ChaosCloth/Source/ChaosCloth/Private/ChaosCloth/ChaosClothingSimulationCollider.cpp b/Engine/Plugins/ChaosCloth/Source/ChaosCloth/Private/ChaosCloth/ChaosClothingSimulationCollider.cpp
--- a/Engine/Plugins/ChaosCloth/Source/ChaosCloth/Private/ChaosCloth/ChaosClothingSimulationCollider.cpp
+++ b/Engine/Plugins/ChaosCloth/Source/ChaosCloth/Private/ChaosCloth/ChaosClothingSimulationCollider.cpp
@@ -1,5 +1,6 @@
 // Copyright Epic Games, Inc. All Rights Reserved.
 #include "ChaosCloth/ChaosClothingSimulationCollider.h"
+#include "ChaosCloth/ChaosClothComplexColliders.h"
 #include "ChaosCloth/ChaosClothingSimulationSolver.h"
 #include "ChaosCloth/ChaosClothingSimulationCloth.h"
 #include "ChaosCloth/ChaosClothingSimulationMesh.h"
@@ -45,11 +46,11 @@ struct FClothingSimulationCollider::FLODData
 		FClothingSimulationSolver* Solver,
 		FClothingSimulationCloth* Cloth,
 		const FClothCollisionData& InClothCollisionData,
-		const TArray<FLevelSetCollisionData>& InLevelSetCollisionData,
-		const TArray<FSkinnedLevelSetCollisionData>& InSkinnedLevelSetCollisionData,
-		const TArray<FMLLevelSetCollisionData>& InMLLevelSetCollisionData,
+		const FClothCollisionDataExtended& InClothCollisionDataExtended,
 		const FReal InScale = 1.f,
-		const TArray<int32>& UsedBoneIndices = TArray<int32>());
+		const TArray<int32>& UsedBoneIndices = TArray<int32>(),
+		const TArray<int32>& UsedSubBoneIndices = TArray<int32>());
+
 	void Remove(FClothingSimulationSolver* Solver, FClothingSimulationCloth* Cloth);
 
 	void Update(FClothingSimulationSolver* Solver, FClothingSimulationCloth* Cloth, const FTransform& ComponentTransform, const TArray<FTransform>& BoneTransforms);
@@ -68,11 +69,10 @@ void FClothingSimulationCollider::FLODData::Add(
 	FClothingSimulationSolver* Solver,
 	FClothingSimulationCloth* Cloth,
 	const FClothCollisionData& InClothCollisionData,
-	const TArray<FLevelSetCollisionData>& InLevelSetCollisionData,
-	const TArray<FSkinnedLevelSetCollisionData>& InSkinnedLevelSetCollisionData,
-	const TArray<FMLLevelSetCollisionData>& InMLLevelSetCollisionData,
+	const FClothCollisionDataExtended& InClothCollisionDataExtended,
 	const FReal InComponentScale,
-	const TArray<int32>& UsedBoneIndices)
+	const TArray<int32>& UsedBoneIndices,
+	const TArray<int32>& UsedSubBoneIndices)
 {
 	check(Solver);
 	check(Cloth);
@@ -107,22 +107,16 @@ void FClothingSimulationCollider::FLODData::Add(
 
 	const int32 NumConvexes = ClothCollisionData.Convexes.Num();
 	const int32 NumBoxes = ClothCollisionData.Boxes.Num();
-	const int32 NumLevelSets = InLevelSetCollisionData.Num();
-	const int32 NumSkinnedLevelSets = InSkinnedLevelSetCollisionData.Num();
-	int32 NumSkinnedLevelSetBones = 0;
-	const int32 NumMLLevelSets = InMLLevelSetCollisionData.Num();
-	int32 NumMLLevelSetProxyBones = 0;
-	for (const FSkinnedLevelSetCollisionData& SkinnedLevelSetData : InSkinnedLevelSetCollisionData)
-	{
-		NumSkinnedLevelSetBones += SkinnedLevelSetData.MappedSkinnedBones.Num();
-	}
+	const int32 NumLevelSets = InClothCollisionDataExtended.LevelSets.Num();
+	const int32 NumSkinnedLevelSets = InClothCollisionDataExtended.SkinnedLevelSets.Num();
+	const int32 NumMLLevelSets = InClothCollisionDataExtended.MLLevelSets.Num();
+	const int32 NumSkinnedTriangleMeshes = InClothCollisionDataExtended.SkinnedTriangleMeshes.Num();
 
-	for (const FMLLevelSetCollisionData& MLLevelSetData : InMLLevelSetCollisionData)
-	{
-		NumMLLevelSetProxyBones += MLLevelSetData.MappedActiveBoneIndices.Num();
-	}
+	const bool bSupportsComplexColliders = !Solver->IsLegacySolver();
+	const int32 NumSimpleColliders = NumSpheres + NumCapsules + NumConvexes + NumBoxes + NumLevelSets;
+	const int32 NumComplexColliders = NumSkinnedLevelSets + NumMLLevelSets + NumSkinnedTriangleMeshes;
 
-	NumGeometries = NumSpheres + NumCapsules + NumConvexes + NumBoxes + NumLevelSets + NumSkinnedLevelSets + NumSkinnedLevelSetBones + NumMLLevelSets + NumMLLevelSetProxyBones;
+	NumGeometries = NumSimpleColliders + (bSupportsComplexColliders ? NumComplexColliders : 0);
 	if(NumGeometries == 0)
 	{
 		return;
@@ -332,7 +326,8 @@ void FClothingSimulationCollider::FLODData::Add(
 			Solver->SetCollisionGeometry(CollisionRangeId, Index, MakeImplicitObjectPtr<TBox<FReal, 3>>(-HalfExtents, HalfExtents));
 		}
 	}
-
+	
+	// Level sets
 	const int32 LevelSetIndexOffset = BoxIndexOffset + NumBoxes;
 	if (NumLevelSets != 0)
 	{
@@ -341,99 +336,67 @@ void FClothingSimulationCollider::FLODData::Add(
 			// Always initialize the collision particle transforms before setting any geometry as otherwise NaNs gets detected during the bounding box updates
 			BaseTransforms[Index] = Softs::FSolverRigidTransform3::Identity;
 
-			BoneIndices[Index] = GetMappedBoneIndex(UsedBoneIndices, InLevelSetCollisionData[LevelSetIndex].BoneIndex);
+			BoneIndices[Index] = GetMappedBoneIndex(UsedBoneIndices, InClothCollisionDataExtended.LevelSets[LevelSetIndex].BoneIndex);
 			UE_LOG(LogChaosCloth, VeryVerbose, TEXT("Found collision level set on bone index %d."), BoneIndices[Index]);
 
 			// Setup the collision particle geometry
-			FImplicitObjectPtr TransformedLevelSet = MakeImplicitObjectPtr<TImplicitObjectTransformed<FReal, 3>>(InLevelSetCollisionData[LevelSetIndex].LevelSet->DeepCopyGeometry(),
-				TRigidTransform<FReal, 3>(InLevelSetCollisionData[LevelSetIndex].Transform));
+			FImplicitObjectPtr TransformedLevelSet = MakeImplicitObjectPtr<TImplicitObjectTransformed<FReal, 3>>(InClothCollisionDataExtended.LevelSets[LevelSetIndex].LevelSet->DeepCopyGeometry(),
+				TRigidTransform<FReal, 3>(InClothCollisionDataExtended.LevelSets[LevelSetIndex].Transform));
 			Solver->SetCollisionGeometry(CollisionRangeId, Index, MoveTemp(TransformedLevelSet));
 		}
 	}
 
-	const int32 SkinnedLevelSetIndexOffset = LevelSetIndexOffset + NumLevelSets;
-	if (NumSkinnedLevelSets != 0)
-	{
-		int32 Index = SkinnedLevelSetIndexOffset;
-		for (const FSkinnedLevelSetCollisionData& SkinnedCollisionData : InSkinnedLevelSetCollisionData)
-		{
-			// Add bone proxies first so they get updated first?
-			TArray<int32> SolverBoneIndices;
-			SolverBoneIndices.Reserve(SkinnedCollisionData.MappedSkinnedBones.Num());
-			const int32 GlobalBoneOffset = Solver->IsLegacySolver() ? CollisionRangeId : 0;
-			for (int32 MappedSubBoneIndex : SkinnedCollisionData.MappedSkinnedBones)
-			{
-				// Always initialize the collision particle transforms before setting any geometry as otherwise NaNs gets detected during the bounding box updates
-				BaseTransforms[Index] = Softs::FSolverRigidTransform3::Identity;
-
-				BoneIndices[Index] = GetMappedBoneIndex(UsedBoneIndices, MappedSubBoneIndex);
-				UE_LOG(LogChaosCloth, VeryVerbose, TEXT("Found collision skinned level set sub-bone on bone index %d."), BoneIndices[Index]);
-
-				FImplicitObjectPtr BoneProxy = MakeImplicitObjectPtr<FWeightedLatticeBoneProxy>();
-				Solver->SetCollisionGeometry(CollisionRangeId, Index, MoveTemp(BoneProxy));
-
-				SolverBoneIndices.Add(GlobalBoneOffset + Index);
+	UE_LOG(LogChaosCloth, VeryVerbose, TEXT("Added collisions: %d spheres, %d capsules, %d convexes, %d boxes, %d level sets."), NumSpheres, NumCapsules, NumConvexes, NumBoxes, NumLevelSets);
 
-				++Index;
-			}
+	if (bSupportsComplexColliders)
+	{
+		FClothComplexColliders* const ComplexColliders = Solver->GetComplexColliders(CollisionRangeId);
+		check(ComplexColliders);
+		ComplexColliders->Reset();
+		ComplexColliders->AddSubBoneIndices(UsedSubBoneIndices);
 
+		const int32 SkinnedLevelSetIndexOffset = LevelSetIndexOffset + NumLevelSets;
+		for (int32 Index = SkinnedLevelSetIndexOffset, SkinnedLevelSetIndex = 0; SkinnedLevelSetIndex < NumSkinnedLevelSets; ++Index, ++SkinnedLevelSetIndex)
+		{
 			// Always initialize the collision particle transforms before setting any geometry as otherwise NaNs gets detected during the bounding box updates
 			BaseTransforms[Index] = Softs::FSolverRigidTransform3::Identity;
 
-			BoneIndices[Index] = GetMappedBoneIndex(UsedBoneIndices, SkinnedCollisionData.BoneIndex);
+			BoneIndices[Index] = GetMappedBoneIndex(UsedBoneIndices, InClothCollisionDataExtended.SkinnedLevelSets[SkinnedLevelSetIndex].BoneIndex);
 			UE_LOG(LogChaosCloth, VeryVerbose, TEXT("Found collision skinned level set on bone index %d."), BoneIndices[Index]);
 
-			FImplicitObjectPtr SkinnedLevelSet = SkinnedCollisionData.WeightedLevelSet->DeepCopyGeometry();
-			SkinnedLevelSet->GetObjectChecked<TWeightedLatticeImplicitObject<FLevelSet>>().SetSolverBoneIndices(MoveTemp(SolverBoneIndices));
-
+			FImplicitObjectPtr SkinnedLevelSet = InClothCollisionDataExtended.SkinnedLevelSets[SkinnedLevelSetIndex].WeightedLevelSet->DeepCopyGeometry();
+			ComplexColliders->AddSkinnedLevelSet(Index, InClothCollisionDataExtended.SkinnedLevelSets[SkinnedLevelSetIndex].MappedSkinnedBones, SkinnedLevelSet);
 			Solver->SetCollisionGeometry(CollisionRangeId, Index, MoveTemp(SkinnedLevelSet));
-
-			++Index;
 		}
-	}
 
-	const int32 MLLevelSetIndexOffset = SkinnedLevelSetIndexOffset + NumSkinnedLevelSets + NumSkinnedLevelSetBones; 
-	if (NumMLLevelSets != 0)
-	{
-		int32 Index = MLLevelSetIndexOffset;
-		for (const FMLLevelSetCollisionData& MLLevelSetCollisionData : InMLLevelSetCollisionData)
+		const int32 MLLevelSetIndexOffset = SkinnedLevelSetIndexOffset + NumSkinnedLevelSets;
+		for (int32 Index = MLLevelSetIndexOffset, MLLevelsetIndex = 0; MLLevelsetIndex < NumMLLevelSets; ++Index, ++MLLevelsetIndex)
 		{
-			TArray<int32> SolverActiveBoneIndices;
-			SolverActiveBoneIndices.Reserve(MLLevelSetCollisionData.MappedActiveBoneIndices.Num());
-			const int32 GlobalBoneOffset = Solver->IsLegacySolver() ? CollisionRangeId : 0;
-			
-			for (int32 MappedActiveBoneIndex : MLLevelSetCollisionData.MappedActiveBoneIndices)
-			{
-				// Always initialize the collision particle transforms before setting any geometry as otherwise NaNs gets detected during the bounding box updates
-				BaseTransforms[Index] = Softs::FSolverRigidTransform3::Identity;
-
-				BoneIndices[Index] = GetMappedBoneIndex(UsedBoneIndices, MappedActiveBoneIndex);
-				UE_LOG(LogChaosCloth, VeryVerbose, TEXT("Found collision ml level set active bone on bone index %d."), BoneIndices[Index]);
-
-				FImplicitObjectPtr BoneProxy = MakeImplicitObjectPtr<FMLLevelSetBoneProxy>();
-				Solver->SetCollisionGeometry(CollisionRangeId, Index, MoveTemp(BoneProxy));
+			// Always initialize the collision particle transforms before setting any geometry as otherwise NaNs gets detected during the bounding box updates
+			BaseTransforms[Index] = Softs::FSolverRigidTransform3::Identity;
 
-				SolverActiveBoneIndices.Add(GlobalBoneOffset + Index);
+			BoneIndices[Index] = GetMappedBoneIndex(UsedBoneIndices, InClothCollisionDataExtended.MLLevelSets[MLLevelsetIndex].BoneIndex);
+			UE_LOG(LogChaosCloth, VeryVerbose, TEXT("Found collision ml level set on bone index %d."), BoneIndices[Index]);
 
-				++Index;
-			}
+			FImplicitObjectPtr MLLevelSet = InClothCollisionDataExtended.MLLevelSets[MLLevelsetIndex].MLLevelSet->CopyGeometry();
+			ComplexColliders->AddMLLevelSet(Index, InClothCollisionDataExtended.MLLevelSets[MLLevelsetIndex].MappedActiveBoneIndices, MLLevelSet);
+			Solver->SetCollisionGeometry(CollisionRangeId, Index, MoveTemp(MLLevelSet));
+		}
 
+		const int32 SkinnedTriMeshIndexOffset = MLLevelSetIndexOffset + NumMLLevelSets;
+		for (int32 Index = SkinnedTriMeshIndexOffset, SkinnedTriMeshIndex = 0; SkinnedTriMeshIndex < NumSkinnedTriangleMeshes; ++Index, ++SkinnedTriMeshIndex)
+		{
 			// Always initialize the collision particle transforms before setting any geometry as otherwise NaNs gets detected during the bounding box updates
 			BaseTransforms[Index] = Softs::FSolverRigidTransform3::Identity;
 
-			BoneIndices[Index] = GetMappedBoneIndex(UsedBoneIndices, MLLevelSetCollisionData.ParentBoneIndex);
+			BoneIndices[Index] = GetMappedBoneIndex(UsedBoneIndices, InClothCollisionDataExtended.SkinnedTriangleMeshes[SkinnedTriMeshIndex].BoneIndex);
 			UE_LOG(LogChaosCloth, VeryVerbose, TEXT("Found collision ml level set on bone index %d."), BoneIndices[Index]);
-			
-			FImplicitObjectPtr MLLevelSet = MLLevelSetCollisionData.MLLevelSet->CopyGeometry();
-			MLLevelSet->GetObjectChecked<FMLLevelSet>().SetSolverActiveBoneIndices(MoveTemp(SolverActiveBoneIndices));
-			MLLevelSet->GetObjectChecked<FMLLevelSet>().SetSolverParentBoneIndex(GlobalBoneOffset + Index);
 
-			Solver->SetCollisionGeometry(CollisionRangeId, Index, MoveTemp(MLLevelSet));
-			++Index;
+			FImplicitObjectPtr SkinnedTriMesh = InClothCollisionDataExtended.SkinnedTriangleMeshes[SkinnedTriMeshIndex].SkinnedTriangleMesh->CopyGeometry();
+			ComplexColliders->AddSkinnedTriangleMesh(Index, InClothCollisionDataExtended.SkinnedTriangleMeshes[SkinnedTriMeshIndex].MappedSkinnedBones, SkinnedTriMesh);
+			Solver->SetCollisionGeometry(CollisionRangeId, Index, MoveTemp(SkinnedTriMesh));
 		}
 	}
-
-	UE_LOG(LogChaosCloth, VeryVerbose, TEXT("Added collisions: %d spheres, %d capsules, %d convexes, %d boxes, %d level sets."), NumSpheres, NumCapsules, NumConvexes, NumBoxes, NumLevelSets);
 }
 
 void FClothingSimulationCollider::FLODData::Remove(FClothingSimulationSolver* Solver, FClothingSimulationCloth* Cloth)
@@ -473,6 +436,11 @@ void FClothingSimulationCollider::FLODData::Update(
 				BaseTransforms[Index] * Softs::FSolverTransform3(BoneTransforms[BoneIndex]) * ComponentToLocalSpace :  // LWC requires the BoneTransform cast to the solver underlying type
 				BaseTransforms[Index] * ComponentToLocalSpace;  // External collisions don't map to a bone
 		}
+
+		if (FClothComplexColliders* const ComplexColliders = Solver->GetComplexColliders(CollisionRangeId))
+		{
+			ComplexColliders->Update(ComponentToLocalSpace, BoneTransforms, CollisionTransforms);
+		}
 	}
 }
 
@@ -495,6 +463,11 @@ void FClothingSimulationCollider::FLODData::ResetStartPose(FClothingSimulationSo
 	{
 		const int32 CollisionRangeId = CollisionRangeIds.FindChecked(FSolverClothPair(Solver, Cloth));
 		Solver->ResetCollisionStartPose(CollisionRangeId, NumGeometries);
+
+		if (FClothComplexColliders* const ComplexColliders = Solver->GetComplexColliders(CollisionRangeId))
+		{
+			ComplexColliders->ResetStartPose();
+		}
 	}
 }
 
@@ -546,29 +519,46 @@ void FClothingSimulationCollider::ExtractPhysicsAssetCollision(
 	const UPhysicsAsset* PhysicsAsset,
 	const FReferenceSkeleton* ReferenceSkeleton,
 	FClothCollisionData& ClothCollisionData,
+	PRAGMA_DISABLE_DEPRECATION_WARNINGS
 	TArray<FLevelSetCollisionData>& LevelSetCollisions,
 	TArray<FSkinnedLevelSetCollisionData>& SkinnedLevelSetCollisions,
+	PRAGMA_ENABLE_DEPRECATION_WARNINGS
 	TArray<int32>& UsedBoneIndices,
 	bool bUseSphylsOnly,
 	bool bSkipMissingBones)
 {
-	TArray<FMLLevelSetCollisionData> MLLevelSetCollisions = {};
-	ExtractPhysicsAssetCollision(PhysicsAsset, ReferenceSkeleton, ClothCollisionData, LevelSetCollisions, SkinnedLevelSetCollisions, MLLevelSetCollisions, UsedBoneIndices, bUseSphylsOnly, bSkipMissingBones);
+	FClothCollisionDataExtended ExtendedCollisionData;
+	TArray<int32> UsedSubBoneIndices;
+	ExtractPhysicsAssetCollision(PhysicsAsset, ReferenceSkeleton, ClothCollisionData, ExtendedCollisionData, UsedBoneIndices, UsedSubBoneIndices, bUseSphylsOnly, bSkipMissingBones);
+	LevelSetCollisions = MoveTemp(ExtendedCollisionData.LevelSets);
+	SkinnedLevelSetCollisions = MoveTemp(ExtendedCollisionData.SkinnedLevelSets);
+
+	// Append SubBones and reindex
+	const int32 Offset = UsedBoneIndices.Num();
+	UsedBoneIndices.Append(UsedSubBoneIndices);
+	for (Chaos::FSkinnedLevelSetCollisionData& SkinnedLevelSetData : SkinnedLevelSetCollisions)
+	{
+		for (int32& MappedBone : SkinnedLevelSetData.MappedSkinnedBones)
+		{
+			MappedBone += Offset;
+		}
+	}
 }
 
 void FClothingSimulationCollider::ExtractPhysicsAssetCollision(
 	const UPhysicsAsset* PhysicsAsset,
 	const FReferenceSkeleton* ReferenceSkeleton,
 	FClothCollisionData& ClothCollisionData,
-	TArray<FLevelSetCollisionData>& LevelSetCollisions,
-	TArray<FSkinnedLevelSetCollisionData>& SkinnedLevelSetCollisions,
-	TArray<FMLLevelSetCollisionData>& MLLevelSetCollisions,
+	FClothCollisionDataExtended& ClothCollisionDataExtended,
 	TArray<int32>& UsedBoneIndices,
+	TArray<int32>& UsedSubBoneIndices,
 	bool bUseSphylsOnly,
 	bool bSkipMissingBones)
 {
 	ClothCollisionData.Reset();
+	ClothCollisionDataExtended.Reset();
 	UsedBoneIndices.Reset();
+	UsedSubBoneIndices.Reset();
 
 	if (PhysicsAsset)
 	{
@@ -734,7 +724,7 @@ void FClothingSimulationCollider::ExtractPhysicsAssetCollision(
 			{
 				if (LevelSetElem.GetLevelSet().IsValid())
 				{
-					LevelSetCollisions.Add({ LevelSetElem.GetLevelSet(), LevelSetElem.GetTransform(), MappedBoneIndex });
+					ClothCollisionDataExtended.LevelSets.Emplace( LevelSetElem.GetLevelSet(), LevelSetElem.GetTransform(), MappedBoneIndex );
 				}
 			}
 
@@ -748,10 +738,10 @@ void FClothingSimulationCollider::ExtractPhysicsAssetCollision(
 					for (const FName& SkinnedBoneName : SkinnedBones)
 					{
 						const int32 SkinnedBoneIndex = ReferenceSkeleton ? ReferenceSkeleton->FindBoneIndex(SkinnedBoneName) : INDEX_NONE;
-						const int32 MappedSkinnedBoneIndex = UsedBoneIndices.Add(SkinnedBoneIndex);
+						const int32 MappedSkinnedBoneIndex = UsedSubBoneIndices.AddUnique(SkinnedBoneIndex);
 						MappedSkinnedBones.Add(MappedSkinnedBoneIndex);
 					}
-					SkinnedLevelSetCollisions.Add({ SkinnedLevelSetElem.WeightedLevelSet(), MappedBoneIndex, MoveTemp(MappedSkinnedBones)});
+					ClothCollisionDataExtended.SkinnedLevelSets.Emplace(SkinnedLevelSetElem.WeightedLevelSet(), MappedBoneIndex, MoveTemp(MappedSkinnedBones));
 				}
 			}
 
@@ -765,10 +755,27 @@ void FClothingSimulationCollider::ExtractPhysicsAssetCollision(
 					for (const FName& ActiveBoneName : ActiveBoneNames)
 					{
 						const int32 ActiveBoneIndex = ReferenceSkeleton ? ReferenceSkeleton->FindBoneIndex(ActiveBoneName) : INDEX_NONE;
-						const int32 MappedActiveBoneIndex = UsedBoneIndices.Add(ActiveBoneIndex);
+						const int32 MappedActiveBoneIndex = UsedSubBoneIndices.AddUnique(ActiveBoneIndex);
 						MappedActiveBones.Add(MappedActiveBoneIndex);
 					}
-					MLLevelSetCollisions.Add({ MLLevelSetElem.GetMLLevelSet(), MappedBoneIndex, MoveTemp(MappedActiveBones)});
+					ClothCollisionDataExtended.MLLevelSets.Emplace( MLLevelSetElem.GetMLLevelSet(), MappedBoneIndex, MoveTemp(MappedActiveBones));
+				}
+			}
+
+			for (const FKSkinnedTriangleMeshElem& SkinnedTriangleMeshElem : AggGeom.SkinnedTriangleMeshElems)
+			{
+				if (SkinnedTriangleMeshElem.GetSkinnedTriangleMesh().IsValid())
+				{
+					TArray<int32> MappedSkinnedBones;
+					const TArray<FName>& SkinnedBones = SkinnedTriangleMeshElem.GetSkinnedTriangleMesh()->GetUsedBones();
+					MappedSkinnedBones.Reserve(SkinnedBones.Num());
+					for (const FName& SkinnedBoneName : SkinnedBones)
+					{
+						const int32 SkinnedBoneIndex = ReferenceSkeleton ? ReferenceSkeleton->FindBoneIndex(SkinnedBoneName) : INDEX_NONE;
+						const int32 MappedSkinnedBoneIndex = UsedSubBoneIndices.AddUnique(SkinnedBoneIndex);
+						MappedSkinnedBones.Add(MappedSkinnedBoneIndex);
+					}
+					ClothCollisionDataExtended.SkinnedTriangleMeshes.Emplace(SkinnedTriangleMeshElem.GetSkinnedTriangleMesh(), MappedBoneIndex, MoveTemp(MappedSkinnedBones));
 				}
 			}
 		}  // End for PhysAsset->SkeletalBodySetups
@@ -843,13 +850,12 @@ PRAGMA_ENABLE_DEPRECATION_WARNINGS
 
 	// Create physics asset collisions, this will affect all LODs, so store it at index 0
 	FClothCollisionData PhysicsAssetCollisionData;
-	TArray<FLevelSetCollisionData> LevelSetCollisions;
-	TArray<FSkinnedLevelSetCollisionData> SkinnedLevelSetCollisions;
-	TArray<FMLLevelSetCollisionData> MLLevelSetCollisions;
+	FClothCollisionDataExtended PhysicsAssetCollisionDataExtended;
 	TArray<int32> UsedBoneIndices;
-	ExtractPhysicsAssetCollision(PhysicsAsset, ReferenceSkeleton, PhysicsAssetCollisionData, LevelSetCollisions, SkinnedLevelSetCollisions, MLLevelSetCollisions, UsedBoneIndices);
+	TArray<int32> UsedSubBoneIndices;
+	ExtractPhysicsAssetCollision(PhysicsAsset, ReferenceSkeleton, PhysicsAssetCollisionData, PhysicsAssetCollisionDataExtended, UsedBoneIndices, UsedSubBoneIndices);
 
-	LODData[(int32)ECollisionDataType::LODless]->Add(Solver, Cloth, PhysicsAssetCollisionData, LevelSetCollisions, SkinnedLevelSetCollisions, MLLevelSetCollisions, Scale, UsedBoneIndices);
+	LODData[(int32)ECollisionDataType::LODless]->Add(Solver, Cloth, PhysicsAssetCollisionData, PhysicsAssetCollisionDataExtended, Scale, UsedBoneIndices, UsedSubBoneIndices);
 }
 
 void FClothingSimulationCollider::Remove(FClothingSimulationSolver* Solver, FClothingSimulationCloth* Cloth)
@@ -874,12 +880,10 @@ void FClothingSimulationCollider::PreUpdate(FClothingSimulationSolver* Solver, F
 	const int32 ExternalCollisionNumGeometries = GetNumGeometries((int32)ECollisionDataType::External);
 	const int32 ExternalCollisionCollisionRangeId = GetCollisionRangeId(Solver, Cloth, (int32)ECollisionDataType::External);
 
-	// TODO: Get level sets?
-	const TArray<FLevelSetCollisionData> LevelSetCollisions;
-	const TArray<FSkinnedLevelSetCollisionData> SkinnedLevelSetCollisions;
-	const TArray<FMLLevelSetCollisionData> MLLevelSetCollisions;
+	// TODO: Handle external extended data 
+	// Will need to handle refreshing constraints that hold pointers to data in ComplexColliders
 
-	LODData[(int32)ECollisionDataType::External]->Add(Solver, Cloth, CollisionData ? *CollisionData : FClothCollisionData(), LevelSetCollisions, SkinnedLevelSetCollisions, MLLevelSetCollisions);
+	LODData[(int32)ECollisionDataType::External]->Add(Solver, Cloth, CollisionData ? *CollisionData : FClothCollisionData(), FClothCollisionDataExtended());
 	
 	// TODO: Find a better way in case the same number but different collisions are being re-added (hash collision data? Provide user dirty function?)
 	bHasExternalCollisionChanged =
diff --git a/Engine/Plugins/ChaosCloth/Source/ChaosCloth/Private/ChaosCloth/ChaosClothingSimulationSolver.cpp b/Engine/Plugins/ChaosCloth/Source/ChaosCloth/Private/ChaosCloth/ChaosClothingSimulationSolver.cpp
--- a/Engine/Plugins/ChaosCloth/Source/ChaosCloth/Private/ChaosCloth/ChaosClothingSimulationSolver.cpp
+++ b/Engine/Plugins/ChaosCloth/Source/ChaosCloth/Private/ChaosCloth/ChaosClothingSimulationSolver.cpp
@@ -1,6 +1,7 @@
 // Copyright Epic Games, Inc. All Rights Reserved.
 #include "ChaosCloth/ChaosClothingSimulationSolver.h"
 #include "ChaosCloth/ChaosClothingSimulationCloth.h"
+#include "ChaosCloth/ChaosClothComplexColliders.h"
 #include "ChaosCloth/ChaosClothingSimulationCollider.h"
 #include "ChaosCloth/ChaosClothingSimulationMesh.h"
 #include "ChaosCloth/ChaosClothingSimulationConfig.h"
@@ -9,9 +10,10 @@
 #include "Chaos/SoftsEvolution.h"
 #include "Chaos/PBDEvolution.h"
 #include "Chaos/PBDFlatWeightMap.h"
-#include "Chaos/WeightedLatticeImplicitObject.h"
 #include "Chaos/Levelset.h"
 #include "Chaos/MLLevelset.h" 
+#include "Chaos/SkinnedTriangleMesh.h"
+#include "Chaos/WeightedLatticeImplicitObject.h"
 #include "GeometryCollection/ManagedArrayCollection.h"
 #if INTEL_ISPC
 #include "ChaosClothingSimulationSolver.ispc.generated.h"
@@ -87,6 +89,7 @@ namespace ClothingSimulationSolverDefault
 	static const int32 MinNumSubsteps = 1;
 	static const FRealSingle DynamicSubstepDeltaTime = 0.f;
 	static const bool bEnableNumSelfCollisionSubsteps = false;
+	static const bool bUseSelfCollisionSubstepsForSkinnedTriangleMeshes = true;
 	static const int32 NumSelfCollisionSubsteps = 1;
 	static const FRealSingle SolverFrequency = 60.f;
 	static const FRealSingle SelfCollisionThickness = 2.f;
@@ -212,12 +215,17 @@ FClothingSimulationSolver::FClothingSimulationSolver(FClothingSimulationConfig*
 			[this](Softs::FSolverParticlesRange& ParticlesInput, const Softs::FSolverReal Dt, const Softs::FSolverReal LocalTime)
 		{
 			Softs::FPAndInvM* const PAndInvM = ParticlesInput.GetPAndInvM().GetData();
+			Softs::FSolverVec3* const V = ParticlesInput.GetV().GetData();
 			const Softs::FSolverVec3* const InterpPos = ParticlesInput.GetConstArrayView(InterpolatedAnimationPositions).GetData();
+			const Softs::FSolverVec3* const AnimV = ParticlesInput.GetConstArrayView(AnimationVelocities).GetData();
 			for (int32 Index = 0; Index < ParticlesInput.GetRangeSize(); ++Index)
 			{
 				if (PAndInvM[Index].InvM == (FSolverReal)0.)
 				{
 					PAndInvM[Index].P = InterpPos[Index];  // X is the step initial condition, here it's P that needs to be updated so that constraints works with the correct step target
+
+					// This could be set per frame since we currently use linear interpolation, but it's cleaner to set here.
+					V[Index] = AnimV[Index];
 				}
 			}
 		});
@@ -248,41 +256,12 @@ FClothingSimulationSolver::FClothingSimulationSolver(FClothingSimulationConfig*
 				const Softs::FSolverVec3 Axis = Delta.GetRotationAxis();
 				W[Index] = (Softs::FSolverVec3)Axis * Angle / Dt;
 				R[Index] = NewR;
+			}
 
-				if (TWeightedLatticeImplicitObject<FLevelSet>* SkinnedLevelSet =
-					const_cast<FImplicitObject*>(ParticlesInput.GetGeometry(Index).GetReference())->GetObject<TWeightedLatticeImplicitObject<FLevelSet>>	())
-				{
-					const TArray<int32>& SubBoneIndices = SkinnedLevelSet->GetSolverBoneIndices();
-					const FTransform RootTransformInv = TRigidTransform<FReal, 3>(X[Index], R[Index]).Inverse();
-					TArray<FTransform> SubBoneTransforms;
-					SubBoneTransforms.SetNum(SubBoneIndices.Num());
-					for (int32 SubBoneIdx = 0; SubBoneIdx < SubBoneIndices.Num(); ++SubBoneIdx)
-					{
-						const int32 SubBoneIndexLocal = SubBoneIndices[SubBoneIdx] - ParticlesInput.GetOffset();
-						checkSlow(SubBoneIndexLocal < Index);
-						const Softs::FSolverCollisionParticlesRange& ConstParticlesInput = ParticlesInput;
-						SubBoneTransforms[SubBoneIdx] = TRigidTransform<FReal, 3>(ParticlesInput.X(SubBoneIndexLocal), ConstParticlesInput.R(SubBoneIndexLocal)) * RootTransformInv;
-					}
-					SkinnedLevelSet->DeformPoints(SubBoneTransforms);
-					SkinnedLevelSet->UpdateSpatialHierarchy();
-				}
-				
-				if (FMLLevelSet* MLLevelSet =
-					const_cast<FImplicitObject*>(ParticlesInput.GetGeometry(Index).GetReference())->GetObject<FMLLevelSet>())
-				{
-					const TArray<int32>& ActiveBoneIndices = MLLevelSet->GetSolverActiveBoneIndices();
-					TArray<FTransform> ActiveBoneTransforms;
-					ActiveBoneTransforms.SetNum(ActiveBoneIndices.Num());
-					const FTransform RootTransform = TRigidTransform<FReal, 3>(X[Index], R[Index]);
-					for (int32 ActiveBoneIndex = 0; ActiveBoneIndex < ActiveBoneIndices.Num(); ++ActiveBoneIndex)
-					{
-						const int32 ActiveBoneIndexLocal = ActiveBoneIndices[ActiveBoneIndex] - ParticlesInput.GetOffset();
-						checkSlow(ActiveBoneIndexLocal < Index);
-						const FTransform ActiveBoneTransform = TRigidTransform<FReal, 3>(X[ActiveBoneIndexLocal], R[ActiveBoneIndexLocal]);
-						ActiveBoneTransforms[ActiveBoneIndex] = ActiveBoneTransform.GetRelativeTransform(RootTransform);
-					}
-					MLLevelSet->UpdateActiveBonesRelativeTransforms(ActiveBoneTransforms);
-				}
+			const int32 CollisionRangeId = ParticlesInput.GetRangeId();
+			if (FClothComplexColliders* const Colliders = GetComplexColliders(CollisionRangeId))
+			{
+				Colliders->KinematicUpdate(ParticlesInput, Dt, Alpha);
 			}
 		});
 	}
@@ -327,6 +306,8 @@ FClothingSimulationSolver::FClothingSimulationSolver(FClothingSimulationConfig*
 			[this](Softs::FSolverParticles& ParticlesInput, const Softs::FSolverReal Dt, const Softs::FSolverReal LocalTime, const int32 Index)
 		{
 			ParticlesInput.P(Index) = InterpolatedAnimationPositions[Index];  // X is the step initial condition, here it's P that needs to be updated so that constraints works with the correct step target
+			// This could be set per frame since we currently use linear interpolation, but it's cleaner to set here.
+			ParticlesInput.V(Index) = AnimationVelocities[Index];
 		});
 
 		PBDEvolution->SetCollisionKinematicUpdateFunction(
@@ -346,22 +327,6 @@ FClothingSimulationSolver::FClothingSimulationSolver(FClothingSimulationConfig*
 			const Softs::FSolverVec3 Axis = Delta.GetRotationAxis();
 			ParticlesInput.W(Index) = (Softs::FSolverVec3)Axis * Angle / Dt;
 			ParticlesInput.SetR(Index, NewR);
-
-			if (TWeightedLatticeImplicitObject<FLevelSet>* SkinnedLevelSet =
-				const_cast<FImplicitObject*>(ParticlesInput.GetGeometry(Index).GetReference())->GetObject<TWeightedLatticeImplicitObject<FLevelSet>>())
-			{
-				const TArray<int32>& SubBoneIndices = SkinnedLevelSet->GetSolverBoneIndices();
-				const FTransform RootTransformInv = TRigidTransform<FReal, 3>(ParticlesInput.GetX(Index), ParticlesInput.GetR(Index)).Inverse();
-				TArray<FTransform> SubBoneTransforms;
-				SubBoneTransforms.SetNum(SubBoneIndices.Num());
-				for (int32 SubBoneIdx = 0; SubBoneIdx < SubBoneIndices.Num(); ++SubBoneIdx)
-				{
-					checkSlow(SubBoneIndices[SubBoneIdx] < Index);
-					SubBoneTransforms[SubBoneIdx] = TRigidTransform<FReal, 3>(ParticlesInput.GetX(SubBoneIndices[SubBoneIdx]), ParticlesInput.GetR(SubBoneIndices[SubBoneIdx])) * RootTransformInv;
-				}
-				SkinnedLevelSet->DeformPoints(SubBoneTransforms);
-				SkinnedLevelSet->UpdateSpatialHierarchy();
-			}
 		});
 	}
 }
@@ -822,8 +787,21 @@ int32 FClothingSimulationSolver::AddCollisionParticles(int32 NumCollisionParticl
 			}
 			// Size has changed. Remove the recycled collision range and reallocate some new particles
 			Evolution->RemoveCollisionParticleRange(RecycledCollisionRangeId);
+			ComplexColliders.Remove(RecycledCollisionRangeId);
+			const TArray<int32>& GroupSoftBodies = Evolution->GetGroupSoftBodies(GroupId);
+			for (int32 SoftBodyId : GroupSoftBodies)
+			{
+				if (TUniquePtr<FClothConstraints>* ClothConstraints = ClothsConstraints.Find(SoftBodyId))
+				{
+					check(ClothConstraints);
+					(*ClothConstraints)->OnCollisionRangeRemoved(RecycledCollisionRangeId);
+				}
+			}
 		}
 		CollisionRangeId = Evolution->AddCollisionParticleRange(GroupId, NumCollisionParticles, bActivateFalse);
+
+		// Add an empty ComplexColliders
+		ComplexColliders.Emplace(CollisionRangeId, MakeUnique<FClothComplexColliders>(Evolution.Get(), CollisionRangeId));
 	}
 	else
 	{
@@ -1739,6 +1717,11 @@ void FClothingSimulationSolver::ApplyPreSimulationTransforms()
 					CollisionParticles.X(Index) = OldCollisionTransformsView[Index].GetTranslation();
 					CollisionParticles.SetR(Index, OldCollisionTransformsView[Index].GetRotation());
 				}
+				if (FClothComplexColliders* const Colliders = GetComplexColliders(CollisionRangeId))
+				{
+					Colliders->ApplyPreSimulationTransforms(GroupSpaceTransform, DeltaLocalSpaceLocation, CollisionParticles.GetConstArrayView(OldCollisionTransforms),
+						CollisionParticles.GetConstArrayView(CollisionTransforms), DeltaTime);
+				}
 			}
 		}, /*bForceSingleThreaded =*/ !bClothSolverParallelClothPreUpdate);
 	}
@@ -1865,13 +1848,22 @@ void FClothingSimulationSolver::PreSubstep( const Softs::FSolverReal Interpolati
 		{
 			const uint32 GroupId = ActiveGroups[ActiveGroupIndex];
 			const TSet<int32>& ActiveSoftBodies = Evolution->GetGroupActiveSoftBodies(GroupId);
+
+			bool bDoTriangleMeshUpdate = false;
+
 			for (const int32 SoftBodyId : ActiveSoftBodies)
 			{
 				TRACE_CPUPROFILER_EVENT_SCOPE(FClothingSimulationSolver_ParticlePreSubstepKinematicInterpolation);
 				SCOPE_CYCLE_COUNTER(STAT_ChaosClothParticlePreSubstepKinematicInterpolation);
 				Softs::FSolverParticlesRange& Particles = Evolution->GetSoftBodyParticles(SoftBodyId);
 
-				GetClothConstraints(SoftBodyId).SetSkipSelfCollisionInit(!bDetectSelfCollisions);
+				FClothConstraints& Constraints = GetClothConstraints(SoftBodyId);
+
+				Constraints.SetSkipSelfCollisionInit(!bDetectSelfCollisions);
+				if (const TSharedPtr<Softs::FPBDSkinnedTriangleMeshCollisionConstraints>& TriMeshConstraint = Constraints.GetSkinnedTriangleCollisionsConstraint())
+				{
+					bDoTriangleMeshUpdate |= TriMeshConstraint->IsEnabled() && (!TriMeshConstraint->GetUseSelfCollisionSubstepsForSkinnedTriangleMeshes() || bDetectSelfCollisions);
+				}
 #if INTEL_ISPC
 				if (bRealTypeCompatibleWithISPC && bChaos_PreSubstepInterpolation_ISPC_Enabled)  // TODO: Make the ISPC works with both Single and Double depending on the FSolverReal type
 				{
@@ -1905,6 +1897,16 @@ void FClothingSimulationSolver::PreSubstep( const Softs::FSolverReal Interpolati
 					}, Particles.GetRangeSize() < ClothSolverMinParallelBatchSize);
 				}
 			}
+			
+			const TSet<int32>& ActiveCollisionRanges = Evolution->GetGroupActiveCollisionParticleRanges(GroupId);
+			for (const int32 CollisionRangeId : ActiveCollisionRanges)
+			{
+				if (FClothComplexColliders* const Colliders = GetComplexColliders(CollisionRangeId))
+				{
+					Colliders->SetSkipSkinnedTriangleMeshKinematicUpdate(!bDoTriangleMeshUpdate);
+				}
+			}
+
 		}, /*bForceSingleThreaded =*/ !bClothSolverParallelClothPreUpdate);
 
 		{
@@ -2047,6 +2049,13 @@ void FClothingSimulationSolver::Update(Softs::FSolverReal InDeltaTime)
 		Swap(OldCollisionTransforms, CollisionTransforms);
 		Swap(OldAnimationPositions, AnimationPositions);
 		Swap(OldAnimationNormals, AnimationNormals);
+		for (TPair<int32, TUniquePtr<FClothComplexColliders>>& Collider : ComplexColliders)
+		{
+			if (Collider.Value)
+			{
+				Collider.Value->SwapBuffersForFrameFlip();
+			}
+		}
 
 		// Clear external collisions so that they can be re-added
 		CollisionParticlesSize = 0;
diff --git a/Engine/Plugins/ChaosCloth/Source/ChaosCloth/Public/ChaosCloth/ChaosClothConstraints.h b/Engine/Plugins/ChaosCloth/Source/ChaosCloth/Public/ChaosCloth/ChaosClothConstraints.h
--- a/Engine/Plugins/ChaosCloth/Source/ChaosCloth/Public/ChaosCloth/ChaosClothConstraints.h
+++ b/Engine/Plugins/ChaosCloth/Source/ChaosCloth/Public/ChaosCloth/ChaosClothConstraints.h
@@ -15,6 +15,7 @@ namespace Chaos
 {
 
 	struct FClothingPatternData;
+	class FClothComplexColliders;
 
 	class FClothConstraints final
 	{
@@ -54,6 +55,8 @@ namespace Chaos
 			int32 InNumParticles);
 
 		void SetSkipSelfCollisionInit(bool bValue) { bSkipSelfCollisionInit = bValue; }
+
+		void OnCollisionRangeRemoved(int32 CollisionRangeId);
 		// ---- End of Solver interface ----
 
 		// ---- Cloth interface ----
@@ -70,7 +73,8 @@ namespace Chaos
 			bool bEnabled,
 			const FTriangleMesh* MultiResCoarseLODMesh = nullptr,
 			const int32 MultiResCoarseLODParticleRangeId = INDEX_NONE,
-			const TSharedPtr<Softs::FMultiResConstraints>& FineLODMultiResConstraint = TSharedPtr<Softs::FMultiResConstraints>(nullptr));
+			const TSharedPtr<Softs::FMultiResConstraints>& FineLODMultiResConstraint = TSharedPtr<Softs::FMultiResConstraints>(nullptr),
+			const TArray<const FClothComplexColliders*>& ComplexColliders = TArray<const FClothComplexColliders*>());
 
 		UE_DEPRECATED(5.4, "Use AddRules() with WeightMaps, VertexSets, FaceSets, FaceIntMaps, and optional PatternData instead.")
 		void AddRules(
@@ -132,6 +136,7 @@ namespace Chaos
 		const TSharedPtr<Softs::FVelocityAndPressureField>& GetVelocityAndPressureField() const { return VelocityAndPressureField; }
 		const TSharedPtr<Softs::FExternalForces>& GetExternalForces() const { return ExternalForces; }
 		const TSharedPtr<Softs::FPBDSoftBodyCollisionConstraint>& GetCollisionConstraint() const { return CollisionConstraint; }
+		const TSharedPtr<Softs::FPBDSkinnedTriangleMeshCollisionConstraints>& GetSkinnedTriangleCollisionsConstraint() const { return SkinnedTriangleCollisionsConstraint; }
 		const TSharedPtr<Softs::FMultiResConstraints>& GetMultiResConstraints() const { return MultiResConstraints; }
 		// ---- End of debug functions ----
 
@@ -184,7 +189,9 @@ namespace Chaos
 		);
 		void CreateCollisionConstraint(
 			const Softs::FCollectionPropertyConstFacade& ConfigProperties,
-			Softs::FSolverReal MeshScale );
+			Softs::FSolverReal MeshScale,
+			const TMap<FString, TConstArrayView<FRealSingle>>& WeightMaps,
+			const TArray<const FClothComplexColliders*>& ComplexColliders);
 		void CreateMultiresConstraint(
 			const Softs::FCollectionPropertyConstFacade& ConfigProperties,
 			const TMap<FString, TConstArrayView<FRealSingle>>& WeightMaps,
@@ -230,6 +237,7 @@ namespace Chaos
 		TSharedPtr<Softs::FVelocityAndPressureField> VelocityAndPressureField;
 		TSharedPtr<Softs::FExternalForces> ExternalForces;
 		TSharedPtr<Softs::FPBDSoftBodyCollisionConstraint> CollisionConstraint;
+		TSharedPtr<Softs::FPBDSkinnedTriangleMeshCollisionConstraints> SkinnedTriangleCollisionsConstraint;
 		//~ End Force-based solver only constraints
 
 		// Exactly one of these should be non-null
@@ -269,7 +277,6 @@ namespace Chaos
 		//~ End Force-based solver only fields
 
 		//~ Begin PBD solver only fields
-		const TArray<Softs::FSolverVec3>* OldAnimationPositions_Deprecated;
 		int32 ConstraintInitOffset;
 		int32 ConstraintRuleOffset;
 		int32 PostCollisionConstraintRuleOffset;
diff --git a/Engine/Plugins/ChaosCloth/Source/ChaosCloth/Public/ChaosCloth/ChaosClothVisualization.h b/Engine/Plugins/ChaosCloth/Source/ChaosCloth/Public/ChaosCloth/ChaosClothVisualization.h
--- a/Engine/Plugins/ChaosCloth/Source/ChaosCloth/Public/ChaosCloth/ChaosClothVisualization.h
+++ b/Engine/Plugins/ChaosCloth/Source/ChaosCloth/Public/ChaosCloth/ChaosClothVisualization.h
@@ -40,7 +40,7 @@ namespace Chaos
 		CHAOSCLOTH_API void DrawOpenEdges(FPrimitiveDrawInterface* PDI = nullptr) const;
 		CHAOSCLOTH_API void DrawPointNormals(FPrimitiveDrawInterface* PDI = nullptr, const FReal Length = 20.) const;
 		CHAOSCLOTH_API void DrawPointVelocities(FPrimitiveDrawInterface* PDI = nullptr) const;
-		CHAOSCLOTH_API void DrawCollision(FPrimitiveDrawInterface* PDI = nullptr) const;
+		CHAOSCLOTH_API void DrawCollision(FPrimitiveDrawInterface* PDI = nullptr, bool bWireframe = false) const;
 		CHAOSCLOTH_API void DrawBackstops(FPrimitiveDrawInterface* PDI = nullptr) const;
 		CHAOSCLOTH_API void DrawBackstopDistances(FPrimitiveDrawInterface* PDI = nullptr) const;
 		CHAOSCLOTH_API void DrawMaxDistances(FPrimitiveDrawInterface* PDI = nullptr) const;
@@ -71,7 +71,7 @@ namespace Chaos
 		void DrawOpenEdges(FPrimitiveDrawInterface* /*PDI*/ = nullptr) const {}
 		void DrawPointNormals(FPrimitiveDrawInterface* /*PDI*/ = nullptr, const FReal /*Length*/ = 20.) const {}
 		void DrawPointVelocities(FPrimitiveDrawInterface* /*PDI*/ = nullptr) const {}
-		void DrawCollision(FPrimitiveDrawInterface* /*PDI*/ = nullptr) const {}
+		void DrawCollision(FPrimitiveDrawInterface* /*PDI*/ = nullptr, bool /*bWireframe*/ = false) const {}
 		void DrawBackstops(FPrimitiveDrawInterface* /*PDI*/ = nullptr) const {}
 		void DrawBackstopDistances(FPrimitiveDrawInterface* /*PDI*/ = nullptr) const {}
 		void DrawMaxDistances(FPrimitiveDrawInterface* /*PDI*/ = nullptr) const {}
diff --git a/Engine/Plugins/ChaosCloth/Source/ChaosCloth/Public/ChaosCloth/ChaosClothingSimulationCollider.h b/Engine/Plugins/ChaosCloth/Source/ChaosCloth/Public/ChaosCloth/ChaosClothingSimulationCollider.h
--- a/Engine/Plugins/ChaosCloth/Source/ChaosCloth/Public/ChaosCloth/ChaosClothingSimulationCollider.h
+++ b/Engine/Plugins/ChaosCloth/Source/ChaosCloth/Public/ChaosCloth/ChaosClothingSimulationCollider.h
@@ -5,14 +5,21 @@
 #include "Containers/ContainersFwd.h"
 #include "Chaos/ImplicitFwd.h"
 #include "Chaos/PBDSoftsEvolutionFwd.h"
+#include "Chaos/Levelset.h"
+#include "Chaos/SkinnedTriangleMesh.h"
+#include "Chaos/WeightedLatticeImplicitObject.h"
+#include "Misc/EnumRange.h"
 
 class USkeletalMeshComponent;
 class UClothingAssetCommon;
 class UPhysicsAsset;
 class FClothingSimulationContextCommon;
+struct FReferenceSkeleton;
+#if UE_ENABLE_INCLUDE_ORDER_DEPRECATED_IN_5_6
+// These classes never actually existed.
 class FLevelSetCollisionData;
 class FSkinnedLevelSetCollisionData;
-struct FReferenceSkeleton;
+#endif
 
 namespace Chaos
 {
@@ -20,6 +27,50 @@ namespace Chaos
 	class FClothingSimulationCloth;
 	template<typename T> class TWeightedLatticeImplicitObject;
 
+	struct FLevelSetCollisionData
+	{
+		const TSharedPtr<Chaos::FLevelSet, ESPMode::ThreadSafe> LevelSet;
+		FTransform Transform;
+		int32 BoneIndex;
+	};
+
+	struct FSkinnedLevelSetCollisionData
+	{
+		const TRefCountPtr<Chaos::TWeightedLatticeImplicitObject<Chaos::FLevelSet>> WeightedLevelSet;
+		int32 BoneIndex;
+		TArray<int32> MappedSkinnedBones;
+	};
+
+	struct FMLLevelSetCollisionData
+	{
+		const TSharedPtr<Chaos::FMLLevelSet> MLLevelSet;
+		int32 BoneIndex;
+		TArray<int32> MappedActiveBoneIndices;
+	};
+
+	struct FSkinnedTriangleMeshCollisionData
+	{
+		const TRefCountPtr<Chaos::FSkinnedTriangleMesh> SkinnedTriangleMesh;
+		int32 BoneIndex;
+		TArray<int32> MappedSkinnedBones;
+	};
+
+	struct FClothCollisionDataExtended
+	{
+		TArray<FLevelSetCollisionData> LevelSets;
+		TArray<FSkinnedLevelSetCollisionData> SkinnedLevelSets;
+		TArray<FMLLevelSetCollisionData> MLLevelSets;
+		TArray<FSkinnedTriangleMeshCollisionData> SkinnedTriangleMeshes;
+
+		void Reset()
+		{
+			LevelSets.Reset();
+			SkinnedLevelSets.Reset();
+			MLLevelSets.Reset();
+			SkinnedTriangleMeshes.Reset();
+		}
+	};
+
 	// Collider simulation node
 	class FClothingSimulationCollider final
 	{
@@ -29,28 +80,13 @@ namespace Chaos
 			LODless = 0,  // Global LODless collision slot filled with physics collisions
 			External,  // External collision slot added/removed at every frame
 			LODs,  // LODIndex based start slot for LODs collisions
+			Count
 		};
 
-		struct FLevelSetCollisionData
-		{
-			const TSharedPtr<Chaos::FLevelSet, ESPMode::ThreadSafe> LevelSet;
-			FTransform Transform;
-			int32 BoneIndex;
-		};
-
-		struct FSkinnedLevelSetCollisionData
-		{
-			const TRefCountPtr<Chaos::TWeightedLatticeImplicitObject<Chaos::FLevelSet>> WeightedLevelSet;
-			int32 BoneIndex;
-			TArray<int32> MappedSkinnedBones;
-		};
-
-		struct FMLLevelSetCollisionData
-		{
-			const TSharedPtr<Chaos::FMLLevelSet> MLLevelSet;
-			int32 ParentBoneIndex;
-			TArray<int32> MappedActiveBoneIndices;
-		};
+		UE_DEPRECATED(5.6, "FLevelSetCollisionData has been moved out of this class.")
+		typedef Chaos::FLevelSetCollisionData FLevelSetCollisionData;
+		UE_DEPRECATED(5.6, "FSkinnedLevelSetCollisionData has been moved out of this class.")
+		typedef Chaos::FSkinnedLevelSetCollisionData FSkinnedLevelSetCollisionData;
 
 		CHAOSCLOTH_API FClothingSimulationCollider(const UPhysicsAsset* InPhysicsAsset, const FReferenceSkeleton* InReferenceSkeleton);
 		CHAOSCLOTH_API ~FClothingSimulationCollider();
@@ -100,13 +136,15 @@ namespace Chaos
 		// ---- End of the debugging and visualization functions ----
 
 		// Extract all collisions from the physics asset into a more cloth friendly format
-		UE_DEPRECATED(5.5, "Use the version that includes MLLevelSetCollisions.")
+		UE_DEPRECATED(5.6, "Use the version that uses FClothCollisionDataExtended.")
 		CHAOSCLOTH_API static void ExtractPhysicsAssetCollision(
 			const UPhysicsAsset* PhysicsAsset,
 			const FReferenceSkeleton* ReferenceSkeleton,
 			FClothCollisionData& ClothCollisionData,
+			PRAGMA_DISABLE_DEPRECATION_WARNINGS
 			TArray<FLevelSetCollisionData>& LevelSetCollisions,
 			TArray<FSkinnedLevelSetCollisionData>& SkinnedLevelSetCollisions,
+			PRAGMA_ENABLE_DEPRECATION_WARNINGS
 			TArray<int32>& UsedBoneIndices,
 			bool bUseSphylOnly = false,
 			bool bSkipMissingBones = false);
@@ -116,22 +154,21 @@ namespace Chaos
 			const UPhysicsAsset* PhysicsAsset,
 			const FReferenceSkeleton* ReferenceSkeleton,
 			FClothCollisionData& ClothCollisionData,
-			TArray<FLevelSetCollisionData>& LevelSetCollisions,
-			TArray<FSkinnedLevelSetCollisionData>& SkinnedLevelSetCollisions,
-			TArray<FMLLevelSetCollisionData>& MLLevelSetCollisions,
-			TArray<int32>& UsedBoneIndices,
+			FClothCollisionDataExtended& ClothCollisionDataExtended,
+			TArray<int32>& UsedBoneIndices, // Bones that correspond with bodies
+			TArray<int32>& UsedSubBoneIndices, // Bones used by extended type's internal bones, e.g., FSkinnedTriangleMeshCollisionData::MappedSkinnedBones
 			bool bUseSphylOnly = false,
 			bool bSkipMissingBones = false);
 
+		// Return the collision particle range ID for the specified type if valid. If ECollisionDataType::LODs is asked, then the offset returned is for the current LOD.
+		CHAOSCLOTH_API int32 GetCollisionRangeId(const FClothingSimulationSolver* Solver, const FClothingSimulationCloth* Cloth, ECollisionDataType CollisionDataType) const;
+
 	private:	
 		CHAOSCLOTH_API int32 GetNumGeometries(int32 InSlotIndex) const;
 
 		// Return the collision particle range ID for the specified slot being LODLess, external, or any of the LODs collision.
 		CHAOSCLOTH_API int32 GetCollisionRangeId(const FClothingSimulationSolver* Solver, const FClothingSimulationCloth* Cloth, int32 InSlotIndex) const;
 
-		// Return the collision particle range ID for the specified type if valid. If ECollisionDataType::LODs is asked, then the offset returned is for the current LOD.
-		CHAOSCLOTH_API int32 GetCollisionRangeId(const FClothingSimulationSolver* Solver, const FClothingSimulationCloth* Cloth, ECollisionDataType CollisionDataType) const;
-
 		// Return the collision particle  range ID and number of geometries for the specified type if valid. If ECollisionDataType::LODs is asked, then the offset returned is for the current LOD.
 		CHAOSCLOTH_API bool GetCollisionRangeIdAndNumGeometries(const FClothingSimulationSolver* Solver, const FClothingSimulationCloth* Cloth, ECollisionDataType CollisionDataType, int32& OutCollisionRangeId, int32& OutNumGeometries) const;
 
@@ -154,3 +191,5 @@ namespace Chaos
 		FReal Scale = 1.f;
 	};
 } // namespace Chaos
+
+ENUM_RANGE_BY_COUNT(Chaos::FClothingSimulationCollider::ECollisionDataType, Chaos::FClothingSimulationCollider::ECollisionDataType::Count);
diff --git a/Engine/Plugins/ChaosCloth/Source/ChaosCloth/Public/ChaosCloth/ChaosClothingSimulationSolver.h b/Engine/Plugins/ChaosCloth/Source/ChaosCloth/Public/ChaosCloth/ChaosClothingSimulationSolver.h
--- a/Engine/Plugins/ChaosCloth/Source/ChaosCloth/Public/ChaosCloth/ChaosClothingSimulationSolver.h
+++ b/Engine/Plugins/ChaosCloth/Source/ChaosCloth/Public/ChaosCloth/ChaosClothingSimulationSolver.h
@@ -22,6 +22,7 @@ namespace Chaos
 	class FClothingSimulationCollider;
 	class FClothingSimulationMesh;
 	class FClothingSimulationConfig;
+	class FClothComplexColliders;
 
 	namespace Softs
 	{
@@ -348,6 +349,23 @@ namespace Chaos
 		// ---- Collider interface ----
 		CHAOSCLOTH_API int32 AddCollisionParticles(int32 NumCollisionParticles, uint32 GroupId, int32 RecycledCollisionRangeId = 0);
 		CHAOSCLOTH_API void EnableCollisionParticles(int32 CollisionRangeId, bool bEnable);
+		// This method is not supported by LegacySolver and will return nullptr.
+		const FClothComplexColliders* GetComplexColliders(int32 CollisionRangeId) const 
+		{ 
+			if (const TUniquePtr<FClothComplexColliders>* const Found = ComplexColliders.Find(CollisionRangeId))
+			{
+				return Found->Get();
+			}
+			return nullptr;
+		}
+		FClothComplexColliders* GetComplexColliders(int32 CollisionRangeId) 
+		{
+			if (TUniquePtr<FClothComplexColliders> const* Found = ComplexColliders.Find(CollisionRangeId))
+			{
+				return Found->Get();
+			}
+			return nullptr;
+		}
 
 		CHAOSCLOTH_API void ResetCollisionStartPose(int32 CollisionRangeId, int32 NumCollisionParticles);
 
@@ -471,6 +489,10 @@ namespace Chaos
 		TArrayCollectionArray<Softs::FSolverRigidTransform3> LastSubframeCollisionTransformsCCD;
 		TArrayCollectionArray<bool> Collided;
 
+		// Extended collision data to support complex colliders (FEvolution only)
+		TMap<int32, TUniquePtr<FClothComplexColliders>> ComplexColliders;
+
+
 		// Cloth constraints
 		TMap<int32, TUniquePtr<FClothConstraints>> ClothsConstraints;
 
diff --git a/Engine/Plugins/ChaosClothAsset/Source/ChaosClothAsset/Public/ChaosClothAsset/CollectionClothFabricFacade.h b/Engine/Plugins/ChaosClothAsset/Source/ChaosClothAsset/Public/ChaosClothAsset/CollectionClothFabricFacade.h
--- a/Engine/Plugins/ChaosClothAsset/Source/ChaosClothAsset/Public/ChaosClothAsset/CollectionClothFabricFacade.h
+++ b/Engine/Plugins/ChaosClothAsset/Source/ChaosClothAsset/Public/ChaosClothAsset/CollectionClothFabricFacade.h
@@ -18,6 +18,7 @@ namespace UE::Chaos::ClothAsset
 		inline static constexpr float Pressure = 0.0f;
 		inline static constexpr int32 Layer = INDEX_NONE;
 		inline static constexpr float CollisionThickness = 1.0f;
+		inline static constexpr float ClothCollisionThickness = 0.f;
 		inline static constexpr float SelfFriction = 0.0f;
 		inline static constexpr float SelfCollisionThickness = 0.5f;
 	};
diff --git a/Engine/Plugins/ChaosClothAsset/Source/ChaosClothAssetEngine/Private/ChaosClothAsset/CollisionSources.cpp b/Engine/Plugins/ChaosClothAsset/Source/ChaosClothAssetEngine/Private/ChaosClothAsset/CollisionSources.cpp
--- a/Engine/Plugins/ChaosClothAsset/Source/ChaosClothAssetEngine/Private/ChaosClothAsset/CollisionSources.cpp
+++ b/Engine/Plugins/ChaosClothAsset/Source/ChaosClothAssetEngine/Private/ChaosClothAsset/CollisionSources.cpp
@@ -1,8 +1,6 @@
 // Copyright Epic Games, Inc. All Rights Reserved.
 #include "ChaosClothAsset/CollisionSources.h"
 
-#include "Chaos/Levelset.h"
-#include "Chaos/WeightedLatticeImplicitObject.h"
 #include "ChaosCloth/ChaosClothingSimulationCollider.h"
 #include "Components/SkinnedMeshComponent.h"
 #include "Engine/SkinnedAsset.h"
@@ -160,19 +158,18 @@ namespace UE::Chaos::ClothAsset
 				if (SkinnedAsset && PhysicsAsset)
 				{
 					// Extract collisions
-					TArray<::Chaos::FClothingSimulationCollider::FLevelSetCollisionData> LevelSetCollisions;
-					TArray<::Chaos::FClothingSimulationCollider::FSkinnedLevelSetCollisionData> SkinnedLevelSetCollisions;
-					TArray<::Chaos::FClothingSimulationCollider::FMLLevelSetCollisionData> MLLevelSetCollisions;
+					// Currently not handling extended collision data
+					::Chaos::FClothCollisionDataExtended ExtendedData;
+					TArray<int32> UsedSubBoneIndices;
 
 					constexpr bool bSkipMissingBones = true;
 					::Chaos::FClothingSimulationCollider::ExtractPhysicsAssetCollision(
 						PhysicsAsset,
 						&SkinnedAsset->GetRefSkeleton(),
 						CachedCollisionData,
-						LevelSetCollisions,
-						SkinnedLevelSetCollisions,
-						MLLevelSetCollisions,
+						ExtendedData,
 						CachedUsedBoneIndices,
+						UsedSubBoneIndices,
 						bUseSphylsOnly,
 						bSkipMissingBones);
 				}
diff --git a/Engine/Plugins/ChaosClothAssetEditor/Source/ChaosClothAssetDataflowNodes/Private/ChaosClothAsset/SimulationCollisionConfigNode.cpp b/Engine/Plugins/ChaosClothAssetEditor/Source/ChaosClothAssetDataflowNodes/Private/ChaosClothAsset/SimulationCollisionConfigNode.cpp
--- a/Engine/Plugins/ChaosClothAssetEditor/Source/ChaosClothAssetDataflowNodes/Private/ChaosClothAsset/SimulationCollisionConfigNode.cpp
+++ b/Engine/Plugins/ChaosClothAssetEditor/Source/ChaosClothAssetDataflowNodes/Private/ChaosClothAsset/SimulationCollisionConfigNode.cpp
@@ -5,11 +5,19 @@
 
 #include UE_INLINE_GENERATED_CPP_BY_NAME(SimulationCollisionConfigNode)
 
+PRAGMA_DISABLE_DEPRECATION_WARNINGS
 FChaosClothAssetSimulationCollisionConfigNode::FChaosClothAssetSimulationCollisionConfigNode(const UE::Dataflow::FNodeParameters& InParam, FGuid InGuid)
 	: FChaosClothAssetSimulationBaseConfigNode(InParam, InGuid)
 {
 	RegisterCollectionConnections();
+	RegisterInputConnection(&FrictionCoefficientWeighted.WeightMap)
+		.SetCanHidePin(true)
+		.SetPinIsHidden(true);
+	RegisterInputConnection(&ClothCollisionThickness.WeightMap)
+		.SetCanHidePin(true)
+		.SetPinIsHidden(true);
 }
+PRAGMA_ENABLE_DEPRECATION_WARNINGS
 
 void FChaosClothAssetSimulationCollisionConfigNode::AddProperties(FPropertyHelper& PropertyHelper) const
 {
@@ -22,11 +30,24 @@ void FChaosClothAssetSimulationCollisionConfigNode::AddProperties(FPropertyHelpe
 			return FabricFacade.GetCollisionThickness();
 		}, {});
 
-	PropertyHelper.SetFabricProperty(FName(TEXT("FrictionCoefficient")), FrictionCoefficientImported,
-		[](UE::Chaos::ClothAsset::FCollectionClothFabricFacade& FabricFacade)-> float
+	PropertyHelper.SetFabricPropertyWeighted(FName(TEXT("FrictionCoefficient")), FrictionCoefficientWeighted,
+		[](const UE::Chaos::ClothAsset::FCollectionClothFabricFacade& FabricFacade)-> float
 		{
 			return FabricFacade.GetFriction();
 		}, {});
+
+	PropertyHelper.SetFabricPropertyWeighted(FName(TEXT("SoftBodyCollisionThickness")), ClothCollisionThickness,
+		[](const UE::Chaos::ClothAsset::FCollectionClothFabricFacade& FabricFacade)-> float
+		{
+			return FabricFacade.GetCollisionThickness();
+		}, {});
+
+	PropertyHelper.SetPropertyBool(this, &bEnableSimpleColliders);
+	PropertyHelper.SetPropertyBool(this, &bUsePlanarConstraintForSimpleColliders);
+	PropertyHelper.SetPropertyBool(this, &bEnableComplexColliders);
+	PropertyHelper.SetPropertyBool(this, &bUsePlanarConstraintForComplexColliders);
+	PropertyHelper.SetPropertyBool(this, &bEnableSkinnedTriangleMeshCollisions);
+	PropertyHelper.SetPropertyBool(this, &bUseSelfCollisionSubstepsForSkinnedTriangleMeshes);
 }
 
 void FChaosClothAssetSimulationCollisionConfigNode::Serialize(FArchive& Ar)
@@ -35,6 +56,7 @@ void FChaosClothAssetSimulationCollisionConfigNode::Serialize(FArchive& Ar)
 	if (Ar.IsLoading())
 	{
 #if WITH_EDITORONLY_DATA
+		PRAGMA_DISABLE_DEPRECATION_WARNINGS
 		if (FrictionCoefficient_DEPRECATED != FrictionCoefficientDeprecatedDefault)
 		{
 			FrictionCoefficientImported.ImportedValue = FrictionCoefficient_DEPRECATED;
@@ -45,6 +67,17 @@ void FChaosClothAssetSimulationCollisionConfigNode::Serialize(FArchive& Ar)
 			CollisionThicknessImported.ImportedValue = CollisionThickness_DEPRECATED;
 			CollisionThickness_DEPRECATED = CollisionThicknessDeprecatedDefault;
 		}
+		if (FrictionCoefficientImported.ImportedValue != UE::Chaos::ClothAsset::FDefaultFabric::Friction)
+		{
+			FrictionCoefficientWeighted.Low = FrictionCoefficientWeighted.High = FrictionCoefficientImported.ImportedValue;
+			FrictionCoefficientImported.ImportedValue = UE::Chaos::ClothAsset::FDefaultFabric::Friction;
+		}
+		if (FrictionCoefficientImported.bUseImportedValue)
+		{
+			FrictionCoefficientWeighted.bImportFabricBounds = true;
+			FrictionCoefficientImported.bUseImportedValue = false;
+		}
+		PRAGMA_ENABLE_DEPRECATION_WARNINGS
 #endif
 	}
 }
diff --git a/Engine/Plugins/ChaosClothAssetEditor/Source/ChaosClothAssetDataflowNodes/Public/ChaosClothAsset/SimulationCollisionConfigNode.h b/Engine/Plugins/ChaosClothAssetEditor/Source/ChaosClothAssetDataflowNodes/Public/ChaosClothAsset/SimulationCollisionConfigNode.h
--- a/Engine/Plugins/ChaosClothAssetEditor/Source/ChaosClothAssetDataflowNodes/Public/ChaosClothAsset/SimulationCollisionConfigNode.h
+++ b/Engine/Plugins/ChaosClothAssetEditor/Source/ChaosClothAssetDataflowNodes/Public/ChaosClothAsset/SimulationCollisionConfigNode.h
@@ -6,10 +6,12 @@
 #include "SimulationCollisionConfigNode.generated.h"
 
 /** Physics mesh collision properties configuration node. */
+PRAGMA_DISABLE_DEPRECATION_WARNINGS
 USTRUCT(Meta = (DataflowCloth))
 struct FChaosClothAssetSimulationCollisionConfigNode : public FChaosClothAssetSimulationBaseConfigNode
 {
-	GENERATED_USTRUCT_BODY()
+	GENERATED_BODY()
+PRAGMA_ENABLE_DEPRECATION_WARNINGS
 	DATAFLOW_NODE_DEFINE_INTERNAL(FChaosClothAssetSimulationCollisionConfigNode, "SimulationCollisionConfig", "Cloth", "Cloth Simulation Collision Config")
 
 public:
@@ -17,9 +19,39 @@ public:
 	UPROPERTY(EditAnywhere, Category = "Collision Properties", DisplayName = "Collision Thickness", Meta = (UIMin = "0", UIMax = "100", ClampMin = "0", ClampMax = "1000", InteractorName = "CollisionThickness"))
 	FChaosClothAssetImportedFloatValue CollisionThicknessImported = {UE::Chaos::ClothAsset::FDefaultFabric::CollisionThickness};
 
-	/** Friction coefficient for cloth - collider interaction. */
+	/** Friction coefficient for cloth - collider interaction. Currently only Skinned Triangle Meshes use the weighted value. All other collisions only use the Low value. */
 	UPROPERTY(EditAnywhere, Category = "Collision Properties", DisplayName = "Friction Coefficient", Meta = (UIMin = "0", UIMax = "1", ClampMin = "0", ClampMax = "10", InteractorName = "FrictionCoefficient"))
-	FChaosClothAssetImportedFloatValue FrictionCoefficientImported = {UE::Chaos::ClothAsset::FDefaultFabric::Friction};
+	FChaosClothAssetWeightedValue FrictionCoefficientWeighted = 
+	{ true, UE::Chaos::ClothAsset::FDefaultFabric::Friction, UE::Chaos::ClothAsset::FDefaultFabric::Friction, TEXT("FrictionCoefficient"), true };
+
+	/** Enable colliding against any simple (e.g., capsules, convexes, spheres, boxes) colliders..*/
+	UPROPERTY(EditAnywhere, Category = "Collision Properties", Meta = (InteractorName = "EnableSimpleColliders"))
+	bool bEnableSimpleColliders = true;
+
+	/** Use Planar Constraints for simple (e.g., capsules, convexes, spheres, boxes) colliders when doing multiple iterations. Planar constraints are cheaper than full collision detection, but less accurate. */
+	UPROPERTY(EditAnywhere, Category="Collision Properties", Meta = (EditCondition = "bEnableSimpleColliders", InteractorName = "UsePlanarConstraintForSimpleColliders"))
+	bool bUsePlanarConstraintForSimpleColliders = false;
+
+	/** Enable colliding against any complex (e.g., SkinnedLevelSet, MLLevelSet) colliders.*/
+	UPROPERTY(EditAnywhere, Category = "Collision Properties", Meta = (InteractorName = "EnableComplexColliders"))
+	bool bEnableComplexColliders = true;
+
+	/** Use Planar Constraints for complex (e.g., SkinnedLevelSet, MLLevelSet) colliders when doing multiple iterations. Planar constraints are cheaper than full collision detection, but less accurate. */
+	UPROPERTY(EditAnywhere, Category = "Collision Properties", Meta = (EditCondition = "bEnableComplexColliders", InteractorName = "UsePlanarConstraintForComplexColliders"))
+	bool bUsePlanarConstraintForComplexColliders = true;
+
+	/** Enable colliding against any Skinned Triangle Mesh colliders.*/
+	UPROPERTY(EditAnywhere, Category = "Collision Properties", Meta = (InteractorName = "EnableSkinnedTriangleMeshCollisions"))
+	bool bEnableSkinnedTriangleMeshCollisions = true;
+
+	/** Use 'NumSelfCollisionSubsteps' (Located on SimulationSolverConfig) to also control Skinned Triangle Mesh collision updates */
+	UPROPERTY(EditAnywhere, Category = "Collision Properties", Meta = (EditCondition = "bEnableSkinnedTriangleMeshCollisions", InteractorName = "UseSelfCollisionSubstepsForSkinnedTriangleMeshes"))
+	bool bUseSelfCollisionSubstepsForSkinnedTriangleMeshes = true;
+
+	/** Thickness added to the cloth when colliding against collision shapes. Currently only Skinned Triangle Meshes use the weighted value. All other collisions only use the Low value. */
+	UPROPERTY(EditAnywhere, Category = "Collision Properties", Meta = (InteractorName = "SoftBodyCollisionThickness"))
+	FChaosClothAssetWeightedValue ClothCollisionThickness = { true, UE::Chaos::ClothAsset::FDefaultFabric::ClothCollisionThickness,
+		UE::Chaos::ClothAsset::FDefaultFabric::ClothCollisionThickness, TEXT("ClothCollisionThickness"), true };
 
 	/** Stiffness for proximity repulsion forces (Force-based solver only). Units = kg cm/ s^2 (same as XPBD springs)*/
 	UPROPERTY(EditAnywhere, Category = "Proximity Force Properties", Meta = (UIMin = "0", UIMax = "10000", ClampMin = "0", ClampMax = "10000000", InteractorName = "ProximityStiffness"))
@@ -32,6 +64,11 @@ public:
 	UPROPERTY(EditAnywhere, Category = "Collision Properties", Meta = (InteractorName = "UseCCD"))
 	bool bUseCCD = false;
 
+
+	UE_DEPRECATED(5.6, "Use FrictionCoefficientWeighted instead.")
+	UPROPERTY()
+	FChaosClothAssetImportedFloatValue FrictionCoefficientImported = { UE::Chaos::ClothAsset::FDefaultFabric::Friction };
+
 	FChaosClothAssetSimulationCollisionConfigNode(const UE::Dataflow::FNodeParameters& InParam, FGuid InGuid = FGuid::NewGuid());
 
 	virtual void Serialize(FArchive& Ar) override;
diff --git a/Engine/Plugins/ChaosClothAssetEditor/Source/ChaosClothAssetEditor/Private/ChaosClothAsset/ClothEditorSimulationVisualization.cpp b/Engine/Plugins/ChaosClothAssetEditor/Source/ChaosClothAssetEditor/Private/ChaosClothAsset/ClothEditorSimulationVisualization.cpp
--- a/Engine/Plugins/ChaosClothAssetEditor/Source/ChaosClothAssetEditor/Private/ChaosClothAsset/ClothEditorSimulationVisualization.cpp
+++ b/Engine/Plugins/ChaosClothAssetEditor/Source/ChaosClothAssetEditor/Private/ChaosClothAsset/ClothEditorSimulationVisualization.cpp
@@ -245,10 +245,19 @@ const FVisualizationOption FVisualizationOption::OptionData[] =
 	FVisualizationOption(
 		FClothVisualizationDebugDraw::CreateLambda([](const FClothEditorSimulationVisualization&, const ::Chaos::FClothVisualizationNoGC& Visualization, FPrimitiveDrawInterface* PDI)
 		{
-			Visualization.DrawCollision(PDI);
+			constexpr bool bForceWireframe = false;
+			Visualization.DrawCollision(PDI, bForceWireframe);
 		}),
 		LOCTEXT("ChaosVisName_Collision", "Collisions"), 
 		LOCTEXT("ChaosVisName_Collision_ToolTip", "Draws the collision bodies the simulation is currently using")),
+	FVisualizationOption(
+		FClothVisualizationDebugDraw::CreateLambda([](const FClothEditorSimulationVisualization&, const ::Chaos::FClothVisualizationNoGC& Visualization, FPrimitiveDrawInterface* PDI)
+		{
+			constexpr bool bForceWireframe = true;
+			Visualization.DrawCollision(PDI, bForceWireframe);
+		}),
+		LOCTEXT("ChaosVisName_CollisionWireframe", "Collisions (Force Wireframe)"),
+		LOCTEXT("ChaosVisName_Collision_ToolTip", "Draws the collision bodies the simulation is currently using")),
 	FVisualizationOption(
 		FClothVisualizationDebugDraw::CreateLambda([](const FClothEditorSimulationVisualization&, const ::Chaos::FClothVisualizationNoGC& Visualization, FPrimitiveDrawInterface* PDI)
 		{
diff --git a/Engine/Source/Runtime/Experimental/Chaos/Private/Chaos/MLLevelset.cpp b/Engine/Source/Runtime/Experimental/Chaos/Private/Chaos/MLLevelset.cpp
--- a/Engine/Source/Runtime/Experimental/Chaos/Private/Chaos/MLLevelset.cpp
+++ b/Engine/Source/Runtime/Experimental/Chaos/Private/Chaos/MLLevelset.cpp
@@ -64,8 +64,6 @@ CHAOS_API FMLLevelSet::FMLLevelSet(TObjectPtr<UNNEModelData> InNNESignedDistance
 FMLLevelSet::FMLLevelSet(FMLLevelSet&& Other)
 	: FImplicitObject(EImplicitObject::HasBoundingBox, ImplicitObjectType::MLLevelSet)
 	, ActiveBoneNames(MoveTemp(Other.ActiveBoneNames))
-	, SolverActiveBoneIndices(MoveTemp(Other.SolverActiveBoneIndices))
-	, SolverParentBoneIndex(Other.SolverParentBoneIndex)
 	, ActiveBonesRelativeTransforms(MoveTemp(Other.ActiveBonesRelativeTransforms))
 
 	, NNESignedDistanceModel(Other.NNESignedDistanceModel)
@@ -102,9 +100,7 @@ Chaos::FImplicitObjectPtr FMLLevelSet::CopyGeometry() const
 	FMLLevelSet* Copy = new FMLLevelSet();
 
 	Copy->ActiveBoneNames = ActiveBoneNames;
-	Copy->SolverActiveBoneIndices = SolverActiveBoneIndices;
 	Copy->ActiveBonesRelativeTransforms = ActiveBonesRelativeTransforms;
-	Copy->SolverParentBoneIndex = SolverParentBoneIndex;
 
 	Copy->NNESignedDistanceModel = NNESignedDistanceModel;
 	Copy->SignedDistanceNeuralInferences = SignedDistanceNeuralInferences;
@@ -138,9 +134,7 @@ Chaos::FImplicitObjectPtr FMLLevelSet::CopyGeometryWithScale(const FVec3& Scale)
 	FMLLevelSet* Copy = new FMLLevelSet();
 
 	Copy->ActiveBoneNames = ActiveBoneNames;
-	Copy->SolverActiveBoneIndices = SolverActiveBoneIndices;
 	Copy->ActiveBonesRelativeTransforms = ActiveBonesRelativeTransforms;
-	Copy->SolverParentBoneIndex = SolverParentBoneIndex;
 
 	Copy->NNESignedDistanceModel = NNESignedDistanceModel;
 	Copy->SignedDistanceNeuralInferences = SignedDistanceNeuralInferences;
@@ -756,192 +750,4 @@ void FMLLevelSet::BatchPhiWithNormal(const TConstArrayView<Softs::FPAndInvM> PAn
 		NormalsOut[Index] = (Softs::FSolverVec3)Normal;
 	}	
 }
-
-void FMLLevelSet::BatchPhiWithNormalWithGreatestInfluenceBone(const TConstArrayView<Softs::FPAndInvM> PAndInvMArray,
-	const Softs::FSolverRigidTransform3& SolverToThis, TArray<Softs::FSolverReal>& OutBatchPhis, TArray<Softs::FSolverVec3>& OutBatchNormals, 
-	TArray<int32>& OutBatchGreatestInfluenceSolverBone, const Chaos::Softs::FSolverReal CollisionThickness, 
-	int32 MLLevelsetThread, int32 BatchBegin, int32 BatchEnd) const
-{	
-	TRACE_CPUPROFILER_EVENT_SCOPE(FMLLevelSet_BatchPhiWithNormal);
-	const uint32 ModelInputShapeSize = 3;  // XYZ coordinates wrt local bounding box
-	const uint32 ModelOutputShapeSize = 1; // At the moment there is a single output (SignedDistance).
-	const Chaos::Softs::FSolverReal MaxPhiValue = 2 * SignedDistanceScaling;
-	const float DeltaForFiniteDifferenceForNormalCalc = SignedDistanceScaling / float(100.0); //delta for finite differencing
-
-	const int32 NumParticles = BatchEnd - BatchBegin;
-	const Softs::FPAndInvM* const PAndInvM = PAndInvMArray.GetData();
-	Softs::FSolverReal* const Phis = OutBatchPhis.GetData();
-	Softs::FSolverVec3* const NormalsOut = OutBatchNormals.GetData();
-	int32* const GreatestInfluenceSolverBones = OutBatchGreatestInfluenceSolverBone.GetData();
-
-	//Query the phi/signed distance values on the cloth
-	TArray<float, TAlignedHeapAllocator<64>> MLInputForPhis;
-	MLInputForPhis.SetNum(NumParticles * ModelInputShapeSize);
-	TArray<float, TAlignedHeapAllocator<64>> MLOutputForPhis;
-
-	TArray<int32> ParticlesInsideTheGridIndexes;
-	ParticlesInsideTheGridIndexes.SetNum(NumParticles);
-	TArray<TArray<float, TAlignedHeapAllocator<64>>> MLWeightsIn;
-	ComputeSignedDistanceNetworkWeightsInput(MLWeightsIn);
-		
-	TArray<FVector3f> TrainingGridVectorsScaled;
-	TrainingGridVectorsScaled.SetNum(3);
-	for (int32 Index = 0; Index < 3; Index++)
-	{
-		TrainingGridVectorsScaled[Index] = TrainingGridUnitAxesXYZ[Index] / TrainingGridAxesLengthsXYZ[Index];
-	}
-
-	int32 InteriorParticlesIndex = 0;
-	for (int32 Index = BatchBegin; Index < BatchEnd; Index++)
-	{
-		OutBatchPhis[Index] = MaxPhiValue;
-
-		if (PAndInvM[Index].InvM == (Softs::FSolverReal)0.)
-		{
-			continue;
-		}
-
-		const FVec3 RigidSpacePosition(SolverToThis.TransformPositionNoScale(PAndInvM[Index].P));
-
-		// only query points inside of the grid
-		if (LocalBoundingBox.Contains(RigidSpacePosition))
-		{
-			const int32 ClosestBoneJointIndex = GetClosestActiveBoneIndex(RigidSpacePosition);
-			const FVector3f LocalGridCornerShift = (FVector3f)ActiveBonesRelativeTransforms[ClosestBoneJointIndex].GetTranslation() - (FVector3f)ActiveBonesReferenceTranslations[ClosestBoneJointIndex];
-			const FVector3f LocalGridCornerShifted = TrainingGridMin + LocalGridCornerShift;
-
-			TVector<FRealSingle, 3> LocationMSShifted = RigidSpacePosition - LocalGridCornerShifted;
-			for (int32 j = 0; j < 3; j++)
-			{
-				MLInputForPhis[InteriorParticlesIndex * ModelInputShapeSize + j] = TrainingGridVectorsScaled[j][0] * LocationMSShifted[0] + TrainingGridVectorsScaled[j][1] * LocationMSShifted[1] + TrainingGridVectorsScaled[j][2] * LocationMSShifted[2];
-			}
-			ParticlesInsideTheGridIndexes[InteriorParticlesIndex] = Index;
-			InteriorParticlesIndex++;
-		}
-		GreatestInfluenceSolverBones[Index] = GetGreatestInfluenceBoneSolverIndex();;
-	}
-
-	MLInputForPhis.SetNum(InteriorParticlesIndex * ModelInputShapeSize);
-	MLOutputForPhis.SetNum(InteriorParticlesIndex * ModelOutputShapeSize);
-	ParticlesInsideTheGridIndexes.SetNum(InteriorParticlesIndex);
-
-	if (SignedDistanceNeuralInferences[MLLevelsetThread].IsValid())
-	{
-		SignedDistanceNeuralInferences[MLLevelsetThread].RunInference(MLInputForPhis, MLOutputForPhis, ModelInputShapeSize, ModelOutputShapeSize, MLWeightsIn);
-	}	
-
-	TArray<int32> NegativePhiValueIndices;
-	NegativePhiValueIndices.SetNumUninitialized(NumParticles);
-
-	// Only Calculate Normals of those particles which are inside of the Body
-	int32 NumberOfNegativePhis = 0;
-	for (int32 IndexML = 0; IndexML < ParticlesInsideTheGridIndexes.Num(); IndexML++)
-	{
-		int32 Index = ParticlesInsideTheGridIndexes[IndexML];
-		Phis[Index] = SignedDistanceScaling * MLOutputForPhis[IndexML * ModelOutputShapeSize];
-		if (Phis[Index] < CollisionThickness)
-		{
-			NegativePhiValueIndices[NumberOfNegativePhis] = Index;
-			NumberOfNegativePhis++;
-		}
-	}
-	NegativePhiValueIndices.SetNum(NumberOfNegativePhis);
-
-	TArray<int32> CorrectZoneIndices;
-	int32 NumberOfCorrectZoneIndices = 0;
-
-	if (bUseIncorrectZoneModel)
-	{
-		TArray<float, TAlignedHeapAllocator<64>> MLInputForIncorrectZone;
-		TArray<float, TAlignedHeapAllocator<64>> MLOutputForIncorrectZone;
-		const int32 ModelOutputShapeSizeForIncorrectZone = 1;
-
-		MLInputForIncorrectZone.SetNumUninitialized(NegativePhiValueIndices.Num() * ModelInputShapeSize);
-		MLOutputForIncorrectZone.SetNumUninitialized(NegativePhiValueIndices.Num() * ModelOutputShapeSizeForIncorrectZone);
-		CorrectZoneIndices.SetNumUninitialized(NegativePhiValueIndices.Num());
-
-		for (int32 NegativePhiIndex = 0; NegativePhiIndex < NegativePhiValueIndices.Num(); NegativePhiIndex++)
-		{
-			int32 Index = NegativePhiValueIndices[NegativePhiIndex];
-			const FVec3 PositionLocal(SolverToThis.TransformPositionNoScale(PAndInvM[Index].P));
-			const int32 ClosestBoneJointIndex = GetClosestActiveBoneIndex(PositionLocal);
-			const FVector3f LocalGridCornerShift = (FVector3f)ActiveBonesRelativeTransforms[ClosestBoneJointIndex].GetTranslation() - (FVector3f)ActiveBonesReferenceTranslations[ClosestBoneJointIndex];
-			const FVector3f LocalGridCornerShifted = TrainingGridMin + LocalGridCornerShift;
-
-			TVector<FRealSingle, 3> LocationMSShifted = PositionLocal - LocalGridCornerShifted;
-			for (int32 j = 0; j < 3; j++)
-			{
-				MLInputForIncorrectZone[NegativePhiIndex * ModelInputShapeSize + j] = TrainingGridVectorsScaled[j][0] * LocationMSShifted[0] + TrainingGridVectorsScaled[j][1] * LocationMSShifted[1] + TrainingGridVectorsScaled[j][2] * LocationMSShifted[2];
-			}
-		}
-		TArray<TArray<float, TAlignedHeapAllocator<64>>> NetworkWeightsInputForIncorrectZone;
-		ComputeIncorrectZoneNetworkWeightsInput(NetworkWeightsInputForIncorrectZone);
-		IncorrectZoneNeuralInferences[MLLevelsetThread].RunInference(MLInputForIncorrectZone, MLOutputForIncorrectZone, ModelInputShapeSize, ModelOutputShapeSizeForIncorrectZone, NetworkWeightsInputForIncorrectZone);
-		
-		for (int32 NegativePhiIndex = 0; NegativePhiIndex < NegativePhiValueIndices.Num(); NegativePhiIndex++)
-		{
-			int32 Index = NegativePhiValueIndices[NegativePhiIndex];
-			if (MLOutputForIncorrectZone[NegativePhiIndex] > 0)
-			{
-				CorrectZoneIndices[NumberOfCorrectZoneIndices] = Index;
-				NumberOfCorrectZoneIndices++;
-			}
-			else
-			{
-				Phis[Index] = MaxPhiValue;
-			}
-		}
-		CorrectZoneIndices.SetNum(NumberOfCorrectZoneIndices);
-	}
-	else
-	{
-		CorrectZoneIndices = MoveTemp(NegativePhiValueIndices);
-	}
-	
-	TArray<float, TAlignedHeapAllocator<64>> MLInputForNormals;
-	MLInputForNormals.SetNum(CorrectZoneIndices.Num() * ModelInputShapeSize * ModelInputShapeSize);
-	TArray<float, TAlignedHeapAllocator<64>> MLOutputForNormals;
-	MLOutputForNormals.SetNum(CorrectZoneIndices.Num() * ModelInputShapeSize * ModelOutputShapeSize);
-	
-	for (int32 CorrectZoneIndex = 0; CorrectZoneIndex < CorrectZoneIndices.Num(); CorrectZoneIndex++)
-	{
-		int32 Index = CorrectZoneIndices[CorrectZoneIndex];
-
-		const FVec3 RigidSpacePosition(SolverToThis.TransformPositionNoScale(PAndInvM[Index].P));
-		const int32 ClosestBoneJointIndex = GetClosestActiveBoneIndex(RigidSpacePosition);
-		const FVector3f LocalGridCornerShift = (FVector3f)ActiveBonesRelativeTransforms[ClosestBoneJointIndex].GetTranslation() - (FVector3f)ActiveBonesReferenceTranslations[ClosestBoneJointIndex];
-		const FVector3f LocalGridCornerShifted = TrainingGridMin + LocalGridCornerShift;
-		TVector<FRealSingle, 3> LocationMSShifted = RigidSpacePosition - LocalGridCornerShifted;
-
-		for (uint32 NormalDirectionIndex = 0; NormalDirectionIndex < 3; NormalDirectionIndex++)
-		{
-			for (uint32 i = 0; i < ModelInputShapeSize; i++)
-			{
-				MLInputForNormals[CorrectZoneIndex * ModelInputShapeSize * ModelInputShapeSize + NormalDirectionIndex * ModelInputShapeSize + i] =
-					TrainingGridUnitAxesXYZ[i].Dot(LocationMSShifted) / TrainingGridAxesLengthsXYZ[i];
-				MLInputForNormals[CorrectZoneIndex * ModelInputShapeSize * ModelInputShapeSize + NormalDirectionIndex * ModelInputShapeSize + i] +=
-					DeltaForFiniteDifferenceForNormalCalc * float(TrainingGridUnitAxesXYZ[i][NormalDirectionIndex]) / TrainingGridAxesLengthsXYZ[i];
-			}
-		}
-	}
-
-	if (SignedDistanceNeuralInferences[MLLevelsetThread].IsValid())
-	{
-		SignedDistanceNeuralInferences[MLLevelsetThread].RunInference(MLInputForNormals, MLOutputForNormals, ModelInputShapeSize, ModelOutputShapeSize, MLWeightsIn);
-	}
-
-	for (int32 CorrectZoneIndex = 0; CorrectZoneIndex < CorrectZoneIndices.Num(); CorrectZoneIndex++)
-	{
-		int32 Index = CorrectZoneIndices[CorrectZoneIndex];
-		FVec3 Normal;
-		for (int32 i = 0; i < 3; i++)
-		{
-			Normal[i] = (SignedDistanceScaling * MLOutputForNormals[CorrectZoneIndex * ModelInputShapeSize * ModelOutputShapeSize + ModelOutputShapeSize * i] - Phis[Index]) / DeltaForFiniteDifferenceForNormalCalc;
-		}
-
-		constexpr FReal EpsilonForSafeNormalization = 1.0e-8;
-		Normal.SafeNormalize(EpsilonForSafeNormalization);
-		NormalsOut[Index] = (Softs::FSolverVec3)Normal;
-	}
-}
 }
