Fix FParticleRangeIndex CIS error.

diff --git a/Engine/Plugins/ChaosCloth/Source/ChaosCloth/Public/ChaosCloth/ChaosClothingSimulationCollider.h b/Engine/Plugins/ChaosCloth/Source/ChaosCloth/Public/ChaosCloth/ChaosClothingSimulationCollider.h
--- a/Engine/Plugins/ChaosCloth/Source/ChaosCloth/Public/ChaosCloth/ChaosClothingSimulationCollider.h
+++ b/Engine/Plugins/ChaosCloth/Source/ChaosCloth/Public/ChaosCloth/ChaosClothingSimulationCollider.h
@@ -32,6 +32,15 @@ namespace Chaos
 		const TSharedPtr<Chaos::FLevelSet, ESPMode::ThreadSafe> LevelSet;
 		FTransform Transform;
 		int32 BoneIndex;
+
+		FLevelSetCollisionData() = default;
+		~FLevelSetCollisionData() = default;
+
+		FLevelSetCollisionData(const TSharedPtr<Chaos::FLevelSet, ESPMode::ThreadSafe>& InLevelSet, const FTransform& InTransform, int32 InBoneIndex)
+			: LevelSet(InLevelSet)
+			, Transform(InTransform)
+			, BoneIndex(InBoneIndex)
+		{}
 	};
 
 	struct FSkinnedLevelSetCollisionData
@@ -39,6 +48,15 @@ namespace Chaos
 		const TRefCountPtr<Chaos::TWeightedLatticeImplicitObject<Chaos::FLevelSet>> WeightedLevelSet;
 		int32 BoneIndex;
 		TArray<int32> MappedSkinnedBones;
+
+		FSkinnedLevelSetCollisionData() = default;
+		~FSkinnedLevelSetCollisionData() = default;
+
+		FSkinnedLevelSetCollisionData(const TRefCountPtr<Chaos::TWeightedLatticeImplicitObject<Chaos::FLevelSet>>& InWeightedLevelSet, int32 InBoneIndex, TArray<int32>&& InMappedSkinnedBones)
+			: WeightedLevelSet(InWeightedLevelSet)
+			, BoneIndex(InBoneIndex)
+			, MappedSkinnedBones(MoveTemp(InMappedSkinnedBones))
+		{}
 	};
 
 	struct FMLLevelSetCollisionData
@@ -46,6 +64,15 @@ namespace Chaos
 		const TSharedPtr<Chaos::FMLLevelSet> MLLevelSet;
 		int32 BoneIndex;
 		TArray<int32> MappedActiveBoneIndices;
+
+		FMLLevelSetCollisionData() = default;
+		~FMLLevelSetCollisionData() = default;
+
+		FMLLevelSetCollisionData(const TSharedPtr<Chaos::FMLLevelSet>& InMLLevelSet, int32 InBoneIndex, TArray<int32>&& InMappedActiveBoneIndices)
+			: MLLevelSet(InMLLevelSet)
+			, BoneIndex(InBoneIndex)
+			, MappedActiveBoneIndices(MoveTemp(InMappedActiveBoneIndices))
+		{}
 	};
 
 	struct FSkinnedTriangleMeshCollisionData
@@ -53,6 +80,15 @@ namespace Chaos
 		const TRefCountPtr<Chaos::FSkinnedTriangleMesh> SkinnedTriangleMesh;
 		int32 BoneIndex;
 		TArray<int32> MappedSkinnedBones;
+
+		FSkinnedTriangleMeshCollisionData() = default;
+		~FSkinnedTriangleMeshCollisionData() = default;
+
+		FSkinnedTriangleMeshCollisionData(const TRefCountPtr<Chaos::FSkinnedTriangleMesh>& InSkinnedTriangleMesh, int32 InBoneIndex, TArray<int32>&& InMappedSkinnedBones)
+			: SkinnedTriangleMesh(InSkinnedTriangleMesh)
+			, BoneIndex(InBoneIndex)
+			, MappedSkinnedBones(MoveTemp(InMappedSkinnedBones))
+		{}
 	};
 
 	struct FClothCollisionDataExtended
diff --git a/Engine/Source/Runtime/Experimental/Chaos/Public/Chaos/ParticlesRange.h b/Engine/Source/Runtime/Experimental/Chaos/Public/Chaos/ParticlesRange.h
--- a/Engine/Source/Runtime/Experimental/Chaos/Public/Chaos/ParticlesRange.h
+++ b/Engine/Source/Runtime/Experimental/Chaos/Public/Chaos/ParticlesRange.h
@@ -12,6 +12,19 @@ struct FParticleRangeIndex
 {
 	int32 RangeId;
 	int32 ParticleIndex;
+
+	FParticleRangeIndex() = default;
+	~FParticleRangeIndex() = default;
+	FParticleRangeIndex(const FParticleRangeIndex&) = default;
+	FParticleRangeIndex(FParticleRangeIndex&&) = default;
+	FParticleRangeIndex& operator=(const FParticleRangeIndex&) = default;
+	FParticleRangeIndex& operator=(FParticleRangeIndex&&) = default;
+
+	FParticleRangeIndex(const int32 InRangeId, const int32 InParticleIndex)
+		: RangeId(InRangeId)
+		, ParticleIndex(InParticleIndex)
+	{}
+
 	bool operator<(const FParticleRangeIndex& Other) const
 	{
 		if (RangeId != Other.RangeId)
