Another CIS fix on top of the last CIS fix, this time for Android.

diff --git a/Engine/Plugins/ChaosCloth/Source/ChaosCloth/Public/ChaosCloth/ChaosClothingSimulationCollider.h b/Engine/Plugins/ChaosCloth/Source/ChaosCloth/Public/ChaosCloth/ChaosClothingSimulationCollider.h
--- a/Engine/Plugins/ChaosCloth/Source/ChaosCloth/Public/ChaosCloth/ChaosClothingSimulationCollider.h
+++ b/Engine/Plugins/ChaosCloth/Source/ChaosCloth/Public/ChaosCloth/ChaosClothingSimulationCollider.h
@@ -33,9 +33,6 @@ namespace Chaos
 		FTransform Transform;
 		int32 BoneIndex;
 
-		FLevelSetCollisionData() = default;
-		~FLevelSetCollisionData() = default;
-
 		FLevelSetCollisionData(const TSharedPtr<Chaos::FLevelSet, ESPMode::ThreadSafe>& InLevelSet, const FTransform& InTransform, int32 InBoneIndex)
 			: LevelSet(InLevelSet)
 			, Transform(InTransform)
@@ -49,9 +46,6 @@ namespace Chaos
 		int32 BoneIndex;
 		TArray<int32> MappedSkinnedBones;
 
-		FSkinnedLevelSetCollisionData() = default;
-		~FSkinnedLevelSetCollisionData() = default;
-
 		FSkinnedLevelSetCollisionData(const TRefCountPtr<Chaos::TWeightedLatticeImplicitObject<Chaos::FLevelSet>>& InWeightedLevelSet, int32 InBoneIndex, TArray<int32>&& InMappedSkinnedBones)
 			: WeightedLevelSet(InWeightedLevelSet)
 			, BoneIndex(InBoneIndex)
@@ -65,9 +59,6 @@ namespace Chaos
 		int32 BoneIndex;
 		TArray<int32> MappedActiveBoneIndices;
 
-		FMLLevelSetCollisionData() = default;
-		~FMLLevelSetCollisionData() = default;
-
 		FMLLevelSetCollisionData(const TSharedPtr<Chaos::FMLLevelSet>& InMLLevelSet, int32 InBoneIndex, TArray<int32>&& InMappedActiveBoneIndices)
 			: MLLevelSet(InMLLevelSet)
 			, BoneIndex(InBoneIndex)
@@ -81,9 +72,6 @@ namespace Chaos
 		int32 BoneIndex;
 		TArray<int32> MappedSkinnedBones;
 
-		FSkinnedTriangleMeshCollisionData() = default;
-		~FSkinnedTriangleMeshCollisionData() = default;
-
 		FSkinnedTriangleMeshCollisionData(const TRefCountPtr<Chaos::FSkinnedTriangleMesh>& InSkinnedTriangleMesh, int32 InBoneIndex, TArray<int32>&& InMappedSkinnedBones)
 			: SkinnedTriangleMesh(InSkinnedTriangleMesh)
 			, BoneIndex(InBoneIndex)
