Add ability to enable all Chaos Cloth ISPC CVars in Shipping configurations via define USE_ISPC_KERNEL_CONSOLE_VARIABLES_IN_SHIPPING. This is an attempt to give us mitigation options for an issue where one of our ISPC kernels has started crashing on PC in live, and we still haven't determined the cause.

Enabling this define increases the Shipping Client exe size by ~50KB for Windows.

This define is set to 0 by default everywhere, including in the new Chaos_Windows.Build.cs. It will set to 1 in a subsequent CL that won't robomerge back to UE.

diff --git a/Engine/Platforms/Windows/Source/Runtime/Experimental/Chaos/Chaos_Windows.Build.cs b/Engine/Platforms/Windows/Source/Runtime/Experimental/Chaos/Chaos_Windows.Build.cs
new file mode 100644
--- /dev/null
+++ b/Engine/Platforms/Windows/Source/Runtime/Experimental/Chaos/Chaos_Windows.Build.cs
@@ -0,0 +1,15 @@
+// Copyright Epic Games, Inc. All Rights Reserved.
+
+using System.Linq;
+using System.Collections.Generic;
+
+namespace UnrealBuildTool.Rules
+{
+    public class Chaos_Windows : Chaos
+    {
+        public Chaos_Windows(ReadOnlyTargetRules Target) : base(Target)
+        {
+			PublicDefinitions.Add("USE_ISPC_KERNEL_CONSOLE_VARIABLES_IN_SHIPPING=0");
+		}
+    }
+}
diff --git a/Engine/Plugins/ChaosCloth/Source/ChaosCloth/Private/ChaosCloth/ChaosClothingSimulation.cpp b/Engine/Plugins/ChaosCloth/Source/ChaosCloth/Private/ChaosCloth/ChaosClothingSimulation.cpp
--- a/Engine/Plugins/ChaosCloth/Source/ChaosCloth/Private/ChaosCloth/ChaosClothingSimulation.cpp
+++ b/Engine/Plugins/ChaosCloth/Source/ChaosCloth/Private/ChaosCloth/ChaosClothingSimulation.cpp
@@ -44,14 +44,16 @@
 #include "Chaos/SoftsMultiResConstraints.h"
 #endif
 
-#if INTEL_ISPC && !UE_BUILD_SHIPPING
+#if INTEL_ISPC
 static_assert(sizeof(ispc::FVector3f) == sizeof(FVector3f), "sizeof(ispc::FVector3f) != sizeof(FVector3f)");
 static_assert(sizeof(ispc::FVector) == sizeof(Chaos::FVec3), "sizeof(ispc::FVector) != sizeof(Chaos::FVec3)");
 static_assert(sizeof(ispc::FTransform) == sizeof(Chaos::FRigidTransform3), "sizeof(ispc::FTransform) != sizeof(Chaos::FRigidTransform3)");
 
+#if !UE_BUILD_SHIPPING || USE_ISPC_KERNEL_CONSOLE_VARIABLES_IN_SHIPPING
 bool bChaos_GetSimData_ISPC_Enabled = CHAOS_GET_SIM_DATA_ISPC_ENABLED_DEFAULT;
 FAutoConsoleVariableRef CVarChaosGetSimDataISPCEnabled(TEXT("p.Chaos.GetSimData.ISPC"), bChaos_GetSimData_ISPC_Enabled, TEXT("Whether to use ISPC optimizations when getting simulation data"));
 #endif
+#endif
 
 DECLARE_CYCLE_STAT(TEXT("Chaos Cloth Simulate"), STAT_ChaosClothSimulate, STATGROUP_ChaosCloth);
 DECLARE_CYCLE_STAT(TEXT("Chaos Cloth Create Actor"), STAT_ChaosClothCreateActor, STATGROUP_ChaosCloth);
diff --git a/Engine/Plugins/ChaosCloth/Source/ChaosCloth/Private/ChaosCloth/ChaosClothingSimulationMesh.cpp b/Engine/Plugins/ChaosCloth/Source/ChaosCloth/Private/ChaosCloth/ChaosClothingSimulationMesh.cpp
--- a/Engine/Plugins/ChaosCloth/Source/ChaosCloth/Private/ChaosCloth/ChaosClothingSimulationMesh.cpp
+++ b/Engine/Plugins/ChaosCloth/Source/ChaosCloth/Private/ChaosCloth/ChaosClothingSimulationMesh.cpp
@@ -16,9 +16,11 @@
 #include "ChaosClothingSimulationMesh.ispc.generated.h"
 #endif
 
-#if INTEL_ISPC && !UE_BUILD_SHIPPING
+#if INTEL_ISPC
+#if !UE_BUILD_SHIPPING || USE_ISPC_KERNEL_CONSOLE_VARIABLES_IN_SHIPPING
 bool bChaos_SkinPhysicsMesh_ISPC_Enabled = CHAOS_SKIN_PHYSICS_MESH_ISPC_ENABLED_DEFAULT;
 FAutoConsoleVariableRef CVarChaosSkinPhysicsMeshISPCEnabled(TEXT("p.Chaos.SkinPhysicsMesh.ISPC"), bChaos_SkinPhysicsMesh_ISPC_Enabled, TEXT("Whether to use ISPC optimizations on skinned physics meshes"));
+#endif
 
 static_assert(sizeof(ispc::FVector3f) == sizeof(Chaos::Softs::FSolverVec3), "sizeof(ispc::FVector3f) != sizeof(Chaos::Softs::FSolverVec3)");
 static_assert(sizeof(ispc::FVector3f) == sizeof(FVector3f), "sizeof(ispc::FVector3f) != sizeof(FVector3f)");
diff --git a/Engine/Plugins/ChaosCloth/Source/ChaosCloth/Private/ChaosCloth/ChaosClothingSimulationSolver.cpp b/Engine/Plugins/ChaosCloth/Source/ChaosCloth/Private/ChaosCloth/ChaosClothingSimulationSolver.cpp
--- a/Engine/Plugins/ChaosCloth/Source/ChaosCloth/Private/ChaosCloth/ChaosClothingSimulationSolver.cpp
+++ b/Engine/Plugins/ChaosCloth/Source/ChaosCloth/Private/ChaosCloth/ChaosClothingSimulationSolver.cpp
@@ -37,13 +37,16 @@ DECLARE_CYCLE_STAT(TEXT("Chaos Cloth Solver Particle Pre Simulation Transforms")
 DECLARE_CYCLE_STAT(TEXT("Chaos Cloth Solver Collision Pre Simulation Transforms"), STAT_ChaosClothCollisionPreSimulationTransforms, STATGROUP_ChaosCloth);
 DECLARE_CYCLE_STAT(TEXT("Chaos Cloth Solver Particle Pre Substep Kinematic Interpolation"), STAT_ChaosClothParticlePreSubstepKinematicInterpolation, STATGROUP_ChaosCloth);
 
-#if INTEL_ISPC && !UE_BUILD_SHIPPING
+#if INTEL_ISPC
+
+#if !UE_BUILD_SHIPPING || USE_ISPC_KERNEL_CONSOLE_VARIABLES_IN_SHIPPING
 bool bChaos_PreSimulationTransforms_ISPC_Enabled = CHAOS_PRE_SIMULATION_TRANSFORMS_ISPC_ENABLED_DEFAULT;
 FAutoConsoleVariableRef CVarChaosPreSimulationTransformsISPCEnabled(TEXT("p.Chaos.PreSimulationTransforms.ISPC"), bChaos_PreSimulationTransforms_ISPC_Enabled, TEXT("Whether to use ISPC optimizations in ApplySimulationTransforms"));
 bool bChaos_CalculateBounds_ISPC_Enabled = CHAOS_CALCULATE_BOUNDS_ISPC_ENABLED_DEFAULT;
 FAutoConsoleVariableRef CVarChaosCalculateBoundsISPCEnabled(TEXT("p.Chaos.CalculateBounds.ISPC"), bChaos_CalculateBounds_ISPC_Enabled, TEXT("Whether to use ISPC optimizations in CalculateBounds"));
 bool bChaos_PreSubstepInterpolation_ISPC_Enabled = CHAOS_PRE_SUBSTEP_INTERPOLATION_ISPC_ENABLED_DEFAULT;
 FAutoConsoleVariableRef CVarChaosPreSubstepInterpolationISPCEnabled(TEXT("p.Chaos.PreSubstepInterpolation.ISPC"), bChaos_PreSubstepInterpolation_ISPC_Enabled, TEXT("Whether to use ISPC optimization in PreSubstep"));
+#endif
 
 static_assert(sizeof(ispc::FVector3f) == sizeof(Chaos::Softs::FSolverVec3), "sizeof(ispc::FVector) != sizeof(Chaos::Softs::FSolverVec3)");
 static_assert(sizeof(ispc::FTransform3f) == sizeof(Chaos::Softs::FSolverRigidTransform3), "sizeof(ispc::FVector) != sizeof(Chaos::Softs::FSolverRigidTransform3)");
diff --git a/Engine/Plugins/ChaosCloth/Source/ChaosCloth/Public/ChaosCloth/ChaosClothingSimulation.h b/Engine/Plugins/ChaosCloth/Source/ChaosCloth/Public/ChaosCloth/ChaosClothingSimulation.h
--- a/Engine/Plugins/ChaosCloth/Source/ChaosCloth/Public/ChaosCloth/ChaosClothingSimulation.h
+++ b/Engine/Plugins/ChaosCloth/Source/ChaosCloth/Public/ChaosCloth/ChaosClothingSimulation.h
@@ -168,8 +168,12 @@ PRAGMA_ENABLE_DEPRECATION_WARNINGS
 #define CHAOS_GET_SIM_DATA_ISPC_ENABLED_DEFAULT 1
 #endif
 
+#if !defined(USE_ISPC_KERNEL_CONSOLE_VARIABLES_IN_SHIPPING)
+#define USE_ISPC_KERNEL_CONSOLE_VARIABLES_IN_SHIPPING 0
+#endif
+
 // Support run-time toggling on supported platforms in non-shipping configurations
-#if !INTEL_ISPC || UE_BUILD_SHIPPING
+#if !INTEL_ISPC || (UE_BUILD_SHIPPING && !USE_ISPC_KERNEL_CONSOLE_VARIABLES_IN_SHIPPING)
 static constexpr bool bChaos_GetSimData_ISPC_Enabled = INTEL_ISPC && CHAOS_GET_SIM_DATA_ISPC_ENABLED_DEFAULT;
 #else
 extern bool bChaos_GetSimData_ISPC_Enabled;
diff --git a/Engine/Plugins/ChaosCloth/Source/ChaosCloth/Public/ChaosCloth/ChaosClothingSimulationMesh.h b/Engine/Plugins/ChaosCloth/Source/ChaosCloth/Public/ChaosCloth/ChaosClothingSimulationMesh.h
--- a/Engine/Plugins/ChaosCloth/Source/ChaosCloth/Public/ChaosCloth/ChaosClothingSimulationMesh.h
+++ b/Engine/Plugins/ChaosCloth/Source/ChaosCloth/Public/ChaosCloth/ChaosClothingSimulationMesh.h
@@ -193,8 +193,12 @@ namespace Chaos
 #define CHAOS_SKIN_PHYSICS_MESH_ISPC_ENABLED_DEFAULT 1
 #endif
 
+#if !defined(USE_ISPC_KERNEL_CONSOLE_VARIABLES_IN_SHIPPING)
+#define USE_ISPC_KERNEL_CONSOLE_VARIABLES_IN_SHIPPING 0
+#endif
+
 // Support run-time toggling on supported platforms in non-shipping configurations
-#if !INTEL_ISPC || UE_BUILD_SHIPPING
+#if !INTEL_ISPC || (UE_BUILD_SHIPPING && !USE_ISPC_KERNEL_CONSOLE_VARIABLES_IN_SHIPPING)
 static constexpr bool bChaos_SkinPhysicsMesh_ISPC_Enabled = INTEL_ISPC && CHAOS_SKIN_PHYSICS_MESH_ISPC_ENABLED_DEFAULT;
 #else
 extern bool bChaos_SkinPhysicsMesh_ISPC_Enabled;
diff --git a/Engine/Plugins/ChaosCloth/Source/ChaosCloth/Public/ChaosCloth/ChaosClothingSimulationSolver.h b/Engine/Plugins/ChaosCloth/Source/ChaosCloth/Public/ChaosCloth/ChaosClothingSimulationSolver.h
--- a/Engine/Plugins/ChaosCloth/Source/ChaosCloth/Public/ChaosCloth/ChaosClothingSimulationSolver.h
+++ b/Engine/Plugins/ChaosCloth/Source/ChaosCloth/Public/ChaosCloth/ChaosClothingSimulationSolver.h
@@ -546,8 +546,12 @@ namespace Chaos
 #define CHAOS_PRE_SUBSTEP_INTERPOLATION_ISPC_ENABLED_DEFAULT 1
 #endif
 
+#if !defined(USE_ISPC_KERNEL_CONSOLE_VARIABLES_IN_SHIPPING)
+#define USE_ISPC_KERNEL_CONSOLE_VARIABLES_IN_SHIPPING 0
+#endif
+
 // Support run-time toggling on supported platforms in non-shipping configurations
-#if !INTEL_ISPC || UE_BUILD_SHIPPING
+#if !INTEL_ISPC || (UE_BUILD_SHIPPING && !USE_ISPC_KERNEL_CONSOLE_VARIABLES_IN_SHIPPING)
 static constexpr bool bChaos_CalculateBounds_ISPC_Enabled = INTEL_ISPC && CHAOS_CALCULATE_BOUNDS_ISPC_ENABLED_DEFAULT;
 static constexpr bool bChaos_PreSimulationTransforms_ISPC_Enabled = INTEL_ISPC && CHAOS_PRE_SIMULATION_TRANSFORMS_ISPC_ENABLED_DEFAULT;
 static constexpr bool bChaos_PreSubstepInterpolation_ISPC_Enabled = INTEL_ISPC && CHAOS_PRE_SUBSTEP_INTERPOLATION_ISPC_ENABLED_DEFAULT;
diff --git a/Engine/Source/Runtime/Experimental/Chaos/Private/Chaos/AABB.cpp b/Engine/Source/Runtime/Experimental/Chaos/Private/Chaos/AABB.cpp
--- a/Engine/Source/Runtime/Experimental/Chaos/Private/Chaos/AABB.cpp
+++ b/Engine/Source/Runtime/Experimental/Chaos/Private/Chaos/AABB.cpp
@@ -15,8 +15,12 @@ static_assert(sizeof(ispc::FVector) == sizeof(Chaos::TVector<Chaos::FReal, 3>),
 #define CHAOS_AABB_TRANSFORM_ISPC_ENABLED_DEFAULT 1
 #endif
 
+#if !defined(USE_ISPC_KERNEL_CONSOLE_VARIABLES_IN_SHIPPING)
+#define USE_ISPC_KERNEL_CONSOLE_VARIABLES_IN_SHIPPING 0
+#endif
+
 // Support run-time toggling on supported platforms in non-shipping configurations
-#if !INTEL_ISPC || UE_BUILD_SHIPPING
+#if !INTEL_ISPC || (UE_BUILD_SHIPPING && !USE_ISPC_KERNEL_CONSOLE_VARIABLES_IN_SHIPPING)
 static constexpr bool bChaos_AABBTransform_ISPC_Enabled = INTEL_ISPC && CHAOS_AABB_TRANSFORM_ISPC_ENABLED_DEFAULT;
 #else
 static bool bChaos_AABBTransform_ISPC_Enabled = CHAOS_AABB_TRANSFORM_ISPC_ENABLED_DEFAULT;
diff --git a/Engine/Source/Runtime/Experimental/Chaos/Private/Chaos/PBDAxialSpringConstraints.cpp b/Engine/Source/Runtime/Experimental/Chaos/Private/Chaos/PBDAxialSpringConstraints.cpp
--- a/Engine/Source/Runtime/Experimental/Chaos/Private/Chaos/PBDAxialSpringConstraints.cpp
+++ b/Engine/Source/Runtime/Experimental/Chaos/Private/Chaos/PBDAxialSpringConstraints.cpp
@@ -9,9 +9,11 @@
 
 DECLARE_CYCLE_STAT(TEXT("Chaos PBD Axial Spring Constraint"), STAT_PBD_AxialSpring, STATGROUP_Chaos);
 
-#if INTEL_ISPC && !UE_BUILD_SHIPPING
+#if INTEL_ISPC
+#if !UE_BUILD_SHIPPING || USE_ISPC_KERNEL_CONSOLE_VARIABLES_IN_SHIPPING
 bool bChaos_AxialSpring_ISPC_Enabled = CHAOS_AXIAL_SPRING_ISPC_ENABLED_DEFAULT;
 FAutoConsoleVariableRef CVarChaosAxialSpringISPCEnabled(TEXT("p.Chaos.AxialSpring.ISPC"), bChaos_AxialSpring_ISPC_Enabled, TEXT("Whether to use ISPC optimizations in AxialSpring constraints"));
+#endif
 
 static_assert(sizeof(ispc::FVector4f) == sizeof(Chaos::Softs::FPAndInvM), "sizeof(ispc::FVector4f) != sizeof(Chaos::Softs::FPAndInvM");
 static_assert(sizeof(ispc::FIntVector) == sizeof(Chaos::TVec3<int32>), "sizeof(ispc::FIntVector) != sizeof(Chaos::TVec3<int32>");
diff --git a/Engine/Source/Runtime/Experimental/Chaos/Private/Chaos/PBDBendingConstraintsBase.cpp b/Engine/Source/Runtime/Experimental/Chaos/Private/Chaos/PBDBendingConstraintsBase.cpp
--- a/Engine/Source/Runtime/Experimental/Chaos/Private/Chaos/PBDBendingConstraintsBase.cpp
+++ b/Engine/Source/Runtime/Experimental/Chaos/Private/Chaos/PBDBendingConstraintsBase.cpp
@@ -4,9 +4,11 @@
 #include "PBDBendingConstraints.ispc.generated.h"
 #endif
 
-#if INTEL_ISPC && !UE_BUILD_SHIPPING
+#if INTEL_ISPC 
+#if !UE_BUILD_SHIPPING || USE_ISPC_KERNEL_CONSOLE_VARIABLES_IN_SHIPPING
 bool bChaos_Bending_ISPC_Enabled = CHAOS_BENDING_ISPC_ENABLED_DEFAULT;
 FAutoConsoleVariableRef CVarChaosBendingISPCEnabled(TEXT("p.Chaos.Bending.ISPC"), bChaos_Bending_ISPC_Enabled, TEXT("Whether to use ISPC optimizations in Bending constraints"));
+#endif
 
 static_assert(sizeof(ispc::FVector4f) == sizeof(Chaos::Softs::FPAndInvM), "sizeof(ispc::FVector4f) != sizeof(Chaos::Softs::FPAndInvM");
 static_assert(sizeof(ispc::FIntVector4) == sizeof(Chaos::TVec4<int32>), "sizeof(ispc::FIntVector4) != sizeof(Chaos::TVec4<int32>");
diff --git a/Engine/Source/Runtime/Experimental/Chaos/Private/Chaos/PBDLongRangeConstraints.cpp b/Engine/Source/Runtime/Experimental/Chaos/Private/Chaos/PBDLongRangeConstraints.cpp
--- a/Engine/Source/Runtime/Experimental/Chaos/Private/Chaos/PBDLongRangeConstraints.cpp
+++ b/Engine/Source/Runtime/Experimental/Chaos/Private/Chaos/PBDLongRangeConstraints.cpp
@@ -5,13 +5,16 @@
 #include "PBDLongRangeConstraints.ispc.generated.h"
 #endif
 
-#if INTEL_ISPC && !UE_BUILD_SHIPPING
+#if INTEL_ISPC
 static_assert(sizeof(ispc::FVector4f) == sizeof(Chaos::Softs::FPAndInvM), "sizeof(ispc::FVector4f) != sizeof(Chaos::Softs::FPAndInvM)");
 static_assert(sizeof(ispc::FTether) == sizeof(Chaos::Softs::FPBDLongRangeConstraintsBase::FTether), "sizeof(ispc::FTether) != sizeof(Chaos::Softs::FPBDLongRangeConstraintsBase::FTether)");
 
+// Support run-time toggling on supported platforms in non-shipping configurations
+#if !UE_BUILD_SHIPPING || USE_ISPC_KERNEL_CONSOLE_VARIABLES_IN_SHIPPING
 bool bChaos_LongRange_ISPC_Enabled = CHAOS_LONG_RANGE_ISPC_ENABLED_DEFAULT;
 FAutoConsoleVariableRef CVarChaosLongRangeISPCEnabled(TEXT("p.Chaos.LongRange.ISPC"), bChaos_LongRange_ISPC_Enabled, TEXT("Whether to use ISPC optimizations in long range constraints"));
 #endif
+#endif
 
 namespace Chaos::Softs {
 
diff --git a/Engine/Source/Runtime/Experimental/Chaos/Private/Chaos/PBDSphericalConstraint.cpp b/Engine/Source/Runtime/Experimental/Chaos/Private/Chaos/PBDSphericalConstraint.cpp
--- a/Engine/Source/Runtime/Experimental/Chaos/Private/Chaos/PBDSphericalConstraint.cpp
+++ b/Engine/Source/Runtime/Experimental/Chaos/Private/Chaos/PBDSphericalConstraint.cpp
@@ -9,9 +9,11 @@
 #include "PBDSphericalConstraint.ispc.generated.h"
 #endif
 
-#if INTEL_ISPC && !UE_BUILD_SHIPPING
+#if INTEL_ISPC
+#if !UE_BUILD_SHIPPING || USE_ISPC_KERNEL_CONSOLE_VARIABLES_IN_SHIPPING
 bool bChaos_Spherical_ISPC_Enabled = CHAOS_SPHERICAL_ISPC_ENABLED_DEFAULT;
 FAutoConsoleVariableRef CVarChaosSphericalISPCEnabled(TEXT("p.Chaos.Spherical.ISPC"), bChaos_Spherical_ISPC_Enabled, TEXT("Whether to use ISPC optimizations in spherical constraints"));
+#endif
 
 static_assert(sizeof(ispc::FVector4f) == sizeof(Chaos::Softs::FPAndInvM), "sizeof(ispc::FVector4f) != sizeof(Chaos::Softs::FPAndInvM)");
 #endif
diff --git a/Engine/Source/Runtime/Experimental/Chaos/Private/Chaos/PBDSpringConstraints.cpp b/Engine/Source/Runtime/Experimental/Chaos/Private/Chaos/PBDSpringConstraints.cpp
--- a/Engine/Source/Runtime/Experimental/Chaos/Private/Chaos/PBDSpringConstraints.cpp
+++ b/Engine/Source/Runtime/Experimental/Chaos/Private/Chaos/PBDSpringConstraints.cpp
@@ -14,13 +14,15 @@
 
 DECLARE_CYCLE_STAT(TEXT("Chaos PBD Spring Constraint"), STAT_PBD_Spring, STATGROUP_Chaos);
 
-#if INTEL_ISPC && !UE_BUILD_SHIPPING
+#if INTEL_ISPC
 static_assert(sizeof(ispc::FVector4f) == sizeof(Chaos::Softs::FPAndInvM), "sizeof(ispc::FVector4f) != sizeof(Chaos::Softs::FPAndInvM)");
 static_assert(sizeof(ispc::FIntVector2) == sizeof(Chaos::TVec2<int32>), "sizeof(ispc::FIntVector2) != sizeof(Chaos::TVec2<int32>)");
 
+#if !UE_BUILD_SHIPPING || USE_ISPC_KERNEL_CONSOLE_VARIABLES_IN_SHIPPING
 bool bChaos_Spring_ISPC_Enabled = CHAOS_SPRING_ISPC_ENABLED_DEFAULT;
 FAutoConsoleVariableRef CVarChaosSpringISPCEnabled(TEXT("p.Chaos.Spring.ISPC"), bChaos_Spring_ISPC_Enabled, TEXT("Whether to use ISPC optimizations in Spring constraints"));
 #endif
+#endif
 
 namespace Chaos::Softs {
 
diff --git a/Engine/Source/Runtime/Experimental/Chaos/Private/Chaos/PerParticleDampVelocity.cpp b/Engine/Source/Runtime/Experimental/Chaos/Private/Chaos/PerParticleDampVelocity.cpp
--- a/Engine/Source/Runtime/Experimental/Chaos/Private/Chaos/PerParticleDampVelocity.cpp
+++ b/Engine/Source/Runtime/Experimental/Chaos/Private/Chaos/PerParticleDampVelocity.cpp
@@ -13,11 +13,14 @@
 #include "PerParticleDampVelocity.ispc.generated.h"
 #endif
 
-#if INTEL_ISPC && !UE_BUILD_SHIPPING
+#if INTEL_ISPC
 static_assert(sizeof(ispc::FVector3f) == sizeof(Chaos::Softs::FSolverVec3), "sizeof(ispc::FVector3f) != sizeof(Chaos::Softs::FSolverVec3)");
+
+#if !UE_BUILD_SHIPPING || USE_ISPC_KERNEL_CONSOLE_VARIABLES_IN_SHIPPING
 bool bChaos_DampVelocity_ISPC_Enabled = CHAOS_DAMP_VELOCITY_ISPC_ENABLED_DEFAULT;
 static FAutoConsoleVariableRef CVarChaosDampVelocityISPCEnabled(TEXT("p.Chaos.DampVelocity.ISPC"), bChaos_DampVelocity_ISPC_Enabled, TEXT("Whether to use ISPC optimizations in per particle damp velocity calculation"));
 #endif
+#endif
 
 namespace Chaos::Softs {
 
diff --git a/Engine/Source/Runtime/Experimental/Chaos/Private/Chaos/PerParticlePBDCollisionConstraint.cpp b/Engine/Source/Runtime/Experimental/Chaos/Private/Chaos/PerParticlePBDCollisionConstraint.cpp
--- a/Engine/Source/Runtime/Experimental/Chaos/Private/Chaos/PerParticlePBDCollisionConstraint.cpp
+++ b/Engine/Source/Runtime/Experimental/Chaos/Private/Chaos/PerParticlePBDCollisionConstraint.cpp
@@ -106,7 +106,7 @@ static_assert(sizeof(ispc::VerticesS) == sizeof(Chaos::FConvexStructureData::FCo
 static_assert(sizeof(ispc::VerticesM) == sizeof(Chaos::FConvexStructureData::FConvexStructureDataMedium::FVertexData));
 static_assert(sizeof(ispc::VerticesL) == sizeof(Chaos::FConvexStructureData::FConvexStructureDataLarge::FVertexData));
 
-#if !UE_BUILD_SHIPPING
+#if !UE_BUILD_SHIPPING || USE_ISPC_KERNEL_CONSOLE_VARIABLES_IN_SHIPPING 
 bool bChaos_PerParticleCollision_ISPC_Enabled = CHAOS_PER_PARTICLE_COLLISION_ISPC_ENABLED_DEFAULT;
 static FAutoConsoleVariableRef CVarChaosPerParticleCollisionISPCEnabled(TEXT("p.Chaos.PerParticleCollision.ISPC"), bChaos_PerParticleCollision_ISPC_Enabled, TEXT("Whether to use ISPC optimizations in per particle collisions"));
 #endif
diff --git a/Engine/Source/Runtime/Experimental/Chaos/Private/Chaos/TriangleMesh.cpp b/Engine/Source/Runtime/Experimental/Chaos/Private/Chaos/TriangleMesh.cpp
--- a/Engine/Source/Runtime/Experimental/Chaos/Private/Chaos/TriangleMesh.cpp
+++ b/Engine/Source/Runtime/Experimental/Chaos/Private/Chaos/TriangleMesh.cpp
@@ -28,8 +28,12 @@ static_assert(sizeof(ispc::TArrayInt) == sizeof(TArray<int32>), "sizeof(ispc::TA
 #define CHAOS_TRIANGLE_MESH_ISPC_ENABLED_DEFAULT 1
 #endif
 
+#if !defined(USE_ISPC_KERNEL_CONSOLE_VARIABLES_IN_SHIPPING)
+#define USE_ISPC_KERNEL_CONSOLE_VARIABLES_IN_SHIPPING 0
+#endif
+
 // Support run-time toggling on supported platforms in non-shipping configurations
-#if !INTEL_ISPC || UE_BUILD_SHIPPING
+#if !INTEL_ISPC || (UE_BUILD_SHIPPING && !USE_ISPC_KERNEL_CONSOLE_VARIABLES_IN_SHIPPING)
 static constexpr bool bChaos_TriangleMesh_ISPC_Enabled = INTEL_ISPC && CHAOS_TRIANGLE_MESH_ISPC_ENABLED_DEFAULT;
 #else
 static bool bChaos_TriangleMesh_ISPC_Enabled = CHAOS_TRIANGLE_MESH_ISPC_ENABLED_DEFAULT;
diff --git a/Engine/Source/Runtime/Experimental/Chaos/Private/Chaos/VelocityField.cpp b/Engine/Source/Runtime/Experimental/Chaos/Private/Chaos/VelocityField.cpp
--- a/Engine/Source/Runtime/Experimental/Chaos/Private/Chaos/VelocityField.cpp
+++ b/Engine/Source/Runtime/Experimental/Chaos/Private/Chaos/VelocityField.cpp
@@ -11,7 +11,7 @@ static_assert(sizeof(ispc::FVector2f) == sizeof(Chaos::Softs::FSolverVec2), "siz
 static_assert(sizeof(ispc::FIntVector) == sizeof(Chaos::TVec3<int32>), "sizeof(ispc::FIntVector) != sizeof(Chaos::TVec3<int32>)");
 #endif
 
-#if INTEL_ISPC && !UE_BUILD_SHIPPING
+#if INTEL_ISPC && (!UE_BUILD_SHIPPING || USE_ISPC_KERNEL_CONSOLE_VARIABLES_IN_SHIPPING)
 bool bChaos_VelocityField_ISPC_Enabled = CHAOS_VELOCITY_FIELD_ISPC_ENABLED_DEFAULT;
 FAutoConsoleVariableRef CVarChaosVelocityFieldISPCEnabled(TEXT("p.Chaos.VelocityField.ISPC"), bChaos_VelocityField_ISPC_Enabled, TEXT("Whether to use ISPC optimizations in velocity field calculations"));
 #endif
diff --git a/Engine/Source/Runtime/Experimental/Chaos/Public/Chaos/PBDAxialSpringConstraints.h b/Engine/Source/Runtime/Experimental/Chaos/Public/Chaos/PBDAxialSpringConstraints.h
--- a/Engine/Source/Runtime/Experimental/Chaos/Public/Chaos/PBDAxialSpringConstraints.h
+++ b/Engine/Source/Runtime/Experimental/Chaos/Public/Chaos/PBDAxialSpringConstraints.h
@@ -134,8 +134,12 @@ private:
 #define CHAOS_AXIAL_SPRING_ISPC_ENABLED_DEFAULT 1
 #endif
 
+#if !defined(USE_ISPC_KERNEL_CONSOLE_VARIABLES_IN_SHIPPING)
+#define USE_ISPC_KERNEL_CONSOLE_VARIABLES_IN_SHIPPING 0
+#endif
+
 // Support run-time toggling on supported platforms in non-shipping configurations
-#if !INTEL_ISPC || UE_BUILD_SHIPPING
+#if !INTEL_ISPC || (UE_BUILD_SHIPPING && !USE_ISPC_KERNEL_CONSOLE_VARIABLES_IN_SHIPPING)
 static constexpr bool bChaos_AxialSpring_ISPC_Enabled = INTEL_ISPC && CHAOS_AXIAL_SPRING_ISPC_ENABLED_DEFAULT;
 #else
 extern CHAOS_API bool bChaos_AxialSpring_ISPC_Enabled;
diff --git a/Engine/Source/Runtime/Experimental/Chaos/Public/Chaos/PBDBendingConstraintsBase.h b/Engine/Source/Runtime/Experimental/Chaos/Public/Chaos/PBDBendingConstraintsBase.h
--- a/Engine/Source/Runtime/Experimental/Chaos/Public/Chaos/PBDBendingConstraintsBase.h
+++ b/Engine/Source/Runtime/Experimental/Chaos/Public/Chaos/PBDBendingConstraintsBase.h
@@ -404,8 +404,12 @@ protected:
 #define CHAOS_BENDING_ISPC_ENABLED_DEFAULT 1
 #endif
 
-// Support ISPC enable/disable in non-shipping builds
-#if !INTEL_ISPC || UE_BUILD_SHIPPING
+#if !defined(USE_ISPC_KERNEL_CONSOLE_VARIABLES_IN_SHIPPING)
+#define USE_ISPC_KERNEL_CONSOLE_VARIABLES_IN_SHIPPING 0
+#endif
+
+// Support run-time toggling on supported platforms in non-shipping configurations
+#if !INTEL_ISPC || (UE_BUILD_SHIPPING && !USE_ISPC_KERNEL_CONSOLE_VARIABLES_IN_SHIPPING)
 static constexpr bool bChaos_Bending_ISPC_Enabled = INTEL_ISPC && CHAOS_BENDING_ISPC_ENABLED_DEFAULT;
 #else
 extern CHAOS_API bool bChaos_Bending_ISPC_Enabled;
diff --git a/Engine/Source/Runtime/Experimental/Chaos/Public/Chaos/PBDLongRangeConstraints.h b/Engine/Source/Runtime/Experimental/Chaos/Public/Chaos/PBDLongRangeConstraints.h
--- a/Engine/Source/Runtime/Experimental/Chaos/Public/Chaos/PBDLongRangeConstraints.h
+++ b/Engine/Source/Runtime/Experimental/Chaos/Public/Chaos/PBDLongRangeConstraints.h
@@ -118,8 +118,12 @@ private:
 #define CHAOS_LONG_RANGE_ISPC_ENABLED_DEFAULT 1
 #endif
 
+#if !defined(USE_ISPC_KERNEL_CONSOLE_VARIABLES_IN_SHIPPING)
+#define USE_ISPC_KERNEL_CONSOLE_VARIABLES_IN_SHIPPING 0
+#endif
+
 // Support run-time toggling on supported platforms in non-shipping configurations
-#if !INTEL_ISPC || UE_BUILD_SHIPPING
+#if !INTEL_ISPC || (UE_BUILD_SHIPPING && !USE_ISPC_KERNEL_CONSOLE_VARIABLES_IN_SHIPPING)
 static constexpr bool bChaos_LongRange_ISPC_Enabled = INTEL_ISPC && CHAOS_LONG_RANGE_ISPC_ENABLED_DEFAULT;
 #else
 extern CHAOS_API bool bChaos_LongRange_ISPC_Enabled;
diff --git a/Engine/Source/Runtime/Experimental/Chaos/Public/Chaos/PBDSphericalConstraint.h b/Engine/Source/Runtime/Experimental/Chaos/Public/Chaos/PBDSphericalConstraint.h
--- a/Engine/Source/Runtime/Experimental/Chaos/Public/Chaos/PBDSphericalConstraint.h
+++ b/Engine/Source/Runtime/Experimental/Chaos/Public/Chaos/PBDSphericalConstraint.h
@@ -17,8 +17,12 @@ DECLARE_CYCLE_STAT(TEXT("Chaos PBD Spherical Backstop Constraints"), STAT_PBD_Sp
 #define CHAOS_SPHERICAL_ISPC_ENABLED_DEFAULT 1
 #endif
 
+#if !defined(USE_ISPC_KERNEL_CONSOLE_VARIABLES_IN_SHIPPING)
+#define USE_ISPC_KERNEL_CONSOLE_VARIABLES_IN_SHIPPING 0
+#endif
+
 // Support run-time toggling on supported platforms in non-shipping configurations
-#if !INTEL_ISPC || UE_BUILD_SHIPPING
+#if !INTEL_ISPC || (UE_BUILD_SHIPPING && !USE_ISPC_KERNEL_CONSOLE_VARIABLES_IN_SHIPPING)
 static constexpr bool bChaos_Spherical_ISPC_Enabled = INTEL_ISPC && CHAOS_SPHERICAL_ISPC_ENABLED_DEFAULT;
 #else
 extern CHAOS_API bool bChaos_Spherical_ISPC_Enabled;
diff --git a/Engine/Source/Runtime/Experimental/Chaos/Public/Chaos/PBDSpringConstraints.h b/Engine/Source/Runtime/Experimental/Chaos/Public/Chaos/PBDSpringConstraints.h
--- a/Engine/Source/Runtime/Experimental/Chaos/Public/Chaos/PBDSpringConstraints.h
+++ b/Engine/Source/Runtime/Experimental/Chaos/Public/Chaos/PBDSpringConstraints.h
@@ -188,8 +188,12 @@ private:
 #define CHAOS_SPRING_ISPC_ENABLED_DEFAULT 1
 #endif
 
+#if !defined(USE_ISPC_KERNEL_CONSOLE_VARIABLES_IN_SHIPPING)
+#define USE_ISPC_KERNEL_CONSOLE_VARIABLES_IN_SHIPPING 0
+#endif
+
 // Support run-time toggling on supported platforms in non-shipping configurations
-#if !INTEL_ISPC || UE_BUILD_SHIPPING
+#if !INTEL_ISPC || (UE_BUILD_SHIPPING && !USE_ISPC_KERNEL_CONSOLE_VARIABLES_IN_SHIPPING)
 static constexpr bool bChaos_Spring_ISPC_Enabled = INTEL_ISPC && CHAOS_SPRING_ISPC_ENABLED_DEFAULT;
 #else
 extern CHAOS_API bool bChaos_Spring_ISPC_Enabled;
diff --git a/Engine/Source/Runtime/Experimental/Chaos/Public/Chaos/PerParticleDampVelocity.h b/Engine/Source/Runtime/Experimental/Chaos/Public/Chaos/PerParticleDampVelocity.h
--- a/Engine/Source/Runtime/Experimental/Chaos/Public/Chaos/PerParticleDampVelocity.h
+++ b/Engine/Source/Runtime/Experimental/Chaos/Public/Chaos/PerParticleDampVelocity.h
@@ -44,11 +44,15 @@ private:
 #define CHAOS_DAMP_VELOCITY_ISPC_ENABLED_DEFAULT 1
 #endif
 
+#if !defined(USE_ISPC_KERNEL_CONSOLE_VARIABLES_IN_SHIPPING)
+#define USE_ISPC_KERNEL_CONSOLE_VARIABLES_IN_SHIPPING 0
+#endif
+
 // Support run-time toggling on supported platforms in non-shipping configurations
 #if INTEL_ISPC
-#if UE_BUILD_SHIPPING
+#if UE_BUILD_SHIPPING && !USE_ISPC_KERNEL_CONSOLE_VARIABLES_IN_SHIPPING
 static constexpr bool bChaos_DampVelocity_ISPC_Enabled = CHAOS_DAMP_VELOCITY_ISPC_ENABLED_DEFAULT;
 #else
 extern CHAOS_API bool bChaos_DampVelocity_ISPC_Enabled;
-#endif // UE_BUILD_SHIPPING
-#endif
\ No newline at end of file
+#endif // UE_BUILD_SHIPPING && !USE_ISPC_KERNEL_CONSOLE_VARIABLES_IN_SHIPPING
+#endif
diff --git a/Engine/Source/Runtime/Experimental/Chaos/Public/Chaos/PerParticlePBDCollisionConstraint.h b/Engine/Source/Runtime/Experimental/Chaos/Public/Chaos/PerParticlePBDCollisionConstraint.h
--- a/Engine/Source/Runtime/Experimental/Chaos/Public/Chaos/PerParticlePBDCollisionConstraint.h
+++ b/Engine/Source/Runtime/Experimental/Chaos/Public/Chaos/PerParticlePBDCollisionConstraint.h
@@ -17,8 +17,12 @@
 #define CHAOS_PER_PARTICLE_COLLISION_ISPC_ENABLED_DEFAULT 1
 #endif
 
+#if !defined(USE_ISPC_KERNEL_CONSOLE_VARIABLES_IN_SHIPPING)
+#define USE_ISPC_KERNEL_CONSOLE_VARIABLES_IN_SHIPPING 0
+#endif
+
 // Support run-time toggling on supported platforms in non-shipping configurations
-#if !INTEL_ISPC || UE_BUILD_SHIPPING
+#if !INTEL_ISPC || (UE_BUILD_SHIPPING && !USE_ISPC_KERNEL_CONSOLE_VARIABLES_IN_SHIPPING)
 const bool bChaos_PerParticleCollision_ISPC_Enabled = INTEL_ISPC && CHAOS_PER_PARTICLE_COLLISION_ISPC_ENABLED_DEFAULT;
 #else
 extern CHAOS_API bool bChaos_PerParticleCollision_ISPC_Enabled;
diff --git a/Engine/Source/Runtime/Experimental/Chaos/Public/Chaos/VelocityField.h b/Engine/Source/Runtime/Experimental/Chaos/Public/Chaos/VelocityField.h
--- a/Engine/Source/Runtime/Experimental/Chaos/Public/Chaos/VelocityField.h
+++ b/Engine/Source/Runtime/Experimental/Chaos/Public/Chaos/VelocityField.h
@@ -323,8 +323,12 @@ private:
 #define CHAOS_VELOCITY_FIELD_ISPC_ENABLED_DEFAULT 1
 #endif
 
+#if !defined(USE_ISPC_KERNEL_CONSOLE_VARIABLES_IN_SHIPPING)
+#define USE_ISPC_KERNEL_CONSOLE_VARIABLES_IN_SHIPPING 0
+#endif
+
 // Support run-time toggling on supported platforms in non-shipping configurations
-#if !INTEL_ISPC || UE_BUILD_SHIPPING
+#if !INTEL_ISPC || (UE_BUILD_SHIPPING && !USE_ISPC_KERNEL_CONSOLE_VARIABLES_IN_SHIPPING)
 static constexpr bool bChaos_VelocityField_ISPC_Enabled = INTEL_ISPC && CHAOS_VELOCITY_FIELD_ISPC_ENABLED_DEFAULT;
 #else
 extern CHAOS_API bool bChaos_VelocityField_ISPC_Enabled;
