Chaos cloth: Fix Acceleration clamping by comparing post-clamped velocities rather than pre-clamped velocities.

diff --git a/Engine/Plugins/ChaosCloth/Source/ChaosCloth/Private/ChaosCloth/ChaosClothingSimulationCloth.cpp b/Engine/Plugins/ChaosCloth/Source/ChaosCloth/Private/ChaosCloth/ChaosClothingSimulationCloth.cpp
--- a/Engine/Plugins/ChaosCloth/Source/ChaosCloth/Private/ChaosCloth/ChaosClothingSimulationCloth.cpp
+++ b/Engine/Plugins/ChaosCloth/Source/ChaosCloth/Private/ChaosCloth/ChaosClothingSimulationCloth.cpp
@@ -931,8 +931,8 @@ void FClothingSimulationCloth::Update(FClothingSimulationSolver* Solver)
 			OutLinearVelocityScale = FVec3f(1.f);
 			OutAngularVelocityScale = 1.f;
 			OutMaxVelocityScale = 1.f;
-			ReferenceSpaceAngularVelocity = FVec3(0.);
-			ReferenceSpaceVelocity = FVec3(0.);
+			AppliedReferenceSpaceAngularVelocity = FVec3(0.);
+			AppliedReferenceSpaceVelocity = FVec3(0.);
 			bDisableFictitiousForces = true; // It doesn't actually matter what value we set here since AngularVelocityScale == 1 means fictitious forces will be 0.
 
 			// Reset to start pose
@@ -954,8 +954,8 @@ void FClothingSimulationCloth::Update(FClothingSimulationSolver* Solver)
 			OutLinearVelocityScale = FVec3f(0.f);
 			OutAngularVelocityScale = 0.f;
 			OutMaxVelocityScale = 1.f;
-			ReferenceSpaceAngularVelocity = FVec3(0.);
-			ReferenceSpaceVelocity = FVec3(0.);
+			AppliedReferenceSpaceAngularVelocity = FVec3(0.);
+			AppliedReferenceSpaceVelocity = FVec3(0.);
 			bDisableFictitiousForces = true; // Disable fictitious forces. Otherwise they will be applied since AngularVelocityScale < 1.
 			UE_LOG(LogChaosCloth, Verbose, TEXT("Cloth in group Id %d Needs teleport."), GroupId);
 		}
@@ -978,8 +978,8 @@ void FClothingSimulationCloth::Update(FClothingSimulationSolver* Solver)
 			GroupId,
 			OldReferenceSpaceTransform,
 			ReferenceSpaceTransform,
-			ReferenceSpaceVelocity,
-			ReferenceSpaceAngularVelocity,
+			AppliedReferenceSpaceVelocity,
+			AppliedReferenceSpaceAngularVelocity,
 			VelocityScaleSpace,
 			OutLinearVelocityScale,
 			MaxLinearVelocity,
diff --git a/Engine/Plugins/ChaosCloth/Source/ChaosCloth/Private/ChaosCloth/ChaosClothingSimulationSolver.cpp b/Engine/Plugins/ChaosCloth/Source/ChaosCloth/Private/ChaosCloth/ChaosClothingSimulationSolver.cpp
--- a/Engine/Plugins/ChaosCloth/Source/ChaosCloth/Private/ChaosCloth/ChaosClothingSimulationSolver.cpp
+++ b/Engine/Plugins/ChaosCloth/Source/ChaosCloth/Private/ChaosCloth/ChaosClothingSimulationSolver.cpp
@@ -1094,8 +1094,8 @@ void FClothingSimulationSolver::SetParticleMassFromDensity(int32 ParticleRangeId
 void FClothingSimulationSolver::SetReferenceVelocityScale(uint32 GroupId,
 	const FRigidTransform3& OldReferenceSpaceTransform,
 	const FRigidTransform3& ReferenceSpaceTransform,
-	TVec3<FReal>& InOutReferenceVelocity, // Old reference velocity is passed in. New reference velocity is returned.
-	TVec3<FReal>& InOutReferenceAngularVelocity, // Old reference velocity is passed in. New reference velocity is returned.
+	TVec3<FReal>& InOutAppliedReferenceVelocity, // Old reference velocity is passed in. New reference velocity is returned.
+	TVec3<FReal>& InOutAppliedReferenceAngularVelocity, // Old reference velocity is passed in. New reference velocity is returned.
 	const EChaosSoftsSimulationSpace VelocityScaleSpace, // the space the following linear velocity properties are in.
 	const TVec3<FRealSingle>& LinearVelocityScale,
 	const TVec3<FRealSingle>& MaxLinearVelocity,
@@ -1114,8 +1114,8 @@ void FClothingSimulationSolver::SetReferenceVelocityScale(uint32 GroupId,
 	const FReal OldLocalSpaceScaleInv = 1. / OldLocalSpaceScale;
 	OldRootBoneLocalTransform.ScaleTranslation(OldLocalSpaceScaleInv);
 
-	const FVec3 OldReferenceVelocity(InOutReferenceVelocity); // In world space
-	const FVec3 OldReferenceAngularVelocity(InOutReferenceAngularVelocity); // In world space
+	const FVec3 OldAppliedReferenceVelocity(InOutAppliedReferenceVelocity); // In world space
+	const FVec3 OldAppliedReferenceAngularVelocity(InOutAppliedReferenceAngularVelocity); // In world space
 
 	const FReal SolverVelocityScale = bClothSolverUseVelocityScale ? VelocityScale : (FReal)1.;
 
@@ -1177,11 +1177,7 @@ void FClothingSimulationSolver::SetReferenceVelocityScale(uint32 GroupId,
 
 	// Apply clamps (since we're ultimately just modifying a single transform, remove accelerations from the velocity first, then clamp velocities).
 	if (DeltaTime > 0.f)
-	{
-		// Only update old reference velocity when ticking.
-		InOutReferenceVelocity = WorldDeltaTranslation / DeltaTime;
-		InOutReferenceAngularVelocity = FullDeltaRotation / DeltaTime;
-	
+	{	
 		FVec3 AppliedLinearVelocity = AppliedDeltaTranslation / DeltaTime;
 		FVec3 AppliedAngularVelocity = AppliedDeltaRotation / DeltaTime;
 
@@ -1192,14 +1188,14 @@ void FClothingSimulationSolver::SetReferenceVelocityScale(uint32 GroupId,
 			switch (VelocityScaleSpace)
 			{
 			case EChaosSoftsSimulationSpace::WorldSpace:
-				OldScaledVelocity = LinearRatio * OldReferenceVelocity;
+				OldScaledVelocity = OldAppliedReferenceVelocity;
 				break;
 			case EChaosSoftsSimulationSpace::ComponentSpace:
-				OldScaledVelocity = LinearRatio * LocalSpaceRotation.UnrotateVector(OldReferenceVelocity);
+				OldScaledVelocity = LocalSpaceRotation.UnrotateVector(OldAppliedReferenceVelocity);
 				break;
 			case EChaosSoftsSimulationSpace::ReferenceBoneSpace:
 			default:
-				OldScaledVelocity = LinearRatio * (OldReferenceSpaceRotationInverse * OldReferenceVelocity);
+				OldScaledVelocity = (OldReferenceSpaceRotationInverse * OldAppliedReferenceVelocity);
 				break;
 			}
 			
@@ -1214,7 +1210,7 @@ void FClothingSimulationSolver::SetReferenceVelocityScale(uint32 GroupId,
 			}
 		}
 
-		// Apply linear velocity clamping
+		// Apply linear velocity clamping 
 		if (MaxLinearVelocity != TVec3<FRealSingle>(TNumericLimits<FRealSingle>::Max()))
 		{
 			for (int32 Index = 0; Index < 3; ++Index)
@@ -1226,12 +1222,26 @@ void FClothingSimulationSolver::SetReferenceVelocityScale(uint32 GroupId,
 				}
 			}
 		}
-
+		// Only update old reference velocity when ticking.
+		switch (VelocityScaleSpace)
+		{
+		case EChaosSoftsSimulationSpace::WorldSpace:
+			InOutAppliedReferenceAngularVelocity = AppliedAngularVelocity;
+			break;
+		case EChaosSoftsSimulationSpace::ComponentSpace:
+			InOutAppliedReferenceAngularVelocity = LocalSpaceRotation.RotateVector(AppliedAngularVelocity);
+			break;
+		case EChaosSoftsSimulationSpace::ReferenceBoneSpace:
+		default:
+			InOutAppliedReferenceAngularVelocity = OldReferenceSpaceTransform.GetRotation().RotateVector(AppliedAngularVelocity);
+			break;
+		}
 		AppliedDeltaTranslation = AppliedLinearVelocity * DeltaTime;
 
+		// Apply angular acceleration clamping
 		if (MaxAngularAcceleration != TNumericLimits<FRealSingle>::Max() && MaxAngularAcceleration >= 0.f)
 		{
-			const FVec3 OldScaledAngularVelocity = AngularRatio * OldReferenceAngularVelocity;
+			const FVec3 OldScaledAngularVelocity = OldAppliedReferenceAngularVelocity;
 			const FVec3 ScaledReferenceSpaceAngularAccelerationTimesDt = AppliedAngularVelocity - OldScaledAngularVelocity;
 			const FReal LenSq = ScaledReferenceSpaceAngularAccelerationTimesDt.SquaredLength();
 			if (LenSq > FMath::Square(MaxAngularAcceleration * DeltaTime))
@@ -1240,6 +1250,7 @@ void FClothingSimulationSolver::SetReferenceVelocityScale(uint32 GroupId,
 				AppliedAngularVelocity = OldScaledAngularVelocity + Scale * ScaledReferenceSpaceAngularAccelerationTimesDt;
 			}
 		}
+		// Apply angular velocity clamping
 		if (MaxAngularVelocity != TNumericLimits<FRealSingle>::Max() && MaxAngularVelocity >= 0.f)
 		{
 			const FReal LenSq = AppliedAngularVelocity.SquaredLength();
@@ -1249,6 +1260,8 @@ void FClothingSimulationSolver::SetReferenceVelocityScale(uint32 GroupId,
 				AppliedAngularVelocity *= Scale;
 			}
 		}
+		// Only update old reference velocity when ticking.
+		InOutAppliedReferenceAngularVelocity = AppliedAngularVelocity;
 
 		AppliedDeltaRotation = AppliedAngularVelocity * DeltaTime;
 	}
diff --git a/Engine/Plugins/ChaosCloth/Source/ChaosCloth/Public/ChaosCloth/ChaosClothingSimulationCloth.h b/Engine/Plugins/ChaosCloth/Source/ChaosCloth/Public/ChaosCloth/ChaosClothingSimulationCloth.h
--- a/Engine/Plugins/ChaosCloth/Source/ChaosCloth/Public/ChaosCloth/ChaosClothingSimulationCloth.h
+++ b/Engine/Plugins/ChaosCloth/Source/ChaosCloth/Public/ChaosCloth/ChaosClothingSimulationCloth.h
@@ -151,8 +151,8 @@ namespace Chaos
 
 		// Reference space transform
 		FRigidTransform3 ReferenceSpaceTransform;  // TODO: Add override in the style of LODIndexOverride
-		FVec3 ReferenceSpaceVelocity;
-		FVec3 ReferenceSpaceAngularVelocity;
+		FVec3 AppliedReferenceSpaceVelocity; // After scaling/clamping
+		FVec3 AppliedReferenceSpaceAngularVelocity; // After scaling/clamping
 
 		// LOD data
 		TArray<TUniquePtr<FLODData>> LODData;
diff --git a/Engine/Plugins/ChaosCloth/Source/ChaosCloth/Public/ChaosCloth/ChaosClothingSimulationSolver.h b/Engine/Plugins/ChaosCloth/Source/ChaosCloth/Public/ChaosCloth/ChaosClothingSimulationSolver.h
--- a/Engine/Plugins/ChaosCloth/Source/ChaosCloth/Public/ChaosCloth/ChaosClothingSimulationSolver.h
+++ b/Engine/Plugins/ChaosCloth/Source/ChaosCloth/Public/ChaosCloth/ChaosClothingSimulationSolver.h
@@ -138,8 +138,8 @@ namespace Chaos
 		CHAOSCLOTH_API void SetReferenceVelocityScale(uint32 GroupId,
 			const FRigidTransform3& OldReferenceSpaceTransform,
 			const FRigidTransform3& ReferenceSpaceTransform,
-			TVec3<FReal>& InOutReferenceVelocity, // Old reference velocity is passed in. New reference velocity is returned.
-			TVec3<FReal>& InOutReferenceAngularVelocity, // Old reference velocity is passed in. New reference velocity is returned.
+			TVec3<FReal>& InOutAppliedReferenceVelocity, // Old reference velocity is passed in. New reference velocity is returned.
+			TVec3<FReal>& InOutAppliedReferenceAngularVelocity, // Old reference velocity is passed in. New reference velocity is returned.
 			const EChaosSoftsSimulationSpace VelocityScaleSpace, // the space the following linear velocity properties are in.
 			const TVec3<FRealSingle>& LinearVelocityScale,
 			const TVec3<FRealSingle>& MaxLinearVelocity,
