Chaos Cloth - Draw all editor visualizations at the foreground.

[RN] Chaos Cloth - Made all editor visualizations (non PIE) draw at the foreground. Can be disabled with p.ChaosClothVisualization.DrawInForeground 0.

diff --git a/Engine/Plugins/ChaosCloth/Source/ChaosCloth/Private/ChaosCloth/ChaosClothVisualization.cpp b/Engine/Plugins/ChaosCloth/Source/ChaosCloth/Private/ChaosCloth/ChaosClothVisualization.cpp
--- a/Engine/Plugins/ChaosCloth/Source/ChaosCloth/Private/ChaosCloth/ChaosClothVisualization.cpp
+++ b/Engine/Plugins/ChaosCloth/Source/ChaosCloth/Private/ChaosCloth/ChaosClothVisualization.cpp
@@ -116,7 +116,7 @@ namespace Chaos
 		{
 			// set depth to SDPG_MAX not SDPG_Foreground when drawing in foreground.
 			// SDPG_Foreground does not draw when PIE is paused (its buffer is flushed).
-			return bDrawInForeground ? SDPG_MAX : SDPG_World; 
+			return bDrawInForeground ? SDPG_MAX : SDPG_World;
 		}
 	}  // namespace Private
 
@@ -279,7 +279,7 @@ namespace Chaos
 		}
 
 		const FMatrix LocalSimSpaceToWorld = GetLocalSpaceToWorld(*Solver);
-		MeshBuilder.Draw(PDI, LocalSimSpaceToWorld, FMaterials::GetInstance().GetClothMaterial()->GetRenderProxy(), SDPG_World, false, false);
+		MeshBuilder.Draw(PDI, LocalSimSpaceToWorld, FMaterials::GetInstance().GetClothMaterial()->GetRenderProxy(), Chaos::Private::GetDepthPriority(), false, false);
 	}
 
 	void FClothVisualizationNoGC::DrawWeightMapWithName(FPrimitiveDrawInterface* PDI, const FString& Name) const
@@ -339,7 +339,7 @@ namespace Chaos
 		}
 
 		const FMatrix LocalSimSpaceToWorld = GetLocalSpaceToWorld(*Solver);
-		MeshBuilder.Draw(PDI, LocalSimSpaceToWorld, FMaterials::GetInstance().GetClothMaterialColor()->GetRenderProxy(), SDPG_World, false, false);
+		MeshBuilder.Draw(PDI, LocalSimSpaceToWorld, FMaterials::GetInstance().GetClothMaterialColor()->GetRenderProxy(), Chaos::Private::GetDepthPriority(), false, false);
 	}
 
 	void FClothVisualizationNoGC::DrawWeightMap(FPrimitiveDrawInterface* PDI) const
@@ -406,7 +406,7 @@ namespace Chaos
 
 		FMatrix LocalSimSpaceToWorld(FMatrix::Identity);
 		LocalSimSpaceToWorld.SetOrigin(Solver->GetLocalSpaceLocation());
-		MeshBuilder.Draw(PDI, LocalSimSpaceToWorld, FMaterials::GetInstance().GetClothMaterialColor()->GetRenderProxy(), SDPG_World, false, false);
+		MeshBuilder.Draw(PDI, LocalSimSpaceToWorld, FMaterials::GetInstance().GetClothMaterialColor()->GetRenderProxy(), Chaos::Private::GetDepthPriority(), false, false);
 	}
 #endif  // #if WITH_EDITOR
 
@@ -547,11 +547,11 @@ namespace Chaos
 			const FMatrix& ViewMatrix = PDI->View->ViewMatrices.GetViewMatrix();
 			const FVector XAxis = ViewMatrix.GetColumn(0); // Just using transpose here (orthogonal transform assumed)
 			const FVector YAxis = ViewMatrix.GetColumn(1);
-			DrawDisc(PDI, Pos, XAxis, YAxis, Color.ToFColor(true), 0.5f, 10, ClothMaterialVertex->GetRenderProxy(), SDPG_World);
+			DrawDisc(PDI, Pos, XAxis, YAxis, Color.ToFColor(true), 0.5f, 10, ClothMaterialVertex->GetRenderProxy(), Chaos::Private::GetDepthPriority());
 		}
 		else
 		{
-			PDI->DrawPoint(Pos, Color, Thickness, SDPG_World);
+			PDI->DrawPoint(Pos, Color, Thickness, Chaos::Private::GetDepthPriority());
 		}
 #endif
 	}
@@ -564,7 +564,7 @@ namespace Chaos
 			return;
 		}
 #if WITH_EDITOR
-		PDI->DrawLine(Pos0, Pos1, Color, SDPG_World, 0.0f, 0.001f);
+		PDI->DrawLine(Pos0, Pos1, Color, Chaos::Private::GetDepthPriority(), 0.0f, 0.001f);
 #endif
 	}
 
@@ -595,7 +595,7 @@ namespace Chaos
 		}
 #if WITH_EDITOR
 		const FTransform Transform(Rotation, Center, FVector(Scale));
-		DrawWireSphere(PDI, Transform, Color, Sphere.GetRadiusf(), 12, SDPG_World, 0.0f, 0.001f, false);
+		DrawWireSphere(PDI, Transform, Color, Sphere.GetRadiusf(), 12, Chaos::Private::GetDepthPriority(), 0.0f, 0.001f, false);
 #endif
 	}
 
@@ -609,7 +609,7 @@ namespace Chaos
 		}
 #if WITH_EDITOR
 		const FMatrix BoxToWorld = FTransform(Rotation, Position, FVector(Scale)).ToMatrixWithScale();
-		DrawWireBox(PDI, BoxToWorld, FBox(Box.Min(), Box.Max()), Color, SDPG_World, 0.0f, 0.001f, false);
+		DrawWireBox(PDI, BoxToWorld, FBox(Box.Min(), Box.Max()), Color, Chaos::Private::GetDepthPriority(), 0.0f, 0.001f, false);
 #endif
 	}
 
@@ -631,7 +631,7 @@ namespace Chaos
 		const FVector X = Rotation.RotateVector(Forward);
 		const FVector Y = Rotation.RotateVector(Right);
 		const FVector Z = Rotation.RotateVector(Up);
-		DrawWireCapsule(PDI, Center, X, Y, Z, Color, Radius, HalfHeight, 12, SDPG_World, 0.0f, 0.001f, false);
+		DrawWireCapsule(PDI, Center, X, Y, Z, Color, Radius, HalfHeight, 12, Chaos::Private::GetDepthPriority(), 0.0f, 0.001f, false);
 #endif
 	}
 
@@ -769,7 +769,7 @@ namespace Chaos
 				MeshBuilder.AddTriangle(T[0], T[1], T[2]);
 			}
 
-			MeshBuilder.Draw(PDI, Transform.ToMatrixWithScale(), MaterialRenderProxy, SDPG_World, false, false);
+			MeshBuilder.Draw(PDI, Transform.ToMatrixWithScale(), MaterialRenderProxy, Chaos::Private::GetDepthPriority(), false, false);
 		}
 		else
 #endif
@@ -799,7 +799,7 @@ namespace Chaos
 			}
 
 			const FTransform LocalToWorld(Rotation, Position, FVector(Scale));
-			MeshBuilder.Draw(PDI, LocalToWorld.ToMatrixWithScale(), MaterialRenderProxy, SDPG_World, false, false);
+			MeshBuilder.Draw(PDI, LocalToWorld.ToMatrixWithScale(), MaterialRenderProxy, Chaos::Private::GetDepthPriority(), false, false);
 
 			if (Private::DrawSkinnedLattice)
 			{
@@ -829,20 +829,20 @@ namespace Chaos
 								const FVector P110 = LocalToWorld.TransformPosition(FVector(DeformedPoints(I + 1, J + 1, K)));
 								const FVector P111 = LocalToWorld.TransformPosition(FVector(DeformedPoints(I + 1, J + 1, K + 1)));
 
-								PDI->AddReserveLines(SDPG_World, 12);
+								PDI->AddReserveLines(Chaos::Private::GetDepthPriority(), 12);
 								const FColor& Color = bIsEmpty ? EmptyLatticeColor : LatticeColor;
-								PDI->DrawLine(P000, P001, Color, SDPG_World);
-								PDI->DrawLine(P000, P010, Color, SDPG_World);
-								PDI->DrawLine(P000, P100, Color, SDPG_World);
-								PDI->DrawLine(P001, P011, Color, SDPG_World);
-								PDI->DrawLine(P001, P101, Color, SDPG_World);
-								PDI->DrawLine(P010, P011, Color, SDPG_World);
-								PDI->DrawLine(P010, P110, Color, SDPG_World);
-								PDI->DrawLine(P011, P111, Color, SDPG_World);
-								PDI->DrawLine(P100, P101, Color, SDPG_World);
-								PDI->DrawLine(P100, P110, Color, SDPG_World);
-								PDI->DrawLine(P101, P111, Color, SDPG_World);
-								PDI->DrawLine(P110, P111, Color, SDPG_World);
+								PDI->DrawLine(P000, P001, Color, Chaos::Private::GetDepthPriority());
+								PDI->DrawLine(P000, P010, Color, Chaos::Private::GetDepthPriority());
+								PDI->DrawLine(P000, P100, Color, Chaos::Private::GetDepthPriority());
+								PDI->DrawLine(P001, P011, Color, Chaos::Private::GetDepthPriority());
+								PDI->DrawLine(P001, P101, Color, Chaos::Private::GetDepthPriority());
+								PDI->DrawLine(P010, P011, Color, Chaos::Private::GetDepthPriority());
+								PDI->DrawLine(P010, P110, Color, Chaos::Private::GetDepthPriority());
+								PDI->DrawLine(P011, P111, Color, Chaos::Private::GetDepthPriority());
+								PDI->DrawLine(P100, P101, Color, Chaos::Private::GetDepthPriority());
+								PDI->DrawLine(P100, P110, Color, Chaos::Private::GetDepthPriority());
+								PDI->DrawLine(P101, P111, Color, Chaos::Private::GetDepthPriority());
+								PDI->DrawLine(P110, P111, Color, Chaos::Private::GetDepthPriority());
 
 							}
 						}
@@ -879,7 +879,7 @@ namespace Chaos
 					MeshBuilder.AddTriangle(T[0], T[1], T[2]);
 				}
 
-				MeshBuilder.Draw(PDI, Transform.ToMatrixWithScale(), MaterialRenderProxy, SDPG_World, false, false);
+				MeshBuilder.Draw(PDI, Transform.ToMatrixWithScale(), MaterialRenderProxy, Chaos::Private::GetDepthPriority(), false, false);
 			}
 
 			// Draw the Local grid that MLLevelSet lives in.
@@ -916,23 +916,23 @@ namespace Chaos
 				const FVector P110 = Transform.TransformPosition(LocalGridCorners[6]);
 				const FVector P111 = Transform.TransformPosition(LocalGridCorners[7]);
 
-				PDI->AddReserveLines(SDPG_World, 12);
+				PDI->AddReserveLines(Chaos::Private::GetDepthPriority(), 12);
 				const FColor& Color = FColor::Purple;
 				const FColor& ColorZ = FColor::Blue;
 				const FColor& ColorY = FColor::Green;
 				const FColor& ColorX = FColor::Red;
-				PDI->DrawLine(P000, P001, ColorZ, SDPG_World);
-				PDI->DrawLine(P000, P010, ColorY, SDPG_World);
-				PDI->DrawLine(P000, P100, ColorX, SDPG_World);
-				PDI->DrawLine(P001, P011, Color, SDPG_World);
-				PDI->DrawLine(P001, P101, Color, SDPG_World);
-				PDI->DrawLine(P010, P011, Color, SDPG_World);
-				PDI->DrawLine(P010, P110, Color, SDPG_World);
-				PDI->DrawLine(P011, P111, Color, SDPG_World);
-				PDI->DrawLine(P100, P101, Color, SDPG_World);
-				PDI->DrawLine(P100, P110, Color, SDPG_World);
-				PDI->DrawLine(P101, P111, Color, SDPG_World);
-				PDI->DrawLine(P110, P111, Color, SDPG_World);
+				PDI->DrawLine(P000, P001, ColorZ, Chaos::Private::GetDepthPriority());
+				PDI->DrawLine(P000, P010, ColorY, Chaos::Private::GetDepthPriority());
+				PDI->DrawLine(P000, P100, ColorX, Chaos::Private::GetDepthPriority());
+				PDI->DrawLine(P001, P011, Color, Chaos::Private::GetDepthPriority());
+				PDI->DrawLine(P001, P101, Color, Chaos::Private::GetDepthPriority());
+				PDI->DrawLine(P010, P011, Color, Chaos::Private::GetDepthPriority());
+				PDI->DrawLine(P010, P110, Color, Chaos::Private::GetDepthPriority());
+				PDI->DrawLine(P011, P111, Color, Chaos::Private::GetDepthPriority());
+				PDI->DrawLine(P100, P101, Color, Chaos::Private::GetDepthPriority());
+				PDI->DrawLine(P100, P110, Color, Chaos::Private::GetDepthPriority());
+				PDI->DrawLine(P101, P111, Color, Chaos::Private::GetDepthPriority());
+				PDI->DrawLine(P110, P111, Color, Chaos::Private::GetDepthPriority());
 			}
 
 			//Draw the training grid.
@@ -952,14 +952,14 @@ namespace Chaos
 				const FVector P010 = Transform.TransformPosition((FVector)CornersXYZ[1]);
 				const FVector P100 = Transform.TransformPosition((FVector)CornersXYZ[0]);
 
-				PDI->AddReserveLines(SDPG_World, 3);
+				PDI->AddReserveLines(Chaos::Private::GetDepthPriority(), 3);
 				const FColor& Color = FColor::Purple;
 				const FColor& ColorZ = FColor::Blue;
 				const FColor& ColorY = FColor::Green;
 				const FColor& ColorX = FColor::Red;
-				PDI->DrawLine(P000, P001, ColorZ, SDPG_World);
-				PDI->DrawLine(P000, P010, ColorY, SDPG_World);
-				PDI->DrawLine(P000, P100, ColorX, SDPG_World);
+				PDI->DrawLine(P000, P001, ColorZ, Chaos::Private::GetDepthPriority());
+				PDI->DrawLine(P000, P010, ColorY, Chaos::Private::GetDepthPriority());
+				PDI->DrawLine(P000, P100, ColorX, Chaos::Private::GetDepthPriority());
 			}
 		}
 		else
@@ -988,7 +988,7 @@ namespace Chaos
 			}
 
 			const FMatrix LocalSimSpaceToWorld = GetLocalSpaceToWorld(*Solver);
-			MeshBuilder.Draw(PDI, LocalSimSpaceToWorld, MaterialRenderProxy, SDPG_World, false, false);
+			MeshBuilder.Draw(PDI, LocalSimSpaceToWorld, MaterialRenderProxy, Chaos::Private::GetDepthPriority(), false, false);
 		}
 		else
 #endif
@@ -1971,7 +1971,7 @@ namespace Chaos
 			}
 
 			const FMatrix LocalSimSpaceToWorld = GetLocalSpaceToWorld(Solver);
-			MeshBuilder.Draw(PDI, LocalSimSpaceToWorld, MaterialRenderProxy, SDPG_World, false, false);
+			MeshBuilder.Draw(PDI, LocalSimSpaceToWorld, MaterialRenderProxy, Chaos::Private::GetDepthPriority(), false, false);
 		}
 #endif
 	}
@@ -2058,7 +2058,7 @@ namespace Chaos
 			}
 
 			const FMatrix LocalSimSpaceToWorld = GetLocalSpaceToWorld(Solver);
-			MeshBuilder.Draw(PDI, LocalSimSpaceToWorld, MaterialRenderProxy, SDPG_World, false, false);
+			MeshBuilder.Draw(PDI, LocalSimSpaceToWorld, MaterialRenderProxy, Chaos::Private::GetDepthPriority(), false, false);
 		}
 #endif
 	}
@@ -2109,7 +2109,7 @@ namespace Chaos
 			}
 
 			const FMatrix LocalSimSpaceToWorld = GetLocalSpaceToWorld(Solver);
-			MeshBuilder.Draw(PDI, LocalSimSpaceToWorld, MaterialRenderProxy, SDPG_World, false, false);
+			MeshBuilder.Draw(PDI, LocalSimSpaceToWorld, MaterialRenderProxy, Chaos::Private::GetDepthPriority(), false, false);
 		}
 #endif
 	}
@@ -2843,7 +2843,7 @@ namespace Chaos
 				}
 			}
 			const FMatrix LocalSimSpaceToWorld = GetLocalSpaceToWorld(*Solver);
-			MeshBuilder.Draw(PDI, LocalSimSpaceToWorld, FMaterials::GetInstance().GetClothMaterialColor()->GetRenderProxy(), SDPG_World, false, false);
+			MeshBuilder.Draw(PDI, LocalSimSpaceToWorld, FMaterials::GetInstance().GetClothMaterialColor()->GetRenderProxy(), Chaos::Private::GetDepthPriority(), false, false);
 		}
 		else
 #endif
@@ -2893,7 +2893,7 @@ namespace Chaos
 						const FVector Position1(Positions[VertexIndex]);
 						const FReal Radius1 = (FReal)SelfCollisionConstraints->GetParticleThickness(VertexIndex + Offset);
 						const FTransform Transform(FRotator::ZeroRotator, GetWorldPosition(*Solver, Position1), FVector(Solver->GetLocalSpaceScale()));
-						DrawWireSphere(PDI, Transform, FLinearColor::Gray, Radius1, 6, SDPG_World, 0.0f, 0.001f, false);
+						DrawWireSphere(PDI, Transform, FLinearColor::Gray, Radius1, 6, Chaos::Private::GetDepthPriority(), 0.0f, 0.001f, false);
 					}
 				}
 			}
@@ -2996,7 +2996,7 @@ namespace Chaos
 		}
 
 		const FMatrix LocalSimSpaceToWorld = GetLocalSpaceToWorld(*Solver);
-		MeshBuilder.Draw(PDI, LocalSimSpaceToWorld, FMaterials::GetInstance().GetCollisionMaterial()->GetRenderProxy(), SDPG_World, false, false);
+		MeshBuilder.Draw(PDI, LocalSimSpaceToWorld, FMaterials::GetInstance().GetCollisionMaterial()->GetRenderProxy(), Chaos::Private::GetDepthPriority(), false, false);
 	}
 #endif // #if WITH_EDITOR
 #else  // #if CHAOS_DEBUG_DRAW
