Fix a bug in Linear Acceleration Clamping. We were not updating the previous linear velocity, so acceleration was always computing assuming the previous velocity was zero.

Adding a temporary cvar (p.ChaosCloth.Solver.BrokenAccelerationClampingHack) to go back to the original behavior while we have critical Silver reviews in flight.

diff --git a/Engine/Plugins/ChaosCloth/Source/ChaosCloth/Private/ChaosCloth/ChaosClothingSimulationSolver.cpp b/Engine/Plugins/ChaosCloth/Source/ChaosCloth/Private/ChaosCloth/ChaosClothingSimulationSolver.cpp
--- a/Engine/Plugins/ChaosCloth/Source/ChaosCloth/Private/ChaosCloth/ChaosClothingSimulationSolver.cpp
+++ b/Engine/Plugins/ChaosCloth/Source/ChaosCloth/Private/ChaosCloth/ChaosClothingSimulationSolver.cpp
@@ -81,6 +81,10 @@ FAutoConsoleVariableRef CVarClothSolverDisableTimeDependentNumIterations(TEXT("p
 FAutoConsoleVariableRef CVarClothSolverUseVelocityScale(TEXT("p.ChaosCloth.Solver.UseVelocityScale"), bClothSolverUseVelocityScale, TEXT("Use the velocity scale to compensate for clamping to MaxPhysicsDelta, in order to avoid miscalculating velocities during hitches."));
 FAutoConsoleVariableRef CVarClothSolverMaxVelocity(TEXT("p.ChaosCloth.Solver.MaxVelocity"), ClothSolverMaxVelocity, TEXT("Maximum relative velocity of the cloth particles relatively to their animated positions equivalent. 0 to disable."));
 
+// Temporary hack to return to previous broken acceleration clamping behavior.
+static bool bClothSolverBrokenAccelerationClampingToZero = false;
+FAutoConsoleVariableRef CVarClothSolverBrokenAccelerationClampingHack(TEXT("p.ChaosCloth.Solver.BrokenAccelerationClampingHack"), bClothSolverBrokenAccelerationClampingToZero, TEXT("Return to previous broken behavior where acceleration clamping would clamp to zero velocity. This cvar will be removed."));
+
 namespace ClothingSimulationSolverDefault
 {
 	static const Softs::FSolverVec3 Gravity((Softs::FSolverReal)0., (Softs::FSolverReal)0., (Softs::FSolverReal)-980.665);  // cm/s^2
@@ -1224,20 +1228,26 @@ void FClothingSimulationSolver::SetReferenceVelocityScale(uint32 GroupId,
 				}
 			}
 		}
-		// Only update old reference velocity when ticking.
-		switch (VelocityScaleSpace)
+
+		// Temporary cvar to return to previous broken behavior where previous reference velocity wasn't stored.
+		if (!bClothSolverBrokenAccelerationClampingToZero)
 		{
-		case EChaosSoftsSimulationSpace::WorldSpace:
-			InOutAppliedReferenceAngularVelocity = AppliedAngularVelocity;
-			break;
-		case EChaosSoftsSimulationSpace::ComponentSpace:
-			InOutAppliedReferenceAngularVelocity = LocalSpaceRotation.RotateVector(AppliedAngularVelocity);
-			break;
-		case EChaosSoftsSimulationSpace::ReferenceBoneSpace:
-		default:
-			InOutAppliedReferenceAngularVelocity = OldReferenceSpaceTransform.GetRotation().RotateVector(AppliedAngularVelocity);
-			break;
+			// Only update old reference velocity when ticking.
+			switch (VelocityScaleSpace)
+			{
+			case EChaosSoftsSimulationSpace::WorldSpace:
+				InOutAppliedReferenceVelocity = AppliedLinearVelocity;
+				break;
+			case EChaosSoftsSimulationSpace::ComponentSpace:
+				InOutAppliedReferenceVelocity = LocalSpaceRotation.RotateVector(AppliedLinearVelocity);
+				break;
+			case EChaosSoftsSimulationSpace::ReferenceBoneSpace:
+			default:
+				InOutAppliedReferenceVelocity = OldReferenceSpaceTransform.GetRotation().RotateVector(AppliedLinearVelocity);
+				break;
+			}
 		}
+
 		AppliedDeltaTranslation = AppliedLinearVelocity * DeltaTime;
 
 		// Apply angular acceleration clamping
