Fix bug where you cannot remove the last external collision source from Chaos Cloth.

diff --git a/Engine/Plugins/ChaosCloth/Source/ChaosCloth/Private/ChaosCloth/ChaosClothingSimulationCollider.cpp b/Engine/Plugins/ChaosCloth/Source/ChaosCloth/Private/ChaosCloth/ChaosClothingSimulationCollider.cpp
--- a/Engine/Plugins/ChaosCloth/Source/ChaosCloth/Private/ChaosCloth/ChaosClothingSimulationCollider.cpp
+++ b/Engine/Plugins/ChaosCloth/Source/ChaosCloth/Private/ChaosCloth/ChaosClothingSimulationCollider.cpp
@@ -117,10 +117,6 @@ void FClothingSimulationCollider::FLODData::Add(
 	const int32 NumComplexColliders = NumSkinnedLevelSets + NumMLLevelSets + NumSkinnedTriangleMeshes;
 
 	NumGeometries = NumSimpleColliders + (bSupportsComplexColliders ? NumComplexColliders : 0);
-	if(NumGeometries == 0)
-	{
-		return;
-	}
 
 	// Retrieve cloth group Id, or use INDEX_NONE if this collider applies to all cloths (when Cloth == nullptr)
 	const uint32 GroupId = Cloth ? Cloth->GetGroupId() : INDEX_NONE;
@@ -132,6 +128,12 @@ void FClothingSimulationCollider::FLODData::Add(
 	int32& CollisionRangeId = bIsNewCollider ? CollisionRangeIds.Add(FSolverClothPair(Solver, Cloth)) : *CollisionRangeIdPtr;
 	CollisionRangeId = Solver->AddCollisionParticles(NumGeometries, GroupId, bIsNewCollider ? INDEX_NONE : CollisionRangeId);
 
+	// Must try to add collision particles even when zero, so it will remove any existing collisions
+	if (NumGeometries == 0)
+	{
+		return;
+	}
+
 	TArrayView<int32> BoneIndices = Solver->GetCollisionBoneIndicesView(CollisionRangeId);
 	TArrayView<Softs::FSolverRigidTransform3> BaseTransforms = Solver->GetCollisionBaseTransformsView(CollisionRangeId);
 
diff --git a/Engine/Plugins/ChaosCloth/Source/ChaosCloth/Private/ChaosCloth/ChaosClothingSimulationSolver.cpp b/Engine/Plugins/ChaosCloth/Source/ChaosCloth/Private/ChaosCloth/ChaosClothingSimulationSolver.cpp
--- a/Engine/Plugins/ChaosCloth/Source/ChaosCloth/Private/ChaosCloth/ChaosClothingSimulationSolver.cpp
+++ b/Engine/Plugins/ChaosCloth/Source/ChaosCloth/Private/ChaosCloth/ChaosClothingSimulationSolver.cpp
@@ -778,11 +778,6 @@ TArrayView<Softs::FSolverVec3> FClothingSimulationSolver::GetAnimationVelocities
 
 int32 FClothingSimulationSolver::AddCollisionParticles(int32 NumCollisionParticles, uint32 GroupId, int32 RecycledCollisionRangeId)
 {
-	if (!NumCollisionParticles)
-	{
-		return INDEX_NONE;
-	}
-
 	constexpr bool bActivateFalse = false;
 
 	int32 CollisionRangeId = INDEX_NONE;
@@ -807,10 +802,13 @@ int32 FClothingSimulationSolver::AddCollisionParticles(int32 NumCollisionParticl
 				}
 			}
 		}
-		CollisionRangeId = Evolution->AddCollisionParticleRange(GroupId, NumCollisionParticles, bActivateFalse);
+		if (NumCollisionParticles)
+		{
+			CollisionRangeId = Evolution->AddCollisionParticleRange(GroupId, NumCollisionParticles, bActivateFalse);
 
-		// Add an empty ComplexColliders
-		ComplexColliders.Emplace(CollisionRangeId, MakeUnique<FClothComplexColliders>(Evolution.Get(), CollisionRangeId));
+			// Add an empty ComplexColliders
+			ComplexColliders.Emplace(CollisionRangeId, MakeUnique<FClothComplexColliders>(Evolution.Get(), CollisionRangeId));
+		}
 	}
 	else
 	{
@@ -832,21 +830,27 @@ int32 FClothingSimulationSolver::AddCollisionParticles(int32 NumCollisionParticl
 			PBDEvolution->ResetCollisionParticles(RecycledCollisionRangeId);
 		}
 
-		CollisionRangeId = PBDEvolution->AddCollisionParticleRange(NumCollisionParticles, GroupId, /*bActivate =*/ false);
+		if (NumCollisionParticles)
+		{
+			CollisionRangeId = PBDEvolution->AddCollisionParticleRange(NumCollisionParticles, GroupId, /*bActivate =*/ false);
+		}
 	}
 
-	// Always initialize the collision particle's transforms as otherwise setting the geometry will get NaNs detected during the bounding box updates
-	Softs::FSolverRotation3* const Rs = GetCollisionParticleRs(CollisionRangeId);
-	Softs::FSolverVec3* const Xs = GetCollisionParticleXs(CollisionRangeId);
-
-	for (int32 Index = 0; Index < NumCollisionParticles; ++Index)
+	if (CollisionRangeId != INDEX_NONE)
 	{
-		Xs[Index] = Softs::FSolverVec3(0.);
-		Rs[Index] = Softs::FSolverRotation3::FromIdentity();
-	}
+		// Always initialize the collision particle's transforms as otherwise setting the geometry will get NaNs detected during the bounding box updates
+		Softs::FSolverRotation3* const Rs = GetCollisionParticleRs(CollisionRangeId);
+		Softs::FSolverVec3* const Xs = GetCollisionParticleXs(CollisionRangeId);
 
-	// Always starts with particles disabled
-	EnableCollisionParticles(CollisionRangeId, false);
+		for (int32 Index = 0; Index < NumCollisionParticles; ++Index)
+		{
+			Xs[Index] = Softs::FSolverVec3(0.);
+			Rs[Index] = Softs::FSolverRotation3::FromIdentity();
+		}
+
+		// Always starts with particles disabled
+		EnableCollisionParticles(CollisionRangeId, false);
+	}
 
 	return CollisionRangeId;
 }
