Enable velocity and acceleration clamps on legacy cloth.

diff --git a/Engine/Plugins/ChaosCloth/Source/ChaosCloth/Private/ChaosCloth/ChaosClothingSimulationConfig.cpp b/Engine/Plugins/ChaosCloth/Source/ChaosCloth/Private/ChaosCloth/ChaosClothingSimulationConfig.cpp
--- a/Engine/Plugins/ChaosCloth/Source/ChaosCloth/Private/ChaosCloth/ChaosClothingSimulationConfig.cpp
+++ b/Engine/Plugins/ChaosCloth/Source/ChaosCloth/Private/ChaosCloth/ChaosClothingSimulationConfig.cpp
@@ -151,9 +151,16 @@ namespace Chaos
 
 			// Velocity scale
 			{
+				Property->AddValue(TEXT("VelocityScaleSpace"), (int32)ClothConfig->VelocityScaleSpace, AnimatablePropertyFlags);
 				Property->AddValue(TEXT("LinearVelocityScale"), FVector3f(ClothConfig->LinearVelocityScale), AnimatablePropertyFlags);
 				Property->AddValue(TEXT("AngularVelocityScale"), ClothConfig->AngularVelocityScale, AnimatablePropertyFlags);
 				Property->AddValue(TEXT("FictitiousAngularScale"), ClothConfig->FictitiousAngularScale, AnimatablePropertyFlags);
+
+				static const FVector3f DefaultLinearClamp = { TNumericLimits<float>::Max(), TNumericLimits<float>::Max(), TNumericLimits<float>::Max() };
+				Property->AddValue(TEXT("MaxLinearVelocity"), ClothConfig->bEnableLinearVelocityClamping ? ClothConfig->MaxLinearVelocity : DefaultLinearClamp, AnimatablePropertyFlags);
+				Property->AddValue(TEXT("MaxLinearAcceleration"), ClothConfig->bEnableLinearAccelerationClamping ? ClothConfig->MaxLinearAcceleration : DefaultLinearClamp, AnimatablePropertyFlags);
+				Property->AddValue(TEXT("MaxAngularVelocity"), ClothConfig->bEnableAngularVelocityClamping ? ClothConfig->MaxAngularVelocity : TNumericLimits<float>::Max(), AnimatablePropertyFlags);
+				Property->AddValue(TEXT("MaxAngularAcceleration"), ClothConfig->bEnableAngularAccelerationClamping ? ClothConfig->MaxAngularAcceleration : TNumericLimits<float>::Max(), AnimatablePropertyFlags);
 			}
 
 			// Aerodynamics
diff --git a/Engine/Plugins/ChaosCloth/Source/ChaosCloth/Private/ChaosCloth/ChaosClothingSimulationInteractor.cpp b/Engine/Plugins/ChaosCloth/Source/ChaosCloth/Private/ChaosCloth/ChaosClothingSimulationInteractor.cpp
--- a/Engine/Plugins/ChaosCloth/Source/ChaosCloth/Private/ChaosCloth/ChaosClothingSimulationInteractor.cpp
+++ b/Engine/Plugins/ChaosCloth/Source/ChaosCloth/Private/ChaosCloth/ChaosClothingSimulationInteractor.cpp
@@ -193,6 +193,19 @@ void UChaosClothingInteractor::SetVelocityScale(FVector LinearVelocityScale, flo
 	}));
 }
 
+void UChaosClothingInteractor::SetVelocityClamps(bool bEnableLinearVelocityClamping, FVector MaxLinearVelocity, bool bEnableLinearAccelerationClamping, FVector MaxLinearAcceleration,
+	bool bEnableAngularVelocityClamping, float MaxAngularVelocity, bool bEnableAngularAccelerationClamping, float MaxAngularAcceleration)
+{
+	ConfigCommands.Add(FChaosClothingInteractorConfigCommand::CreateLambda([bEnableLinearVelocityClamping, MaxLinearVelocity, bEnableLinearAccelerationClamping, MaxLinearAcceleration, bEnableAngularVelocityClamping, MaxAngularVelocity, bEnableAngularAccelerationClamping, MaxAngularAcceleration](Chaos::FClothingSimulationConfig* Config, int32 LODIndex)
+		{
+			static const FVector3f DefaultLinearClamp = { TNumericLimits<float>::Max(), TNumericLimits<float>::Max(), TNumericLimits<float>::Max() };
+			Config->GetProperties(LODIndex).SetValue(TEXT("MaxLinearVelocity"), bEnableLinearVelocityClamping ? FVector3f(MaxLinearVelocity) : DefaultLinearClamp);
+			Config->GetProperties(LODIndex).SetValue(TEXT("MaxLinearAcceleration"), bEnableLinearAccelerationClamping ? FVector3f(MaxLinearAcceleration) : DefaultLinearClamp);
+			Config->GetProperties(LODIndex).SetValue(TEXT("MaxAngularVelocity"), bEnableAngularVelocityClamping ? MaxAngularVelocity : TNumericLimits<float>::Max());
+			Config->GetProperties(LODIndex).SetValue(TEXT("MaxAngularAcceleration"), bEnableAngularAccelerationClamping ? MaxAngularAcceleration : TNumericLimits<float>::Max());
+		}));
+}
+
 void UChaosClothingInteractor::ResetAndTeleport(bool bReset, bool bTeleport)
 {
 	if (bReset)
diff --git a/Engine/Plugins/ChaosCloth/Source/ChaosCloth/Public/ChaosCloth/ChaosClothConfig.h b/Engine/Plugins/ChaosCloth/Source/ChaosCloth/Public/ChaosCloth/ChaosClothConfig.h
--- a/Engine/Plugins/ChaosCloth/Source/ChaosCloth/Public/ChaosCloth/ChaosClothConfig.h
+++ b/Engine/Plugins/ChaosCloth/Source/ChaosCloth/Public/ChaosCloth/ChaosClothConfig.h
@@ -1,5 +1,6 @@
 // Copyright Epic Games, Inc. All Rights Reserved.
 #pragma once
+#include "Chaos/SoftsSimulationSpace.h"
 #include "Components/SkeletalMeshComponent.h"
 #include "Engine/SkeletalMesh.h"
 #include "ClothConfig.h"
@@ -397,6 +398,12 @@ public:
 	UPROPERTY(EditAnywhere, Category = "Animation Properties", meta = (UIMin = "0", UIMax = "1", ClampMin = "0", ClampMax = "1"))
 	FChaosClothWeightedValue AnimDriveDamping = { 0.f, 1.f };
 
+	/**
+	 * All vector properties in this section (e.g., Linear Velocity Scale, Max Linear Acceleration)
+	 * will be evaluated in this space.
+	 */
+	UPROPERTY(EditAnywhere, Category = "Animation Properties")
+	EChaosSoftsSimulationSpace VelocityScaleSpace = EChaosSoftsSimulationSpace::ReferenceBoneSpace;
 	/**
 	 * The amount of linear velocities sent to the local cloth space from the reference bone
 	 * (the closest bone to the root on which the cloth section has been skinned, or the root itself if the cloth isn't skinned).
@@ -404,6 +411,32 @@ public:
 	UPROPERTY(EditAnywhere, Category = "Animation Properties", meta = (UIMin = "0", UIMax = "1", ClampMin = "0", ClampMax = "1"))
 	FVector LinearVelocityScale = { 0.75f, 0.75f, 0.75f };
 
+	/**
+	 * Enable linear velocity clamping.
+	 */
+	UPROPERTY(EditAnywhere, Category = "Animation Properties", meta = (InlineEditConditionToggle))
+	bool bEnableLinearVelocityClamping = false;
+
+	/**
+	 * The maximum amount of linear velocity sent to the local cloth space from the reference bone
+	 * (the closest bone to the root on which the cloth section has been skinned, or the root itself if the cloth isn't skinned).
+	 */
+	UPROPERTY(EditAnywhere, Category = "Animation Properties", meta = (ClampMin = "0", EditCondition = "bEnableLinearVelocityClamping"))
+	FVector3f MaxLinearVelocity = { 1000.f, 1000.f, 1000.f }; // Approx 22mph or 36kph per direction
+
+	/**
+	 * Enable linear acceleration clamping.
+	 */
+	UPROPERTY(EditAnywhere, Category = "Animation Properties", meta = (InlineEditConditionToggle))
+	bool bEnableLinearAccelerationClamping = false;
+
+	/**
+	 * The maximum amount of linear acceleration sent to the local cloth space from the reference bone
+	 * (the closest bone to the root on which the cloth section has been skinned, or the root itself if the cloth isn't skinned).
+	 */
+	UPROPERTY(EditAnywhere, Category = "Animation Properties", meta = (ClampMin = "0", EditCondition = "bEnableLinearAccelerationClamping"))
+	FVector3f MaxLinearAcceleration = { 60000.f, 60000.f, 60000.f };
+
 	/**
 	 * The amount of angular velocities sent to the local cloth space from the reference bone
 	 * (the closest bone to the root on which the cloth section has been skinned, or the root itself if the cloth isn't skinned).
@@ -411,6 +444,32 @@ public:
 	UPROPERTY(EditAnywhere, Category = "Animation Properties", meta = (UIMin = "0", UIMax = "1", ClampMin = "0", ClampMax = "1"))
 	float AngularVelocityScale = 0.75f;
 
+	/**
+	 * Enable angular velocity clamping.
+	 */
+	UPROPERTY(EditAnywhere, Category = "Animation Properties", meta = (InlineEditConditionToggle))
+	bool bEnableAngularVelocityClamping = false;
+
+	/**
+	 * The maximum amount of angular velocity sent to the local cloth space from the reference bone
+	 * (the closest bone to the root on which the cloth section has been skinned, or the root itself if the cloth isn't skinned).
+	 */
+	UPROPERTY(EditAnywhere, Category = "Animation Properties", meta = (ClampMin = "0", EditCondition = "bEnableAngularVelocityClamping"))
+	float MaxAngularVelocity = 200.f;
+
+	/**
+	 * Enable angular acceleration clamping.
+	 */
+	UPROPERTY(EditAnywhere, Category = "Animation Properties", meta = (InlineEditConditionToggle))
+	bool bEnableAngularAccelerationClamping = false;
+
+	/**
+	 * The maximum amount of angular acceleration sent to the local cloth space from the reference bone
+	 * (the closest bone to the root on which the cloth section has been skinned, or the root itself if the cloth isn't skinned).
+	 */
+	UPROPERTY(EditAnywhere, Category = "Animation Properties", meta = (ClampMin = "0", EditCondition = "bEnableAngularAccelerationClamping"))
+	float MaxAngularAcceleration = 12000.f;
+
 	/**
 	 * The portion of the angular velocity that is used to calculate the strength of all fictitious forces (e.g. centrifugal force).
 	 * This parameter is only having an effect on the portion of the reference bone's angular velocity that has been removed from the
diff --git a/Engine/Plugins/ChaosCloth/Source/ChaosCloth/Public/ChaosCloth/ChaosClothingSimulationInteractor.h b/Engine/Plugins/ChaosCloth/Source/ChaosCloth/Public/ChaosCloth/ChaosClothingSimulationInteractor.h
--- a/Engine/Plugins/ChaosCloth/Source/ChaosCloth/Public/ChaosCloth/ChaosClothingSimulationInteractor.h
+++ b/Engine/Plugins/ChaosCloth/Source/ChaosCloth/Public/ChaosCloth/ChaosClothingSimulationInteractor.h
@@ -72,6 +72,12 @@ public:
 	UFUNCTION(BlueprintCallable, Category = ClothingSimulation, Meta = (Keywords = "Chaos Linear Angular Velocity Scale"))
 	CHAOSCLOTH_API void SetVelocityScale(FVector LinearVelocityScale = FVector(0.75f, 0.75f, 0.75f), float AngularVelocityScale = 0.75f, float FictitiousAngularScale = 1.f);
 
+	UFUNCTION(BlueprintCallable, Category = ClothingSimulation, Meta = (Keywords = "Chaos Max Linear Angular Velocity Acceleration Clamp"))
+	CHAOSCLOTH_API void SetVelocityClamps(bool bEnableLinearVelocityClamping = false, FVector MaxLinearVelocity = FVector(1000., 1000., 1000.),
+		bool bEnableLinearAccelerationClamping = false, FVector MaxLinearAcceleration = FVector(60000., 60000., 60000.),
+		bool bEnableAngularVelocityClamping = false, float MaxAngularVelocity = 200.f,
+		bool bEnableAngularAccelerationClamping = false, float MaxAngularAcceleration = 12000.f);
+
 	UFUNCTION(BlueprintCallable, Category = ClothingSimulation, Meta = (Keywords = "Chaos"))
 	CHAOSCLOTH_API void ResetAndTeleport(bool bReset = false, bool bTeleport = false);
 
