Fix acceleration clamping when teleporting cloth.

diff --git a/Engine/Plugins/ChaosCloth/Source/ChaosCloth/Private/ChaosCloth/ChaosClothingSimulationCloth.cpp b/Engine/Plugins/ChaosCloth/Source/ChaosCloth/Private/ChaosCloth/ChaosClothingSimulationCloth.cpp
--- a/Engine/Plugins/ChaosCloth/Source/ChaosCloth/Private/ChaosCloth/ChaosClothingSimulationCloth.cpp
+++ b/Engine/Plugins/ChaosCloth/Source/ChaosCloth/Private/ChaosCloth/ChaosClothingSimulationCloth.cpp
@@ -997,6 +997,13 @@ void FClothingSimulationCloth::Update(FClothingSimulationSolver* Solver)
 			FictitiousAngularScale,
 			OutMaxVelocityScale,
 			bDisableFictitiousForces);
+
+		if (bNeedsReset || bNeedsTeleport)
+		{
+			// Record this frame's velocity as zero
+			AppliedReferenceSpaceAngularVelocity = FVec3(0.);
+			AppliedReferenceSpaceVelocity = FVec3(0.);
+		}
 		if (!Solver->IsLegacySolver())
 		{
 			Solver->SetProperties(ParticleRangeId, ConfigProperties, LODData[LODIndex]->WeightMaps);
