Add Chaos Cloth teleport and reset debug draw

diff --git a/Engine/Plugins/ChaosCloth/Source/ChaosCloth/Private/ChaosCloth/ChaosClothVisualization.cpp b/Engine/Plugins/ChaosCloth/Source/ChaosCloth/Private/ChaosCloth/ChaosClothVisualization.cpp
--- a/Engine/Plugins/ChaosCloth/Source/ChaosCloth/Private/ChaosCloth/ChaosClothVisualization.cpp
+++ b/Engine/Plugins/ChaosCloth/Source/ChaosCloth/Private/ChaosCloth/ChaosClothVisualization.cpp
@@ -3067,6 +3067,57 @@ namespace Chaos
 			}
 		}
 	}
+	void FClothVisualizationNoGC::DrawTeleportReset(FPrimitiveDrawInterface* PDI) const
+	{
+		if (!Solver)
+		{
+			return;
+		}
+		for (const FClothingSimulationCloth* const Cloth : Solver->GetCloths())
+		{
+			const int32 ParticleRangeId = Cloth->GetParticleRangeId(Solver);
+			if (ParticleRangeId == INDEX_NONE)
+			{
+				continue;
+			}
+
+			FLinearColor Color;
+			constexpr FRealSingle ColorHysteresis = .25f;
+			if (Cloth->GetTimeSinceLastReset() <= Cloth->GetTimeSinceLastTeleport())
+			{
+				// Blend from Red to White as time approaches ColorHysteresis
+				Color.R = 1.f;
+				Color.A = 1.f;
+				Color.G = Color.B = Cloth->GetTimeSinceLastReset() == 0.f ? 0.f : FMath::Clamp(.5f * (Cloth->GetTimeSinceLastReset() / ColorHysteresis + 1.f), 0.f, 1.f);
+			}
+			else
+			{
+				// Blend from Green to white as time approaches ColorHysteresis
+				Color.G = 1.f;
+				Color.A = 1.f;
+				Color.R = Color.B = Cloth->GetTimeSinceLastTeleport() == 0.f ? 0.f : FMath::Clamp(.5f * (Cloth->GetTimeSinceLastReset() / ColorHysteresis + 1.f), 0.f, 1.f);
+			}
+
+			// Elements are local indexed for new solver
+			const int32 Offset = Solver->IsLegacySolver() ? ParticleRangeId : 0;
+
+			const TConstArrayView<TVec3<int32>> Elements = Cloth->GetTriangleMesh(Solver).GetElements();
+			const TConstArrayView<Softs::FSolverVec3> Positions = Cloth->GetParticlePositions(Solver);
+
+			for (int32 ElementIndex = 0; ElementIndex < Elements.Num(); ++ElementIndex)
+			{
+				const TVec3<int32>& Element = Elements[ElementIndex];
+
+				const FVector Pos0 = GetWorldPosition(*Solver, Positions[Element.X - Offset]); // TODO: Triangle Mesh shouldn't really be solver dependent (ie not use an offset)
+				const FVector Pos1 = GetWorldPosition(*Solver, Positions[Element.Y - Offset]);
+				const FVector Pos2 = GetWorldPosition(*Solver, Positions[Element.Z - Offset]);
+
+				DrawLine(PDI, Pos0, Pos1, Color);
+				DrawLine(PDI, Pos1, Pos2, Color);
+				DrawLine(PDI, Pos2, Pos0, Color);
+			}
+		}
+	}
 
 #if WITH_EDITOR
 	TArray<FString> FClothVisualizationNoGC::GetAllWeightMapNames() const
diff --git a/Engine/Plugins/ChaosCloth/Source/ChaosCloth/Private/ChaosCloth/ChaosClothingSimulation.cpp b/Engine/Plugins/ChaosCloth/Source/ChaosCloth/Private/ChaosCloth/ChaosClothingSimulation.cpp
--- a/Engine/Plugins/ChaosCloth/Source/ChaosCloth/Private/ChaosCloth/ChaosClothingSimulation.cpp
+++ b/Engine/Plugins/ChaosCloth/Source/ChaosCloth/Private/ChaosCloth/ChaosClothingSimulation.cpp
@@ -96,6 +96,7 @@ namespace ClothingSimulationCVar
 	TAutoConsoleVariable<bool> DebugDrawParticleIndices     (TEXT("p.ChaosCloth.DebugDrawParticleIndices"     ), false, TEXT("Whether to debug draw the Chaos Cloth particle indices"), ECVF_Cheat);
 	TAutoConsoleVariable<bool> DebugDrawElementIndices      (TEXT("p.ChaosCloth.DebugDrawElementIndices"      ), false, TEXT("Whether to debug draw the Chaos Cloth element indices"), ECVF_Cheat);
 	TAutoConsoleVariable<bool> DebugDrawClothClothConstraints    (TEXT("p.ChaosCloth.DebugDrawClothClothConstraints"    ), false, TEXT("Whether to debug draw the Chaos Cloth-Cloth constraints"), ECVF_Cheat);
+	TAutoConsoleVariable<bool> DebugDrawTeleportReset(TEXT("p.ChaosCloth.DebugDrawTeleportReset"), false, TEXT("Whether to debug draw the Chaos Cloth teleport/reset status"), ECVF_Cheat);
 }
 #endif  // #if CHAOS_DEBUG_DRAW
 
@@ -603,6 +604,7 @@ void FClothingSimulation::Simulate(IClothingSimulationContext* InContext)
 	if (ClothingSimulationCVar::DebugDrawWindForces          .GetValueOnAnyThread()) { DebugDrawWindAndPressureForces(); }
 	if (ClothingSimulationCVar::DebugDrawSelfCollision       .GetValueOnAnyThread()) { DebugDrawSelfCollision       (); }
 	if (ClothingSimulationCVar::DebugDrawSelfIntersection    .GetValueOnAnyThread()) { DebugDrawSelfIntersection    (); }
+	if (ClothingSimulationCVar::DebugDrawTeleportReset       .GetValueOnAnyThread()) { DebugDrawTeleportReset       (); }
 #endif  // #if CHAOS_DEBUG_DRAW
 }
 
diff --git a/Engine/Plugins/ChaosCloth/Source/ChaosCloth/Private/ChaosCloth/ChaosClothingSimulationCloth.cpp b/Engine/Plugins/ChaosCloth/Source/ChaosCloth/Private/ChaosCloth/ChaosClothingSimulationCloth.cpp
--- a/Engine/Plugins/ChaosCloth/Source/ChaosCloth/Private/ChaosCloth/ChaosClothingSimulationCloth.cpp
+++ b/Engine/Plugins/ChaosCloth/Source/ChaosCloth/Private/ChaosCloth/ChaosClothingSimulationCloth.cpp
@@ -1099,6 +1099,24 @@ void FClothingSimulationCloth::Update(FClothingSimulationSolver* Solver)
 	}
 
 	// Reset trigger flags
+#if CHAOS_DEBUG_DRAW
+	if (bNeedsReset)
+	{
+		TimeSinceLastReset = 0.f;
+		TimeSinceLastTeleport = 0.f;
+	}
+	else if (bNeedsTeleport)
+	{
+		TimeSinceLastReset += Solver->GetDeltaTime();
+		TimeSinceLastTeleport = 0.f;
+	}
+	else
+	{
+		TimeSinceLastReset += Solver->GetDeltaTime();
+		TimeSinceLastTeleport += Solver->GetDeltaTime();
+	}
+
+#endif
 	bNeedsTeleport = false;
 	bNeedsReset = false;
 }
diff --git a/Engine/Plugins/ChaosCloth/Source/ChaosCloth/Public/ChaosCloth/ChaosClothVisualization.h b/Engine/Plugins/ChaosCloth/Source/ChaosCloth/Public/ChaosCloth/ChaosClothVisualization.h
--- a/Engine/Plugins/ChaosCloth/Source/ChaosCloth/Public/ChaosCloth/ChaosClothVisualization.h
+++ b/Engine/Plugins/ChaosCloth/Source/ChaosCloth/Public/ChaosCloth/ChaosClothVisualization.h
@@ -59,6 +59,7 @@ namespace Chaos
 		CHAOSCLOTH_API void DrawFictitiousAngularForces(FPrimitiveDrawInterface* PDI = nullptr) const;
 		CHAOSCLOTH_API void DrawMultiResConstraint(FPrimitiveDrawInterface* PDI = nullptr) const;
 		CHAOSCLOTH_API void DrawClothClothConstraints(FPrimitiveDrawInterface* PDI = nullptr) const;
+		CHAOSCLOTH_API void DrawTeleportReset(FPrimitiveDrawInterface* PDI = nullptr) const;
 
 #else  // #if CHAOS_DEBUG_DRAW
 		void SetSolver(const ::Chaos::FClothingSimulationSolver* /*InSolver*/) {}
@@ -91,6 +92,7 @@ namespace Chaos
 		void DrawFictitiousAngularForces(FPrimitiveDrawInterface* /*PDI*/ = nullptr) const {}
 		void DrawMultiResConstraint(FPrimitiveDrawInterface* /*PDI*/ = nullptr) const {}
 		void DrawClothClothConstraints(FPrimitiveDrawInterface* /*PDI*/ = nullptr) const {}
+		CHAOSCLOTH_API void DrawTeleportReset(FPrimitiveDrawInterface* /*PDI*/ = nullptr) const {}
 
 #endif  // #if CHAOS_DEBUG_DRAW
 
diff --git a/Engine/Plugins/ChaosCloth/Source/ChaosCloth/Public/ChaosCloth/ChaosClothingSimulation.h b/Engine/Plugins/ChaosCloth/Source/ChaosCloth/Public/ChaosCloth/ChaosClothingSimulation.h
--- a/Engine/Plugins/ChaosCloth/Source/ChaosCloth/Public/ChaosCloth/ChaosClothingSimulation.h
+++ b/Engine/Plugins/ChaosCloth/Source/ChaosCloth/Public/ChaosCloth/ChaosClothingSimulation.h
@@ -114,6 +114,7 @@ namespace Chaos
 		void DebugDrawSelfIntersection(FPrimitiveDrawInterface* PDI = nullptr) const { Visualization.DrawSelfIntersection(PDI); }
 		void DebugDrawBounds(FPrimitiveDrawInterface* PDI = nullptr) const { Visualization.DrawBounds(PDI); }
 		void DebugDrawGravity(FPrimitiveDrawInterface* PDI = nullptr) const { Visualization.DrawGravity(PDI); }
+		void DebugDrawTeleportReset(FPrimitiveDrawInterface* PDI = nullptr) const { Visualization.DrawTeleportReset(PDI); }
 #endif  // #if CHAOS_DEBUG_DRAW
 
 	private:
diff --git a/Engine/Plugins/ChaosCloth/Source/ChaosCloth/Public/ChaosCloth/ChaosClothingSimulationCloth.h b/Engine/Plugins/ChaosCloth/Source/ChaosCloth/Public/ChaosCloth/ChaosClothingSimulationCloth.h
--- a/Engine/Plugins/ChaosCloth/Source/ChaosCloth/Public/ChaosCloth/ChaosClothingSimulationCloth.h
+++ b/Engine/Plugins/ChaosCloth/Source/ChaosCloth/Public/ChaosCloth/ChaosClothingSimulationCloth.h
@@ -114,6 +114,10 @@ namespace Chaos
 		CHAOSCLOTH_API int32 GetReferenceBoneIndex() const;
 		// Return the local reference space transform for this cloth.
 		const FRigidTransform3& GetReferenceSpaceTransform() const { return ReferenceSpaceTransform;  }
+#if CHAOS_DEBUG_DRAW
+		FRealSingle GetTimeSinceLastTeleport() const { return TimeSinceLastTeleport; }
+		FRealSingle GetTimeSinceLastReset() const { return TimeSinceLastReset; }
+#endif
 		// ---- End of the debugging/visualization functions
 
 		// ---- Solver interface ----
@@ -149,6 +153,11 @@ namespace Chaos
 		bool bNeedsReset = false;
 		bool bNeedsTeleport = false;
 
+#if CHAOS_DEBUG_DRAW
+		FRealSingle TimeSinceLastTeleport = 0.f;
+		FRealSingle TimeSinceLastReset = 0.f;
+#endif
+
 		// Reference space transform
 		FRigidTransform3 ReferenceSpaceTransform;  // TODO: Add override in the style of LODIndexOverride
 		FVec3 AppliedReferenceSpaceVelocity; // After scaling/clamping
diff --git a/Engine/Plugins/ChaosClothAsset/Source/ChaosClothAssetEngine/Private/ChaosClothAsset/ClothSimulationProxy.cpp b/Engine/Plugins/ChaosClothAsset/Source/ChaosClothAssetEngine/Private/ChaosClothAsset/ClothSimulationProxy.cpp
--- a/Engine/Plugins/ChaosClothAsset/Source/ChaosClothAssetEngine/Private/ChaosClothAsset/ClothSimulationProxy.cpp
+++ b/Engine/Plugins/ChaosClothAsset/Source/ChaosClothAssetEngine/Private/ChaosClothAsset/ClothSimulationProxy.cpp
@@ -555,6 +555,7 @@ PRAGMA_ENABLE_DEPRECATION_WARNINGS
 		static const TConsoleVariableData<bool>* const DebugDrawParticleIndicesCVar = IConsoleManager::Get().FindTConsoleVariableDataBool(TEXT("p.ChaosCloth.DebugDrawParticleIndices"));
 		static const TConsoleVariableData<bool>* const DebugDrawElementIndicesCVar = IConsoleManager::Get().FindTConsoleVariableDataBool(TEXT("p.ChaosCloth.DebugDrawElementIndices"));
 		static const TConsoleVariableData<bool>* const DebugDrawClothClothConstraintsCVar = IConsoleManager::Get().FindTConsoleVariableDataBool(TEXT("p.ChaosCloth.DebugDrawClothClothConstraints"));
+		static const TConsoleVariableData<bool>* const DebugDrawTeleportResetCVar = IConsoleManager::Get().FindTConsoleVariableDataBool(TEXT("p.ChaosCloth.DebugDrawTeleportReset"));
 
 		if (DebugDrawLocalSpaceCVar && DebugDrawLocalSpaceCVar->GetValueOnAnyThread()) { Visualization->DrawLocalSpace(); }
 		if (DebugDrawBoundsCVar && DebugDrawBoundsCVar->GetValueOnAnyThread()) { Visualization->DrawBounds(); }
@@ -582,6 +583,10 @@ PRAGMA_ENABLE_DEPRECATION_WARNINGS
 		{
 			Visualization->DrawClothClothConstraints();
 		}
+		if (DebugDrawTeleportResetCVar  && DebugDrawTeleportResetCVar->GetValueOnAnyThread())
+		{
+			Visualization->DrawTeleportReset();
+		}
 #endif  // #if CHAOS_DEBUG_DRAW
 	}
 
diff --git a/Engine/Plugins/ChaosClothAssetEditor/Source/ChaosClothAssetEditor/Private/ChaosClothAsset/ClothEditorSimulationVisualization.cpp b/Engine/Plugins/ChaosClothAssetEditor/Source/ChaosClothAssetEditor/Private/ChaosClothAsset/ClothEditorSimulationVisualization.cpp
--- a/Engine/Plugins/ChaosClothAssetEditor/Source/ChaosClothAssetEditor/Private/ChaosClothAsset/ClothEditorSimulationVisualization.cpp
+++ b/Engine/Plugins/ChaosClothAssetEditor/Source/ChaosClothAssetEditor/Private/ChaosClothAsset/ClothEditorSimulationVisualization.cpp
@@ -398,6 +398,13 @@ const FVisualizationOption FVisualizationOption::OptionData[] =
 		}),
 		LOCTEXT("ChaosVisName_DrawClothClothConstraints", "Cloth-Cloth constraints"), 
 		LOCTEXT("ChaosVisName_DrawClothClothConstraints_Tooltip", "Draw cloth-cloth constraints.")),
+	FVisualizationOption(
+		FClothVisualizationDebugDraw::CreateLambda([](const FClothEditorSimulationVisualization&, const ::Chaos::FClothVisualizationNoGC& Visualization, FPrimitiveDrawInterface* PDI)
+		{
+			Visualization.DrawTeleportReset(PDI);
+		}),
+		LOCTEXT("ChaosVisName_DrawTeleportReset", "Teleport/Reset"), 
+		LOCTEXT("ChaosVisName_DrawTeleportReset_Tooltip", "Draw teleport/reset status.")),
 	FVisualizationOption(
 		FClothVisualizationDebugDraw::CreateLambda([](const FClothEditorSimulationVisualization&, const ::Chaos::FClothVisualizationNoGC& Visualization, FPrimitiveDrawInterface* PDI)
 		{
