Chaos Cloth: Draw single backstop and full sphere based on a cvar p.ChaosClothVisualization.DrawSingleBackstopVertex

diff --git a/Engine/Plugins/ChaosCloth/Source/ChaosCloth/Private/ChaosCloth/ChaosClothVisualization.cpp b/Engine/Plugins/ChaosCloth/Source/ChaosCloth/Private/ChaosCloth/ChaosClothVisualization.cpp
--- a/Engine/Plugins/ChaosCloth/Source/ChaosCloth/Private/ChaosCloth/ChaosClothVisualization.cpp
+++ b/Engine/Plugins/ChaosCloth/Source/ChaosCloth/Private/ChaosCloth/ChaosClothVisualization.cpp
@@ -101,6 +101,9 @@ namespace Chaos
 		static bool bDrawInForeground = true;
 		static FAutoConsoleVariableRef CVarClothVizDrawInForeground(TEXT("p.ChaosClothVisualization.DrawInForeground"), bDrawInForeground, TEXT("Draw in foreground when outside the cloth/SKM editor"));
 
+		static int32 BackstopVertexIndex = -1;
+		static FAutoConsoleVariableRef CVarClothVizDrawSingleBackstopVertex(TEXT("p.ChaosClothVisualization.DrawSingleBackstopVertex"), BackstopVertexIndex, TEXT("Draw backstop for a single vertex with a full sphere. Set to -1 to display all vertices with partial spheres."));
+
 		// copied from ClothEditorMode
 		FLinearColor PseudoRandomColor(int32 NumColorRotations)
 		{
@@ -586,7 +589,7 @@ namespace Chaos
 		}
 	}
 
-	static void DrawSphere(FPrimitiveDrawInterface* PDI, const ::Chaos::TSphere<::Chaos::FReal, 3>& Sphere, const FQuat& Rotation, const FVector& Position, const FReal Scale, const FLinearColor& Color)
+	static void DrawSphere(FPrimitiveDrawInterface* PDI, const ::Chaos::TSphere<::Chaos::FReal, 3>& Sphere, const FQuat& Rotation, const FVector& Position, const FReal Scale, const FLinearColor& Color, const FMaterialRenderProxy* SolidMaterial = nullptr)
 	{
 		const FVec3 Center = Position + Rotation.RotateVector(Scale * FVec3(Sphere.GetCenterf()));
 		if (!PDI)
@@ -596,8 +599,15 @@ namespace Chaos
 			return;
 		}
 #if WITH_EDITOR
-		const FTransform Transform(Rotation, Center, FVector(Scale));
-		DrawWireSphere(PDI, Transform, Color, Sphere.GetRadiusf(), 12, Chaos::Private::GetDepthPriority(), 0.0f, 0.001f, false);
+		if (SolidMaterial)
+		{
+			::DrawSphere(PDI, Center, FRotator(Rotation), FVector(Sphere.GetRadiusf()), 32, 16, SolidMaterial, Chaos::Private::GetDepthPriority());
+		}
+		else
+		{
+			const FTransform Transform(Rotation, Center, FVector(Scale));
+			DrawWireSphere(PDI, Transform, Color, Sphere.GetRadiusf(), 12, Chaos::Private::GetDepthPriority(), 0.0f, 0.001f, false);
+		}
 #endif
 	}
 
@@ -1716,30 +1726,64 @@ namespace Chaos
 				const TConstArrayView<Softs::FSolverVec3> AnimationNormals = Cloth->GetAnimationNormals(Solver);
 				const TConstArrayView<Softs::FSolverVec3> ParticlePositions = Cloth->GetParticlePositions(Solver);
 
-				for (int32 Index = 0; Index < AnimationPositions.Num(); ++Index)
+				if (AnimationPositions.IsValidIndex(Private::BackstopVertexIndex))
 				{
-					ColorSeed += 157;  // Prime number that gives a good spread of colors without getting too similar as a rand might do.
 					const FLinearColor ColorLight = FLinearColor::MakeFromHSV8(ColorSeed, 160, 128);
 					const FLinearColor ColorDark = FLinearColor::MakeFromHSV8(ColorSeed, 160, 64);
 
-					const FReal BackstopRadius = BackstopConstraint->GetBackstopRadius(Index) * BackstopConstraint->GetScale() * Solver->GetLocalSpaceScale();
-					const FReal BackstopDistance = BackstopConstraint->GetBackstopDistance(Index) * BackstopConstraint->GetScale() * Solver->GetLocalSpaceScale();
+					const FReal BackstopRadius = BackstopConstraint->GetBackstopRadius(Private::BackstopVertexIndex) * BackstopConstraint->GetScale() * Solver->GetLocalSpaceScale();
+					const FReal BackstopDistance = BackstopConstraint->GetBackstopDistance(Private::BackstopVertexIndex) * BackstopConstraint->GetScale() * Solver->GetLocalSpaceScale();
 
-					const FVector AnimationNormal(AnimationNormals[Index]);
+					const FVector AnimationNormal(AnimationNormals[Private::BackstopVertexIndex]);
 
 					// Draw a line to show the current distance to the sphere
-					const FVector Pos0 = GetWorldPosition(*Solver, AnimationPositions[Index]);
+					const FVector Pos0 = GetWorldPosition(*Solver, AnimationPositions[Private::BackstopVertexIndex]);
 					const FVector Pos1 = Pos0 - (bUseLegacyBackstop ? BackstopDistance - BackstopRadius : BackstopDistance) * AnimationNormal;
-					const FVector Pos2 = GetWorldPosition(*Solver, ParticlePositions[Index]);
+					const FVector Pos2 = GetWorldPosition(*Solver, ParticlePositions[Private::BackstopVertexIndex]);
+					DrawPoint(PDI, Pos2, ColorLight, nullptr, 5.f);
 					DrawLine(PDI, Pos1, Pos2, ColorLight);
 
 					// Draw the sphere
 					if (BackstopRadius > 0.f)
 					{
 						const FVector Center = Pos0 - (bUseLegacyBackstop ? BackstopDistance : BackstopRadius + BackstopDistance) * AnimationNormal;
-						DrawBackstop(Center, AnimationNormal, BackstopRadius, FVector::ForwardVector, ColorDark);
-						DrawBackstop(Center, AnimationNormal, BackstopRadius, FVector::UpVector, ColorDark);
-						DrawBackstop(Center, AnimationNormal, BackstopRadius, FVector::RightVector, ColorDark);
+						DrawLine(PDI, Pos1, Center, ColorDark);
+						const FMaterialRenderProxy* MaterialRenderProxy =
+#if WITH_EDITOR
+							FMaterials::GetInstance().GetCollisionMaterial() ? FMaterials::GetInstance().GetCollisionMaterial()->GetRenderProxy() : nullptr;
+#else
+							nullptr;
+#endif
+						DrawSphere(PDI, FSphere(FVector::ZeroVector, BackstopRadius), FQuat::Identity, Center, 1.f, ColorDark, MaterialRenderProxy);
+					}
+				}
+				else
+				{
+					for (int32 Index = 0; Index < AnimationPositions.Num(); ++Index)
+					{
+						ColorSeed += 157;  // Prime number that gives a good spread of colors without getting too similar as a rand might do.
+						const FLinearColor ColorLight = FLinearColor::MakeFromHSV8(ColorSeed, 160, 128);
+						const FLinearColor ColorDark = FLinearColor::MakeFromHSV8(ColorSeed, 160, 64);
+
+						const FReal BackstopRadius = BackstopConstraint->GetBackstopRadius(Index) * BackstopConstraint->GetScale() * Solver->GetLocalSpaceScale();
+						const FReal BackstopDistance = BackstopConstraint->GetBackstopDistance(Index) * BackstopConstraint->GetScale() * Solver->GetLocalSpaceScale();
+
+						const FVector AnimationNormal(AnimationNormals[Index]);
+
+						// Draw a line to show the current distance to the sphere
+						const FVector Pos0 = GetWorldPosition(*Solver, AnimationPositions[Index]);
+						const FVector Pos1 = Pos0 - (bUseLegacyBackstop ? BackstopDistance - BackstopRadius : BackstopDistance) * AnimationNormal;
+						const FVector Pos2 = GetWorldPosition(*Solver, ParticlePositions[Index]);
+						DrawLine(PDI, Pos1, Pos2, ColorLight);
+
+						// Draw the sphere
+						if (BackstopRadius > 0.f)
+						{
+							const FVector Center = Pos0 - (bUseLegacyBackstop ? BackstopDistance : BackstopRadius + BackstopDistance) * AnimationNormal;
+							DrawBackstop(Center, AnimationNormal, BackstopRadius, FVector::ForwardVector, ColorDark);
+							DrawBackstop(Center, AnimationNormal, BackstopRadius, FVector::UpVector, ColorDark);
+							DrawBackstop(Center, AnimationNormal, BackstopRadius, FVector::RightVector, ColorDark);
+						}
 					}
 				}
 			}
