Skeletal Mesh Cloth - Fixed crash when the simulation is disabled while a visualizer is set to debug draw by the Skeletal Mesh Component.

There are many places where to prevent the crash, but I think this is the best location.
Simply because there is an expectation that the SKM editor uses the extender and also sets the bDisableClothSimulation flag to stop the sim while painting the masks.
And on the other hand the debug draw crash is technically a Chaos Cloth specific issue that wouldn't necessary happen within another simulation extender implementation.

[RN] Skeletal Mesh Cloth - Fixed crash when the simulation is disabled while a visualizer is set to debug draw by the Skeletal Mesh Component.

diff --git a/Engine/Plugins/ChaosCloth/Source/ChaosClothEditor/Private/ChaosClothEditor/ChaosSimulationEditorExtender.cpp b/Engine/Plugins/ChaosCloth/Source/ChaosClothEditor/Private/ChaosClothEditor/ChaosSimulationEditorExtender.cpp
--- a/Engine/Plugins/ChaosCloth/Source/ChaosClothEditor/Private/ChaosClothEditor/ChaosSimulationEditorExtender.cpp
+++ b/Engine/Plugins/ChaosCloth/Source/ChaosClothEditor/Private/ChaosClothEditor/ChaosSimulationEditorExtender.cpp
@@ -139,9 +139,12 @@ void FSimulationEditorExtender::ExtendViewportShowMenu(FMenuBuilder& MenuBuilder
 	MenuBuilder.EndSection();
 }
 
-void FSimulationEditorExtender::DebugDrawSimulation(const IClothingSimulation* Simulation, USkeletalMeshComponent* /*OwnerComponent*/, FPrimitiveDrawInterface* PDI)
+void FSimulationEditorExtender::DebugDrawSimulation(const IClothingSimulation* Simulation, USkeletalMeshComponent* OwnerComponent, FPrimitiveDrawInterface* PDI)
 {
-	if (!ensure(Simulation)) { return; }
+	if (!ensure(Simulation) || (OwnerComponent && OwnerComponent->bDisableClothSimulation))  // Can't debug draw when the simulation is disabled
+	{
+		return;
+	}
 
 	const FClothingSimulation* const ChaosSimulation = static_cast<const FClothingSimulation*>(Simulation);
 
@@ -155,9 +158,12 @@ void FSimulationEditorExtender::DebugDrawSimulation(const IClothingSimulation* S
 	}
 }
 
-void FSimulationEditorExtender::DebugDrawSimulationTexts(const IClothingSimulation* Simulation, USkeletalMeshComponent* /*OwnerComponent*/, FCanvas* Canvas, const FSceneView* SceneView)
+void FSimulationEditorExtender::DebugDrawSimulationTexts(const IClothingSimulation* Simulation, USkeletalMeshComponent* OwnerComponent, FCanvas* Canvas, const FSceneView* SceneView)
 {
-	if (!ensure(Simulation)) { return; }
+	if (!ensure(Simulation) || (OwnerComponent && OwnerComponent->bDisableClothSimulation))  // Can't debug draw when the simulation is disabled
+	{
+		return;
+	}
 
 	const FClothingSimulation* const ChaosSimulation = static_cast<const FClothingSimulation*>(Simulation);
 
