Chaos Cloth - Remove redundant cloth reset in UChaosClothSharedSimConfig::PostEditChangeChainProperty that already happens in UClothingAssetCommon::PostEditChangeChainProperty.

#rnx

diff --git a/Engine/Plugins/ChaosCloth/Source/ChaosCloth/Private/ChaosCloth/ChaosClothConfig.cpp b/Engine/Plugins/ChaosCloth/Source/ChaosCloth/Private/ChaosCloth/ChaosClothConfig.cpp
--- a/Engine/Plugins/ChaosCloth/Source/ChaosCloth/Private/ChaosCloth/ChaosClothConfig.cpp
+++ b/Engine/Plugins/ChaosCloth/Source/ChaosCloth/Private/ChaosCloth/ChaosClothConfig.cpp
@@ -280,27 +280,6 @@ void UChaosClothSharedSimConfig::PostEditChangeChainProperty(FPropertyChangedCha
 	{
 		MaxIterationCount = FMath::Max(MaxIterationCount, IterationCount);
 	}
-
-	// Update the simulation if there is any interactor attached to the skeletal mesh component
-	if (ChainEvent.ChangeType != EPropertyChangeType::Interactive)
-	{
-		if (USkeletalMesh* const OwnerMesh = Cast<USkeletalMesh>(GetOuter()))
-		{
-			for (TObjectIterator<USkeletalMeshComponent> It; It; ++It)
-			{
-				if (const USkeletalMeshComponent* const Component = *It)
-				{
-					if (Component->GetSkeletalMeshAsset() == OwnerMesh)
-					{
-						if (UClothingSimulationInteractor* const CurInteractor = Component->GetClothingSimulationInteractor())
-						{
-							CurInteractor->ClothConfigUpdated();
-						}
-					}
-				}
-			}
-		}
-	}
 }
 #endif  // #if WITH_EDITOR
 
