Cloth - Renamed the Skeletal Mesh Editor Clothing Visualizations to add the words "Skeletal Mesh" so that it can be discerned with other types of clothing (such as "Cloth Asset Visualizations").

#rnx

diff --git a/Engine/Plugins/ChaosCloth/Source/ChaosClothEditor/Private/ChaosClothEditor/ChaosSimulationEditorExtender.cpp b/Engine/Plugins/ChaosCloth/Source/ChaosClothEditor/Private/ChaosClothEditor/ChaosSimulationEditorExtender.cpp
--- a/Engine/Plugins/ChaosCloth/Source/ChaosClothEditor/Private/ChaosClothEditor/ChaosSimulationEditorExtender.cpp
+++ b/Engine/Plugins/ChaosCloth/Source/ChaosClothEditor/Private/ChaosClothEditor/ChaosSimulationEditorExtender.cpp
@@ -98,7 +98,7 @@ UClass* FSimulationEditorExtender::GetSupportedSimulationFactoryClass()
 
 void FSimulationEditorExtender::ExtendViewportShowMenu(FMenuBuilder& MenuBuilder, TSharedRef<IPersonaPreviewScene> PreviewScene)
 {
-	MenuBuilder.BeginSection(TEXT("ChaosSimulation_Visualizations"), LOCTEXT("VisualizationSection", "Visualizations"));
+	MenuBuilder.BeginSection(TEXT("ChaosSimulation_Visualizations"), LOCTEXT("VisualizationSection", "Skeletal Mesh Clothing Visualization"));
 	{
 		for (uint32 OptionIndex = 0; OptionIndex < FVisualizationOption::Count; ++OptionIndex)
 		{
