Fix CIS issue when INTEL_ISPC is not defined.

diff --git a/Engine/Plugins/ChaosCloth/Source/ChaosCloth/Private/ChaosCloth/ChaosClothingAccessoryMeshData.cpp b/Engine/Plugins/ChaosCloth/Source/ChaosCloth/Private/ChaosCloth/ChaosClothingAccessoryMeshData.cpp
--- a/Engine/Plugins/ChaosCloth/Source/ChaosCloth/Private/ChaosCloth/ChaosClothingAccessoryMeshData.cpp
+++ b/Engine/Plugins/ChaosCloth/Source/ChaosCloth/Private/ChaosCloth/ChaosClothingAccessoryMeshData.cpp
@@ -59,6 +59,7 @@ namespace Chaos
 	void FClothingAccessoryMeshData::ApplyPreSimulationTransforms(const Softs::FSolverVec3& DeltaLocalSpaceLocation, const Softs::FSolverRigidTransform3& GroupSpaceTransform, const Softs::FSolverReal DeltaTime)
 	{
 		TRACE_CPUPROFILER_EVENT_SCOPE(FClothingAccessoryMeshData_ApplyPreSimulationTransforms);
+#if INTEL_ISPC
 		if (bRealTypeCompatibleWithISPC && bChaos_PreSimulationTransforms_ISPC_Enabled)
 		{
 			ispc::ApplyAnimationPreSimulationTransform(
@@ -71,6 +72,7 @@ namespace Chaos
 				OldAnimationPositions.Num());
 		}
 		else
+#endif
 		{
 			for (int32 Index = 0; Index < OldAnimationPositions.Num(); ++Index)
 			{
@@ -83,6 +85,7 @@ namespace Chaos
 	void FClothingAccessoryMeshData::PreSubstep(const Softs::FSolverReal InterpolationAlpha)
 	{
 		TRACE_CPUPROFILER_EVENT_SCOPE(FClothingAccessoryMeshData_PreSubstep);
+#if INTEL_ISPC
 		if (bRealTypeCompatibleWithISPC && bChaos_PreSubstepInterpolation_ISPC_Enabled)
 		{
 			ispc::PreSubstepInterpolation(
@@ -97,6 +100,7 @@ namespace Chaos
 				InterpolatedAnimationPositions.Num());
 		}
 		else
+#endif
 		{
 			for (int32 Index = 0; Index < InterpolatedAnimationPositions.Num(); ++Index)
 			{
