Chaos Cloth Asset - Minor dereference fix.

#rnx

diff --git a/Engine/Plugins/ChaosCloth/Source/ChaosCloth/Private/ChaosCloth/SkeletalMeshComponentCacheAdapter.cpp b/Engine/Plugins/ChaosCloth/Source/ChaosCloth/Private/ChaosCloth/SkeletalMeshComponentCacheAdapter.cpp
--- a/Engine/Plugins/ChaosCloth/Source/ChaosCloth/Private/ChaosCloth/SkeletalMeshComponentCacheAdapter.cpp
+++ b/Engine/Plugins/ChaosCloth/Source/ChaosCloth/Private/ChaosCloth/SkeletalMeshComponentCacheAdapter.cpp
@@ -86,7 +86,7 @@ namespace Chaos
 
 				for (FClothingSimulationInstance& ClothingSimulationInstance : MeshComp->GetClothingSimulationInstances())
 				{
-					if (ClothingSimulationInstance.GetClothingSimulationFactory())
+					if (ClothingSimulationInstance.GetClothingSimulation() && ClothingSimulationInstance.GetClothingSimulationFactory())
 					{
 						// TODO: Manage multiple solvers and solver types as to process all instances
 						if (ClothingSimulationInstance.GetClothingSimulationFactory()->GetClass() == UChaosClothingSimulationFactory::StaticClass())
