Handle empty cloth initialization with new accessory meshes.

diff --git a/Engine/Plugins/ChaosCloth/Source/ChaosCloth/Private/ChaosCloth/ChaosClothingSimulationSolver.cpp b/Engine/Plugins/ChaosCloth/Source/ChaosCloth/Private/ChaosCloth/ChaosClothingSimulationSolver.cpp
--- a/Engine/Plugins/ChaosCloth/Source/ChaosCloth/Private/ChaosCloth/ChaosClothingSimulationSolver.cpp
+++ b/Engine/Plugins/ChaosCloth/Source/ChaosCloth/Private/ChaosCloth/ChaosClothingSimulationSolver.cpp
@@ -1880,8 +1880,10 @@ void FClothingSimulationSolver::ApplyPreSimulationTransforms()
 
 	for (FClothingSimulationCloth* const Cloth : Cloths)
 	{
-		check(Cloth);
-		Cloth->ApplyPreSimulationTransforms(this, DeltaLocalSpaceLocation, PreSimulationTransforms[Cloth->GetGroupId()], DeltaTime);
+		if (Cloth && PreSimulationTransforms.IsValidIndex(Cloth->GetGroupId()))
+		{
+			Cloth->ApplyPreSimulationTransforms(this, DeltaLocalSpaceLocation, PreSimulationTransforms[Cloth->GetGroupId()], DeltaTime);
+		}
 	}
 }
 
