Chaos Cloth Asset - Added SKM Cloth asset editor extender class for debug draw display.

#rnx

diff --git a/Engine/Plugins/ChaosCloth/Source/ChaosCloth/Private/ChaosCloth/ChaosClothVisualization.cpp b/Engine/Plugins/ChaosCloth/Source/ChaosCloth/Private/ChaosCloth/ChaosClothVisualization.cpp
--- a/Engine/Plugins/ChaosCloth/Source/ChaosCloth/Private/ChaosCloth/ChaosClothVisualization.cpp
+++ b/Engine/Plugins/ChaosCloth/Source/ChaosCloth/Private/ChaosCloth/ChaosClothVisualization.cpp
@@ -3432,10 +3432,4 @@ namespace Chaos
 
 	FClothVisualizationNoGC::~FClothVisualizationNoGC() = default;
 #endif  // #else #if CHAOS_DEBUG_DRAW
-
-	FClothVisualization::FClothVisualization(const ::Chaos::FClothingSimulationSolver* InSolver)
-		: FClothVisualizationNoGC(InSolver)
-	{}
-
-	FClothVisualization::~FClothVisualization() = default;
 }  // End namespace Chaos
diff --git a/Engine/Plugins/ChaosCloth/Source/ChaosCloth/Public/ChaosCloth/ChaosClothVisualization.h b/Engine/Plugins/ChaosCloth/Source/ChaosCloth/Public/ChaosCloth/ChaosClothVisualization.h
--- a/Engine/Plugins/ChaosCloth/Source/ChaosCloth/Public/ChaosCloth/ChaosClothVisualization.h
+++ b/Engine/Plugins/ChaosCloth/Source/ChaosCloth/Public/ChaosCloth/ChaosClothVisualization.h
@@ -37,11 +37,11 @@ namespace Chaos
 		CHAOSCLOTH_API void DrawPhysMeshWired(FPrimitiveDrawInterface* PDI = nullptr) const;
 		CHAOSCLOTH_API void DrawAnimMeshWired(FPrimitiveDrawInterface* PDI = nullptr) const;
 		CHAOSCLOTH_API void DrawAnimVelocities(FPrimitiveDrawInterface* PDI = nullptr) const;
-		CHAOSCLOTH_API void DrawAnimNormals(FPrimitiveDrawInterface* PDI = nullptr, const FReal Length = 20.) const;
+		CHAOSCLOTH_API void DrawAnimNormals(FPrimitiveDrawInterface* PDI, const FReal Length) const;
 		CHAOSCLOTH_API void DrawOpenEdges(FPrimitiveDrawInterface* PDI = nullptr) const;
-		CHAOSCLOTH_API void DrawPointNormals(FPrimitiveDrawInterface* PDI = nullptr, const FReal Length = 20.) const;
+		CHAOSCLOTH_API void DrawPointNormals(FPrimitiveDrawInterface* PDI, const FReal Length) const;
 		CHAOSCLOTH_API void DrawPointVelocities(FPrimitiveDrawInterface* PDI = nullptr) const;
-		CHAOSCLOTH_API void DrawCollision(FPrimitiveDrawInterface* PDI = nullptr, bool bWireframe = false) const;
+		CHAOSCLOTH_API void DrawCollision(FPrimitiveDrawInterface* PDI, bool bWireframe) const;
 		CHAOSCLOTH_API void DrawBackstops(FPrimitiveDrawInterface* PDI = nullptr) const;
 		CHAOSCLOTH_API void DrawBackstopDistances(FPrimitiveDrawInterface* PDI = nullptr) const;
 		CHAOSCLOTH_API void DrawMaxDistances(FPrimitiveDrawInterface* PDI = nullptr) const;
@@ -49,7 +49,7 @@ namespace Chaos
 		CHAOSCLOTH_API void DrawEdgeConstraint(FPrimitiveDrawInterface* PDI = nullptr) const;
 		CHAOSCLOTH_API void DrawBendingConstraint(FPrimitiveDrawInterface* PDI = nullptr) const;
 		CHAOSCLOTH_API void DrawLongRangeConstraint(FPrimitiveDrawInterface* PDI = nullptr) const;
-		CHAOSCLOTH_API void DrawWindAndPressureForces(FPrimitiveDrawInterface* PDI = nullptr, const FReal ForceLengthScale = (FReal)10.) const;
+		CHAOSCLOTH_API void DrawWindAndPressureForces(FPrimitiveDrawInterface* PDI, const FReal ForceLengthScale) const;
 		CHAOSCLOTH_API void DrawLocalSpace(FPrimitiveDrawInterface* PDI = nullptr) const;
 		CHAOSCLOTH_API void DrawSelfCollision(FPrimitiveDrawInterface* PDI = nullptr) const;
 		CHAOSCLOTH_API void DrawSelfIntersection(FPrimitiveDrawInterface* PDI = nullptr) const;
@@ -73,11 +73,11 @@ namespace Chaos
 		void DrawMaxDistanceValues(FCanvas* /*Canvas*/, const FSceneView* /*SceneView*/) const {}
 		void DrawPhysMeshWired(FPrimitiveDrawInterface* /*PDI*/ = nullptr) const {}
 		void DrawAnimMeshWired(FPrimitiveDrawInterface* /*PDI*/ = nullptr) const {}
-		void DrawAnimNormals(FPrimitiveDrawInterface* /*PDI*/ = nullptr, const FReal /*Length*/ = 20.) const {}
+		void DrawAnimNormals(FPrimitiveDrawInterface* /*PDI*/, const FReal /*Length*/) const {}
 		void DrawOpenEdges(FPrimitiveDrawInterface* /*PDI*/ = nullptr) const {}
-		void DrawPointNormals(FPrimitiveDrawInterface* /*PDI*/ = nullptr, const FReal /*Length*/ = 20.) const {}
+		void DrawPointNormals(FPrimitiveDrawInterface* /*PDI*/, const FReal /*Length*/) const {}
 		void DrawPointVelocities(FPrimitiveDrawInterface* /*PDI*/ = nullptr) const {}
-		void DrawCollision(FPrimitiveDrawInterface* /*PDI*/ = nullptr, bool /*bWireframe*/ = false) const {}
+		void DrawCollision(FPrimitiveDrawInterface* /*PDI*/, bool /*bWireframe*/) const {}
 		void DrawBackstops(FPrimitiveDrawInterface* /*PDI*/ = nullptr) const {}
 		void DrawBackstopDistances(FPrimitiveDrawInterface* /*PDI*/ = nullptr) const {}
 		void DrawMaxDistances(FPrimitiveDrawInterface* /*PDI*/ = nullptr) const {}
@@ -85,7 +85,7 @@ namespace Chaos
 		void DrawEdgeConstraint(FPrimitiveDrawInterface* /*PDI*/ = nullptr) const {}
 		void DrawBendingConstraint(FPrimitiveDrawInterface* /*PDI*/ = nullptr) const {}
 		void DrawLongRangeConstraint(FPrimitiveDrawInterface* /*PDI*/ = nullptr) const {}
-		void DrawWindAndPressureForces(FPrimitiveDrawInterface* /*PDI*/ = nullptr, const FReal /*ForceLengthScale*/ = (FReal)10.) const {}
+		void DrawWindAndPressureForces(FPrimitiveDrawInterface* /*PDI*/, const FReal /*ForceLengthScale*/) const {}
 		void DrawLocalSpace(FPrimitiveDrawInterface* /*PDI*/ = nullptr) const {}
 		void DrawSelfCollision(FPrimitiveDrawInterface* /*PDI*/ = nullptr) const {}
 		void DrawSelfIntersection(FPrimitiveDrawInterface* /*PDI*/ = nullptr) const {}
@@ -97,8 +97,31 @@ namespace Chaos
 		void DrawMultiResConstraint(FPrimitiveDrawInterface* /*PDI*/ = nullptr) const {}
 		void DrawClothClothConstraints(FPrimitiveDrawInterface* /*PDI*/ = nullptr) const {}
 		void DrawExtremlyDeformedEdges(FPrimitiveDrawInterface* /*PDI*/ = nullptr) const {}
+		void DrawAccessoryMesh(FPrimitiveDrawInterface* /*PDI*/ = nullptr, const FName& AccessoryMeshName = NAME_None) const {}
+		void DrawAccessoryMeshNormals(FPrimitiveDrawInterface* /*PDI*/ = nullptr, const FName& AccessoryMeshName = NAME_None, const FReal Length = 20.) const;
 #endif  // #if CHAOS_DEBUG_DRAW
 
+		void DrawAnimNormals(FPrimitiveDrawInterface* PDI = nullptr) const
+		{
+			constexpr FReal DefaultLengtth = 20.;
+			DrawAnimNormals(PDI, DefaultLengtth);
+		}
+		void DrawPointNormals(FPrimitiveDrawInterface* PDI = nullptr) const
+		{
+			constexpr FReal DefaultLengtth = 20.;
+			DrawPointNormals(PDI, DefaultLengtth);
+		}
+		void DrawCollision(FPrimitiveDrawInterface* PDI = nullptr) const
+		{
+			constexpr bool bDefaultWireframe = false;
+			DrawCollision(PDI, bDefaultWireframe);
+		}
+		void DrawWindAndPressureForces(FPrimitiveDrawInterface* PDI = nullptr) const
+		{
+			constexpr FReal DefaultForceLengthScale = 10.;
+			DrawWindAndPressureForces(PDI, DefaultForceLengthScale);
+		}
+
 #if WITH_EDITOR && CHAOS_DEBUG_DRAW
 		// Editor only functions
 		CHAOSCLOTH_API void DrawPhysMeshShaded(FPrimitiveDrawInterface* PDI) const;
@@ -125,24 +148,4 @@ namespace Chaos
 
 		class FMaterials;
 	};
-
-	// Deprecated FGCObject inheritance
-	class UE_DEPRECATED(5.5, "Use FClothVisualizationNoGC instead") FClothVisualization : public ::Chaos::FClothVisualizationNoGC
-#if WITH_EDITOR
-		, public FGCObject
-#endif  // #if WITH_EDITOR
-	{
-	public:
-		CHAOSCLOTH_API explicit FClothVisualization(const ::Chaos::FClothingSimulationSolver* InSolver = nullptr);
-		CHAOSCLOTH_API virtual ~FClothVisualization() override;
-#if WITH_EDITOR
-	protected:
-		//~ Begin FGCObject interface
-		UE_DEPRECATED(5.5, "The FGCObject inheritance is no longer needed.")
-		virtual void AddReferencedObjects(FReferenceCollector& /*Collector*/) override {}
-		UE_DEPRECATED(5.5, "The FGCObject inheritance is no longer needed.")
-		virtual FString GetReferencerName() const override { return TEXT("Chaos::FVisualization"); }
-		//~ End FGCObject interface
-#endif  // #if WITH_EDITOR
- 	};
 } // End namespace Chaos
diff --git a/Engine/Plugins/ChaosCloth/Source/ChaosCloth/Public/ChaosCloth/ChaosClothingSimulation.h b/Engine/Plugins/ChaosCloth/Source/ChaosCloth/Public/ChaosCloth/ChaosClothingSimulation.h
--- a/Engine/Plugins/ChaosCloth/Source/ChaosCloth/Public/ChaosCloth/ChaosClothingSimulation.h
+++ b/Engine/Plugins/ChaosCloth/Source/ChaosCloth/Public/ChaosCloth/ChaosClothingSimulation.h
@@ -118,6 +118,12 @@ namespace Chaos
 		void DebugDrawExtremlyDeformedEdges(FPrimitiveDrawInterface* PDI = nullptr) const { Visualization.DrawExtremlyDeformedEdges(PDI); }
 #endif  // #if CHAOS_DEBUG_DRAW
 
+		/** Return the visualization object for this simulation. */
+		const FClothVisualizationNoGC* GetClothVisualization() const
+		{
+			return &Visualization;
+		}
+
 	private:
 		void ResetStats();
 		void UpdateStats(const FClothingSimulationCloth* Cloth);
@@ -128,9 +134,7 @@ namespace Chaos
 
 	private:
 		// Visualization object
-PRAGMA_DISABLE_DEPRECATION_WARNINGS
-		FClothVisualization Visualization;  // TODO: Replace with FClothVisualizationNoGC in 5.7
-PRAGMA_ENABLE_DEPRECATION_WARNINGS
+		FClothVisualizationNoGC Visualization;
 
 		// Simulation objects
 		TUniquePtr<FClothingSimulationSolver> Solver;  // Default solver
diff --git a/Engine/Plugins/ChaosCloth/Source/ChaosClothEditor/Private/ChaosClothEditor/ChaosSimulationEditorExtender.cpp b/Engine/Plugins/ChaosCloth/Source/ChaosClothEditor/Private/ChaosClothEditor/ChaosSimulationEditorExtender.cpp
--- a/Engine/Plugins/ChaosCloth/Source/ChaosClothEditor/Private/ChaosClothEditor/ChaosSimulationEditorExtender.cpp
+++ b/Engine/Plugins/ChaosCloth/Source/ChaosClothEditor/Private/ChaosClothEditor/ChaosSimulationEditorExtender.cpp
@@ -4,6 +4,7 @@
 
 #include "ChaosCloth/ChaosClothingSimulationFactory.h"
 #include "ChaosCloth/ChaosClothingSimulation.h"
+#include "ChaosCloth/ChaosClothVisualization.h"
 #include "Framework/MultiBox/MultiBoxBuilder.h"
 #include "Animation/DebugSkelMeshComponent.h"
 #include "Rendering/SkeletalMeshRenderData.h"
@@ -21,8 +22,8 @@ struct FVisualizationOption
 	static const uint32 Count;
 
 	// Chaos debug draw function
-	typedef void (Chaos::FClothingSimulation::*FDebugDrawFunction)(FPrimitiveDrawInterface*) const;
-	typedef void (Chaos::FClothingSimulation::*FDebugDrawTextsFunction)(FCanvas*, const FSceneView*) const;
+	typedef void (FClothVisualizationNoGC::*FDebugDrawFunction)(FPrimitiveDrawInterface*) const;
+	typedef void (FClothVisualizationNoGC::*FDebugDrawTextsFunction)(FCanvas*, const FSceneView*) const;
 	FDebugDrawFunction DebugDrawFunction;
 	FDebugDrawTextsFunction DebugDrawTextsFunction;
 
@@ -62,27 +63,27 @@ struct FVisualizationOption
 
 const FVisualizationOption FVisualizationOption::OptionData[] = 
 {
-	FVisualizationOption(&Chaos::FClothingSimulation::DebugDrawPhysMeshShaded      , TEXT("p.ChaosClothEditor.DebugDrawPhysMeshShaded"      ), LOCTEXT("ChaosVisName_PhysMesh"            , "Physical Mesh (Flat Shaded)"), LOCTEXT("ChaosVisName_PhysMeshShaded_ToolTip"      , "Draws the current physical result as a doubled sided flat shaded mesh"), /*bDisablesSimulation =*/false, /*bHidesClothSections=*/true),
-	FVisualizationOption(&Chaos::FClothingSimulation::DebugDrawPhysMeshWired       , TEXT("p.ChaosClothEditor.DebugDrawPhysMeshWired"       ), LOCTEXT("ChaosVisName_PhysMeshWire"        , "Physical Mesh (Wireframe)"  ), LOCTEXT("ChaosVisName_PhysMeshWired_ToolTip"       , "Draws the current physical mesh result in wireframe")),
-	FVisualizationOption(&Chaos::FClothingSimulation::DebugDrawAnimMeshWired       , TEXT("p.ChaosClothEditor.DebugDrawAnimMeshWired"       ), LOCTEXT("ChaosVisName_AnimMeshWire"        , "Animated Mesh (Wireframe)"  ), LOCTEXT("ChaosVisName_AnimMeshWired_ToolTip"       , "Draws the current animated mesh input in wireframe")),
-	FVisualizationOption(&Chaos::FClothingSimulation::DebugDrawParticleIndices     , TEXT("p.ChaosClothEditor.DebugDrawParticleIndices"     ), LOCTEXT("ChaosVisName_ParticleIndices"     , "Particle Indices"           ), LOCTEXT("ChaosVisName_ParticleIndices_ToolTip"     , "Draws the particle indices as instantiated by the solver")),
-	FVisualizationOption(&Chaos::FClothingSimulation::DebugDrawElementIndices      , TEXT("p.ChaosClothEditor.DebugDrawElementIndices"      ), LOCTEXT("ChaosVisName_ElementIndices"      , "Element Indices"            ), LOCTEXT("ChaosVisName_ElementIndices_ToolTip"      , "Draws the element's (triangle or other) indices as instantiated by the solver")),
-	FVisualizationOption(&Chaos::FClothingSimulation::DebugDrawPointNormals        , TEXT("p.ChaosClothEditor.DebugDrawPointNormals"        ), LOCTEXT("ChaosVisName_PointNormals"        , "Physical Mesh Normals"      ), LOCTEXT("ChaosVisName_PointNormals_ToolTip"        , "Draws the current point normals for the simulation mesh")),
-	FVisualizationOption(&Chaos::FClothingSimulation::DebugDrawPointVelocities     , TEXT("p.ChaosClothEditor.DebugDrawPointVelocities"     ), LOCTEXT("ChaosVisName_PointVelocities"     , "Point Velocities"           ), LOCTEXT("ChaosVisName_PointVelocities_ToolTip"     , "Draws the current point velocities for the simulation mesh")),
-	FVisualizationOption(&Chaos::FClothingSimulation::DebugDrawAnimNormals         , TEXT("p.ChaosClothEditor.DebugDrawAnimNormals"         ), LOCTEXT("ChaosVisName_AnimNormals"         , "Animated Mesh Normals"      ), LOCTEXT("ChaosVisName_AnimNormals_ToolTip"         , "Draws the current point normals for the animated mesh")),
-	FVisualizationOption(&Chaos::FClothingSimulation::DebugDrawCollision           , TEXT("p.ChaosClothEditor.DebugDrawCollision"           ), LOCTEXT("ChaosVisName_Collision"           , "Collisions"                 ), LOCTEXT("ChaosVisName_Collision_ToolTip"           , "Draws the collision bodies the simulation is currently using")),
-	FVisualizationOption(&Chaos::FClothingSimulation::DebugDrawBackstops           , TEXT("p.ChaosClothEditor.DebugDrawBackstops"           ), LOCTEXT("ChaosVisName_Backstop"            , "Backstops"                  ), LOCTEXT("ChaosVisName_Backstop_ToolTip"            , "Draws the backstop radius and position for each simulation particle")),
-	FVisualizationOption(&Chaos::FClothingSimulation::DebugDrawBackstopDistances   , TEXT("p.ChaosClothEditor.DebugDrawBackstopDistances"   ), LOCTEXT("ChaosVisName_BackstopDistance"    , "Backstop Distances"         ), LOCTEXT("ChaosVisName_BackstopDistance_ToolTip"    , "Draws the backstop distance offset for each simulation particle"), /*bDisablesSimulation =*/true),
-	FVisualizationOption(&Chaos::FClothingSimulation::DebugDrawMaxDistances        , TEXT("p.ChaosClothEditor.DebugDrawMaxDistances"        ), LOCTEXT("ChaosVisName_MaxDistance"         , "Max Distances"              ), LOCTEXT("ChaosVisName_MaxDistance_ToolTip"         , "Draws the current max distances for the sim particles as a line along its normal"), /*bDisablesSimulation =*/true),
-	FVisualizationOption(&Chaos::FClothingSimulation::DebugDrawMaxDistanceValues   , TEXT("p.ChaosClothEditor.DebugDrawMaxDistanceValues"   ), LOCTEXT("ChaosVisName_MaxDistanceValue"    , "Max Distances As Numbers"   ), LOCTEXT("ChaosVisName_MaxDistanceValue_ToolTip"    , "Draws the current max distances as numbers")),
-	FVisualizationOption(&Chaos::FClothingSimulation::DebugDrawAnimDrive           , TEXT("p.ChaosClothEditor.DebugDrawAnimDrive"           ), LOCTEXT("ChaosVisName_AnimDrive"           , "Anim Drive"                 ), LOCTEXT("ChaosVisName_AnimDrive_Tooltip"           , "Draws the current skinned reference mesh for the simulation which anim drive will attempt to reach if enabled")),
-	FVisualizationOption(&Chaos::FClothingSimulation::DebugDrawEdgeConstraint      , TEXT("p.ChaosClothEditor.DebugDrawEdgeConstraint"      ), LOCTEXT("ChaosVisName_EdgeConstraint"      , "Edge Constraint"            ), LOCTEXT("ChaosVisName_EdgeConstraint_Tooltip"      , "Draws the edge spring constraints")),
-	FVisualizationOption(&Chaos::FClothingSimulation::DebugDrawBendingConstraint   , TEXT("p.ChaosClothEditor.DebugDrawBendingConstraint"   ), LOCTEXT("ChaosVisName_BendingConstraint"   , "Bending Constraint"         ), LOCTEXT("ChaosVisName_BendingConstraint_Tooltip"   , "Draws the bending spring constraints")),
-	FVisualizationOption(&Chaos::FClothingSimulation::DebugDrawLongRangeConstraint , TEXT("p.ChaosClothEditor.DebugDrawLongRangeConstraint" ), LOCTEXT("ChaosVisName_LongRangeConstraint" , "Long Range Constraint"      ), LOCTEXT("ChaosVisName_LongRangeConstraint_Tooltip" , "Draws the long range attachment constraint distances")),
-	FVisualizationOption(&Chaos::FClothingSimulation::DebugDrawWindAndPressureForces, TEXT("p.ChaosClothEditor.DebugDrawWindAndPressureForces"), LOCTEXT("ChaosVisName_WindAndPressureForces", "Wind Aerodynamic And Pressure Forces"), LOCTEXT("ChaosVisName_WindAndPressure_Tooltip", "Draws the Wind drag and lift and pressure forces")),
-	FVisualizationOption(&Chaos::FClothingSimulation::DebugDrawLocalSpace          , TEXT("p.ChaosClothEditor.DebugDrawLocalSpace"          ), LOCTEXT("ChaosVisName_LocalSpace"          , "Local Space Reference Bone" ), LOCTEXT("ChaosVisName_LocalSpace_Tooltip"          , "Draws the local space reference bone")),
-	FVisualizationOption(&Chaos::FClothingSimulation::DebugDrawSelfCollision       , TEXT("p.ChaosClothEditor.DebugDrawSelfCollision"       ), LOCTEXT("ChaosVisName_SelfCollision"       , "Self Collision"             ), LOCTEXT("ChaosVisName_SelfCollision_Tooltip"       , "Draws the self collision thickness/debugging information")),
-	FVisualizationOption(&Chaos::FClothingSimulation::DebugDrawSelfIntersection    , TEXT("p.ChaosClothEditor.DebugDrawSelfIntersection"    ), LOCTEXT("ChaosVisName_SelfIntersection"    , "Self Intersection"          ), LOCTEXT("ChaosVisName_SelfIntersection_Tooltip"    , "Draws the self intersection contour/region information")),
+	FVisualizationOption(&FClothVisualizationNoGC::DrawPhysMeshShaded      , TEXT("p.ChaosClothEditor.DebugDrawPhysMeshShaded"      ), LOCTEXT("ChaosVisName_PhysMesh"            , "Physical Mesh (Flat Shaded)"), LOCTEXT("ChaosVisName_PhysMeshShaded_ToolTip"      , "Draws the current physical result as a doubled sided flat shaded mesh"), /*bDisablesSimulation =*/false, /*bHidesClothSections=*/true),
+	FVisualizationOption(&FClothVisualizationNoGC::DrawPhysMeshWired       , TEXT("p.ChaosClothEditor.DebugDrawPhysMeshWired"       ), LOCTEXT("ChaosVisName_PhysMeshWire"        , "Physical Mesh (Wireframe)"  ), LOCTEXT("ChaosVisName_PhysMeshWired_ToolTip"       , "Draws the current physical mesh result in wireframe")),
+	FVisualizationOption(&FClothVisualizationNoGC::DrawAnimMeshWired       , TEXT("p.ChaosClothEditor.DebugDrawAnimMeshWired"       ), LOCTEXT("ChaosVisName_AnimMeshWire"        , "Animated Mesh (Wireframe)"  ), LOCTEXT("ChaosVisName_AnimMeshWired_ToolTip"       , "Draws the current animated mesh input in wireframe")),
+	FVisualizationOption(&FClothVisualizationNoGC::DrawParticleIndices     , TEXT("p.ChaosClothEditor.DebugDrawParticleIndices"     ), LOCTEXT("ChaosVisName_ParticleIndices"     , "Particle Indices"           ), LOCTEXT("ChaosVisName_ParticleIndices_ToolTip"     , "Draws the particle indices as instantiated by the solver")),
+	FVisualizationOption(&FClothVisualizationNoGC::DrawElementIndices      , TEXT("p.ChaosClothEditor.DebugDrawElementIndices"      ), LOCTEXT("ChaosVisName_ElementIndices"      , "Element Indices"            ), LOCTEXT("ChaosVisName_ElementIndices_ToolTip"      , "Draws the element's (triangle or other) indices as instantiated by the solver")),
+	FVisualizationOption(&FClothVisualizationNoGC::DrawPointNormals        , TEXT("p.ChaosClothEditor.DebugDrawPointNormals"        ), LOCTEXT("ChaosVisName_PointNormals"        , "Physical Mesh Normals"      ), LOCTEXT("ChaosVisName_PointNormals_ToolTip"        , "Draws the current point normals for the simulation mesh")),
+	FVisualizationOption(&FClothVisualizationNoGC::DrawPointVelocities     , TEXT("p.ChaosClothEditor.DebugDrawPointVelocities"     ), LOCTEXT("ChaosVisName_PointVelocities"     , "Point Velocities"           ), LOCTEXT("ChaosVisName_PointVelocities_ToolTip"     , "Draws the current point velocities for the simulation mesh")),
+	FVisualizationOption(&FClothVisualizationNoGC::DrawAnimNormals         , TEXT("p.ChaosClothEditor.DebugDrawAnimNormals"         ), LOCTEXT("ChaosVisName_AnimNormals"         , "Animated Mesh Normals"      ), LOCTEXT("ChaosVisName_AnimNormals_ToolTip"         , "Draws the current point normals for the animated mesh")),
+	FVisualizationOption(&FClothVisualizationNoGC::DrawCollision           , TEXT("p.ChaosClothEditor.DebugDrawCollision"           ), LOCTEXT("ChaosVisName_Collision"           , "Collisions"                 ), LOCTEXT("ChaosVisName_Collision_ToolTip"           , "Draws the collision bodies the simulation is currently using")),
+	FVisualizationOption(&FClothVisualizationNoGC::DrawBackstops           , TEXT("p.ChaosClothEditor.DebugDrawBackstops"           ), LOCTEXT("ChaosVisName_Backstop"            , "Backstops"                  ), LOCTEXT("ChaosVisName_Backstop_ToolTip"            , "Draws the backstop radius and position for each simulation particle")),
+	FVisualizationOption(&FClothVisualizationNoGC::DrawBackstopDistances   , TEXT("p.ChaosClothEditor.DebugDrawBackstopDistances"   ), LOCTEXT("ChaosVisName_BackstopDistance"    , "Backstop Distances"         ), LOCTEXT("ChaosVisName_BackstopDistance_ToolTip"    , "Draws the backstop distance offset for each simulation particle"), /*bDisablesSimulation =*/true),
+	FVisualizationOption(&FClothVisualizationNoGC::DrawMaxDistances        , TEXT("p.ChaosClothEditor.DebugDrawMaxDistances"        ), LOCTEXT("ChaosVisName_MaxDistance"         , "Max Distances"              ), LOCTEXT("ChaosVisName_MaxDistance_ToolTip"         , "Draws the current max distances for the sim particles as a line along its normal"), /*bDisablesSimulation =*/true),
+	FVisualizationOption(&FClothVisualizationNoGC::DrawMaxDistanceValues   , TEXT("p.ChaosClothEditor.DebugDrawMaxDistanceValues"   ), LOCTEXT("ChaosVisName_MaxDistanceValue"    , "Max Distances As Numbers"   ), LOCTEXT("ChaosVisName_MaxDistanceValue_ToolTip"    , "Draws the current max distances as numbers")),
+	FVisualizationOption(&FClothVisualizationNoGC::DrawAnimDrive           , TEXT("p.ChaosClothEditor.DebugDrawAnimDrive"           ), LOCTEXT("ChaosVisName_AnimDrive"           , "Anim Drive"                 ), LOCTEXT("ChaosVisName_AnimDrive_Tooltip"           , "Draws the current skinned reference mesh for the simulation which anim drive will attempt to reach if enabled")),
+	FVisualizationOption(&FClothVisualizationNoGC::DrawEdgeConstraint      , TEXT("p.ChaosClothEditor.DebugDrawEdgeConstraint"      ), LOCTEXT("ChaosVisName_EdgeConstraint"      , "Edge Constraint"            ), LOCTEXT("ChaosVisName_EdgeConstraint_Tooltip"      , "Draws the edge spring constraints")),
+	FVisualizationOption(&FClothVisualizationNoGC::DrawBendingConstraint   , TEXT("p.ChaosClothEditor.DebugDrawBendingConstraint"   ), LOCTEXT("ChaosVisName_BendingConstraint"   , "Bending Constraint"         ), LOCTEXT("ChaosVisName_BendingConstraint_Tooltip"   , "Draws the bending spring constraints")),
+	FVisualizationOption(&FClothVisualizationNoGC::DrawLongRangeConstraint , TEXT("p.ChaosClothEditor.DebugDrawLongRangeConstraint" ), LOCTEXT("ChaosVisName_LongRangeConstraint" , "Long Range Constraint"      ), LOCTEXT("ChaosVisName_LongRangeConstraint_Tooltip" , "Draws the long range attachment constraint distances")),
+	FVisualizationOption(&FClothVisualizationNoGC::DrawWindAndPressureForces, TEXT("p.ChaosClothEditor.DebugDrawWindAndPressureForces"), LOCTEXT("ChaosVisName_WindAndPressureForces", "Wind Aerodynamic And Pressure Forces"), LOCTEXT("ChaosVisName_WindAndPressure_Tooltip", "Draws the Wind drag and lift and pressure forces")),
+	FVisualizationOption(&FClothVisualizationNoGC::DrawLocalSpace          , TEXT("p.ChaosClothEditor.DebugDrawLocalSpace"          ), LOCTEXT("ChaosVisName_LocalSpace"          , "Local Space Reference Bone" ), LOCTEXT("ChaosVisName_LocalSpace_Tooltip"          , "Draws the local space reference bone")),
+	FVisualizationOption(&FClothVisualizationNoGC::DrawSelfCollision       , TEXT("p.ChaosClothEditor.DebugDrawSelfCollision"       ), LOCTEXT("ChaosVisName_SelfCollision"       , "Self Collision"             ), LOCTEXT("ChaosVisName_SelfCollision_Tooltip"       , "Draws the self collision thickness/debugging information")),
+	FVisualizationOption(&FClothVisualizationNoGC::DrawSelfIntersection    , TEXT("p.ChaosClothEditor.DebugDrawSelfIntersection"    ), LOCTEXT("ChaosVisName_SelfIntersection"    , "Self Intersection"          ), LOCTEXT("ChaosVisName_SelfIntersection_Tooltip"    , "Draws the self intersection contour/region information")),
 };
 const uint32 FVisualizationOption::Count = sizeof(OptionData) / sizeof(FVisualizationOption);
 
@@ -91,6 +92,8 @@ FSimulationEditorExtender::FSimulationEditorExtender()
 {
 }
 
+FSimulationEditorExtender::~FSimulationEditorExtender() = default;
+
 UClass* FSimulationEditorExtender::GetSupportedSimulationFactoryClass()
 {
 	return UChaosClothingSimulationFactory::StaticClass();
@@ -98,7 +101,7 @@ UClass* FSimulationEditorExtender::GetSupportedSimulationFactoryClass()
 
 void FSimulationEditorExtender::ExtendViewportShowMenu(FMenuBuilder& MenuBuilder, TSharedRef<IPersonaPreviewScene> PreviewScene)
 {
-	MenuBuilder.BeginSection(TEXT("ChaosSimulation_Visualizations"), LOCTEXT("VisualizationSection", "Skeletal Mesh Clothing Visualization"));
+	MenuBuilder.BeginSection(GetSectionName(), GetSectionHeadingText());
 	{
 		for (uint32 OptionIndex = 0; OptionIndex < FVisualizationOption::Count; ++OptionIndex)
 		{
@@ -146,14 +149,15 @@ void FSimulationEditorExtender::DebugDrawSimulation(const IClothingSimulation* S
 		return;
 	}
 
-	const FClothingSimulation* const ChaosSimulation = static_cast<const FClothingSimulation*>(Simulation);
-
-	for (int32 OptionIndex = 0; OptionIndex < FVisualizationOption::Count; ++OptionIndex)
+	if (const FClothVisualizationNoGC* const ClothVisualization = GetClothVisualization(Simulation))
 	{
-		if (FVisualizationOption::OptionData[OptionIndex].DebugDrawFunction &&
-			(FVisualizationOption::OptionData[OptionIndex].IsConsoleVariableEnabled() || Flags[OptionIndex]))
+		for (int32 OptionIndex = 0; OptionIndex < FVisualizationOption::Count; ++OptionIndex)
 		{
-			(ChaosSimulation->*(FVisualizationOption::OptionData[OptionIndex].DebugDrawFunction))(PDI);
+			if (FVisualizationOption::OptionData[OptionIndex].DebugDrawFunction &&
+				(FVisualizationOption::OptionData[OptionIndex].IsConsoleVariableEnabled() || Flags[OptionIndex]))
+			{
+				(ClothVisualization->*(FVisualizationOption::OptionData[OptionIndex].DebugDrawFunction))(PDI);
+			}
 		}
 	}
 }
@@ -165,18 +169,34 @@ void FSimulationEditorExtender::DebugDrawSimulationTexts(const IClothingSimulati
 		return;
 	}
 
-	const FClothingSimulation* const ChaosSimulation = static_cast<const FClothingSimulation*>(Simulation);
-
-	for (int32 OptionIndex = 0; OptionIndex < FVisualizationOption::Count; ++OptionIndex)
+	if (const FClothVisualizationNoGC* const ClothVisualization = GetClothVisualization(Simulation))
 	{
-		if (FVisualizationOption::OptionData[OptionIndex].DebugDrawTextsFunction &&
-			(FVisualizationOption::OptionData[OptionIndex].IsConsoleVariableEnabled() || Flags[OptionIndex]))
+		for (int32 OptionIndex = 0; OptionIndex < FVisualizationOption::Count; ++OptionIndex)
 		{
-			(ChaosSimulation->*(FVisualizationOption::OptionData[OptionIndex].DebugDrawTextsFunction))(Canvas, SceneView);
+			if (FVisualizationOption::OptionData[OptionIndex].DebugDrawTextsFunction &&
+				(FVisualizationOption::OptionData[OptionIndex].IsConsoleVariableEnabled() || Flags[OptionIndex]))
+			{
+				(ClothVisualization->*(FVisualizationOption::OptionData[OptionIndex].DebugDrawTextsFunction))(Canvas, SceneView);
+			}
 		}
 	}
 }
 
+FName FSimulationEditorExtender::GetSectionName() const
+{
+	return FName("ChaosSimulation_Visualizations");
+}
+
+FText FSimulationEditorExtender::GetSectionHeadingText() const
+{
+	return FText(LOCTEXT("VisualizationSection", "Skeletal Mesh Clothing Visualization"));
+}
+
+const FClothVisualizationNoGC* FSimulationEditorExtender::GetClothVisualization(const IClothingSimulation* Simulation) const
+{
+	return Simulation ? static_cast<const FClothingSimulation*>(Simulation)->GetClothVisualization() : nullptr;
+}
+
 bool FSimulationEditorExtender::ShouldDisableSimulation() const
 {
 	for (uint32 OptionIndex = 0; OptionIndex < FVisualizationOption::Count; ++OptionIndex)
diff --git a/Engine/Plugins/ChaosCloth/Source/ChaosClothEditor/Public/ChaosClothEditor/ChaosSimulationEditorExtender.h b/Engine/Plugins/ChaosCloth/Source/ChaosClothEditor/Public/ChaosClothEditor/ChaosSimulationEditorExtender.h
--- a/Engine/Plugins/ChaosCloth/Source/ChaosClothEditor/Public/ChaosClothEditor/ChaosSimulationEditorExtender.h
+++ b/Engine/Plugins/ChaosCloth/Source/ChaosClothEditor/Public/ChaosClothEditor/ChaosSimulationEditorExtender.h
@@ -6,20 +6,32 @@
 
 namespace Chaos
 {
+	class FClothVisualizationNoGC;
+
 	/** Chaos extension to the asset editor. */
 	class FSimulationEditorExtender : public ISimulationEditorExtender
 	{
 	public:
-		FSimulationEditorExtender();
+		CHAOSCLOTHEDITOR_API FSimulationEditorExtender();
+		CHAOSCLOTHEDITOR_API virtual ~FSimulationEditorExtender() override;
 
 		// ISimulationEditorExtender Interface
 		virtual UClass* GetSupportedSimulationFactoryClass() override;
-		virtual void ExtendViewportShowMenu(FMenuBuilder& MenuBuilder, TSharedRef<IPersonaPreviewScene> PreviewScene) override;
-		virtual void DebugDrawSimulation(const IClothingSimulation* Simulation, USkeletalMeshComponent* OwnerComponent, FPrimitiveDrawInterface* PDI) override;
-		virtual void DebugDrawSimulationTexts(const IClothingSimulation* Simulation, USkeletalMeshComponent* OwnerComponent, FCanvas* Canvas, const FSceneView* SceneView) override;
+		CHAOSCLOTHEDITOR_API virtual void ExtendViewportShowMenu(FMenuBuilder& MenuBuilder, TSharedRef<IPersonaPreviewScene> PreviewScene) override;
+		CHAOSCLOTHEDITOR_API virtual void DebugDrawSimulation(const IClothingSimulation* Simulation, USkeletalMeshComponent* OwnerComponent, FPrimitiveDrawInterface* PDI) override;
+		CHAOSCLOTHEDITOR_API virtual void DebugDrawSimulationTexts(const IClothingSimulation* Simulation, USkeletalMeshComponent* OwnerComponent, FCanvas* Canvas, const FSceneView* SceneView) override;
 		// End ISimulationEditorExtender Interface
 
 	private:
+		/** Return the name of the menu section. */
+		virtual FName GetSectionName() const;
+
+		/** Return the heading text of the menu section. */
+		virtual FText GetSectionHeadingText() const;
+
+		/** Return the visualization object for the specified simulation. */
+		virtual const FClothVisualizationNoGC* GetClothVisualization(const IClothingSimulation* Simulation) const;
+
 		/** Return whether or not - given the current enabled options - the simulation should be disabled. */
 		bool ShouldDisableSimulation() const;
 		/** Show/hide all cloth sections for the specified mesh compoment. */
diff --git a/Engine/Plugins/ChaosClothAsset/Source/ChaosClothAssetEngine/Private/ChaosClothAsset/ClothAssetSKMClothingSimulation.h b/Engine/Plugins/ChaosClothAsset/Source/ChaosClothAssetEngine/Private/ChaosClothAsset/ClothAssetSKMClothingSimulation.h
deleted file mode 100644
diff --git a/Engine/Plugins/ChaosClothAsset/Source/ChaosClothAssetEngine/Private/ChaosClothAsset/ClothSimulationProxy.cpp b/Engine/Plugins/ChaosClothAsset/Source/ChaosClothAssetEngine/Private/ChaosClothAsset/ClothSimulationProxy.cpp
--- a/Engine/Plugins/ChaosClothAsset/Source/ChaosClothAssetEngine/Private/ChaosClothAsset/ClothSimulationProxy.cpp
+++ b/Engine/Plugins/ChaosClothAsset/Source/ChaosClothAssetEngine/Private/ChaosClothAsset/ClothSimulationProxy.cpp
@@ -175,9 +175,8 @@ namespace UE::Chaos::ClothAsset
 		// Use new SoftsEvolution, not PBDEvolution.
 		constexpr bool bUseLegacySolver = false;
 		Solver = MakeUnique<::Chaos::FClothingSimulationSolver>(Configs[SolverConfigIndex].Get(), bUseLegacySolver);
-PRAGMA_DISABLE_DEPRECATION_WARNINGS
-		Visualization = MakeUnique<::Chaos::FClothVisualization>(Solver.Get());
-PRAGMA_ENABLE_DEPRECATION_WARNINGS
+
+		Visualization = MakeUnique<::Chaos::FClothVisualizationNoGC>(Solver.Get());
 
 		// Need a valid context to initialize the mesh
 		constexpr bool bIsInitialization = true;
diff --git a/Engine/Plugins/ChaosClothAsset/Source/ChaosClothAssetEngine/Public/ChaosClothAsset/ClothAssetSKMClothingSimulation.h b/Engine/Plugins/ChaosClothAsset/Source/ChaosClothAssetEngine/Public/ChaosClothAsset/ClothAssetSKMClothingSimulation.h
new file mode 100644
--- /dev/null
+++ b/Engine/Plugins/ChaosClothAsset/Source/ChaosClothAssetEngine/Public/ChaosClothAsset/ClothAssetSKMClothingSimulation.h
@@ -0,0 +1,195 @@
+// Copyright Epic Games, Inc. All Rights Reserved.
+
+#pragma once
+
+#include "ChaosClothAsset/ClothComponentAdapter.h"
+#include "ChaosClothAsset/ClothSimulationProxy.h"
+#include "ClothingSimulationInterface.h"
+#include "ClothingSimulationFactory.h"
+#include "UObject/GCObject.h"
+#include "ClothAssetSKMClothingSimulation.generated.h"
+
+class USkinnedMeshComponent;
+class UChaosClothAssetInteractor;
+class UChaosClothAssetSKMClothingAsset;
+struct FChaosClothSimulationProperties;
+
+namespace Chaos
+{
+	class FClothVisualizationNoGC;
+	namespace Softs
+	{
+		class FCollectionPropertyFacade;
+	}
+}
+
+namespace UE::Chaos::ClothAsset
+{
+
+	/** Context for the SKM simulation, only used to temporarily store the simulation timestep, the actual simulation context is managed by the simulation proxy. */
+	class FSKMClothingSimulationContext final: public IClothingSimulationContext
+	{
+	private:
+		friend class FSKMClothingSimulation;
+
+		FSKMClothingSimulationContext() : IClothingSimulationContext() {}
+		virtual ~FSKMClothingSimulationContext() override {}
+
+		/** Simulation timestep. */
+		float DeltaTime = 0.f;
+	};
+
+	/** Skeletal Mesh simulation class that interfaces between the Skeletal Mesh Component and the Cloth Simulation Proxy. */
+	class FSKMClothingSimulation final
+		: public IClothingSimulation
+		, public IClothComponentAdapter
+		, public FGCObject
+	{
+	public:
+		FSKMClothingSimulation();
+		virtual ~FSKMClothingSimulation() override;
+
+		/** Reset all cloth simulation config properties to the values stored in the original cloth asset. Game thread only. */
+		void ResetConfigProperties();
+
+		/** Hard reset the cloth simulation by recreating the proxy. */
+		void RecreateClothSimulationProxy();
+
+		/** Return a property interactor for the specified asset and model. */
+		UChaosClothAssetInteractor* GetPropertyInteractor(const UChaosClothAssetBase* Asset, int32 ModelIndex) const;
+
+		/** Return the visualization object for this simulation. */
+		const ::Chaos::FClothVisualizationNoGC* GetClothVisualization() const
+		{
+			return ClothSimulationProxy.IsValid() ? ClothSimulationProxy->GetClothVisualization() : nullptr;
+		}
+
+	private:
+		//~ Begin IClothingSimulation Interface
+		virtual void Initialize() override;
+		virtual void Shutdown() override;
+
+		virtual IClothingSimulationContext* CreateContext() override
+		{
+			return new FSKMClothingSimulationContext();
+		}
+		virtual void FillContext(const USkeletalMeshComponent* InComponent, float InDeltaTime, IClothingSimulationContext* InOutContext, bool bIsInitialization) const override;
+		virtual void DestroyContext(IClothingSimulationContext* Context) override
+		{
+			delete Context;
+		}
+
+		virtual void CreateActor(USkeletalMeshComponent* InOwnerComponent, const UClothingAssetBase* ClothingAsset, int32 SimDataIndex) override;
+		virtual void EndCreateActor() override;
+		virtual void DestroyActors() override;
+
+		virtual bool ShouldSimulate() const override;
+		virtual void Simulate(const IClothingSimulationContext* Context) override;
+		virtual void GetSimulationData(TMap<int32, FClothSimulData>& OutData, const USkeletalMeshComponent* OwnerComponent, const USkinnedMeshComponent* OverrideComponent) const override;
+
+		virtual FBoxSphereBounds GetBounds(const USkeletalMeshComponent* OwnerComponent) const override;
+
+		virtual void AddExternalCollisions(const FClothCollisionData& Data) override;
+		virtual void ClearExternalCollisions() override;
+		virtual void GetCollisions(FClothCollisionData& OutCollisions, bool bIncludeExternal = true) const override;
+		virtual void SetNumIterations(int32 NumIterations) override;
+		virtual void SetMaxNumIterations(int32 MaxNumIterations) override;
+		virtual void SetNumSubsteps(int32 NumSubsteps) override;
+		virtual int32 GetNumCloths() const override { return ClothSimulationProxy.IsValid() ? ClothSimulationProxy->GetNumCloths() : 0; }
+		virtual int32 GetNumKinematicParticles() const override { return ClothSimulationProxy.IsValid() ? ClothSimulationProxy->GetNumKinematicParticles() : 0; }
+		virtual int32 GetNumDynamicParticles() const override { return ClothSimulationProxy.IsValid() ? ClothSimulationProxy->GetNumDynamicParticles() : 0; }
+		virtual int32 GetNumIterations() const override { return ClothSimulationProxy.IsValid() ? ClothSimulationProxy->GetNumIterations() : 0; }
+		virtual int32 GetNumSubsteps() const override { return ClothSimulationProxy.IsValid() ? ClothSimulationProxy->GetNumSubsteps() : 0; }
+		virtual float GetSimulationTime() const override { return ClothSimulationProxy.IsValid() ? ClothSimulationProxy->GetSimulationTime() : 0.f; }
+		virtual bool IsTeleported() const override { return ClothSimulationProxy.IsValid() ? ClothSimulationProxy->IsTeleported() : false; }
+		//~ End IClothingSimulation Interface
+
+		//~ Begin IClothComponentAdapter Interface
+		virtual const USkinnedMeshComponent& GetOwnerComponent() const override;
+		virtual FCollisionSources& GetCollisionSources() const override
+		{
+			return *CollisionSources;
+		}
+		virtual TArray<const UChaosClothAssetBase*> GetAssets() const override;
+		virtual int32 GetSimulationGroupId(const UChaosClothAssetBase* Asset, int32 ModelIndex) const override;
+
+		virtual const FReferenceSkeleton* GetReferenceSkeleton() const override;
+		virtual EClothingTeleportMode GetClothTeleportMode() const override;
+		virtual bool IsSimulationEnabled() const override;
+		virtual bool IsSimulationSuspended() const override;
+
+		virtual bool NeedsResetRestLengths() const override  // TODO: Make this a property interactor function
+		{
+			return bNeedsResetRestLengths;
+		}
+		virtual const FString& GetRestLengthsMorphTargetName() const override  // TODO: Make this a property interactor function
+		{
+			return RestLengthsMorphTargetName;
+		}
+		virtual float GetClothGeometryScale() const override;
+		virtual const TArray<TSharedPtr<const FManagedArrayCollection>>& GetPropertyCollections(const UChaosClothAssetBase* Asset, int32 ModelIndex) const override;
+		virtual const TArray<TSharedPtr<const FManagedArrayCollection>>& GetSolverPropertyCollections() const override;
+		virtual bool HasAnyRenderClothMappingData(int32 LodIndex) const override;
+		//~ End IClothComponentAdapter Interface
+
+		//~ Begin FGCObject Interface
+		virtual void AddReferencedObjects(FReferenceCollector& Collector) override;
+		virtual FString GetReferencerName() const override;
+		//~ End FGCObject Interface
+
+		const FChaosClothSimulationProperties* GetSolverProperties() const;
+		const TArray<TSharedPtr<::Chaos::Softs::FCollectionPropertyFacade>>& GetSolverPropertyFacades() const;
+
+		USkeletalMeshComponent* OwnerComponent = nullptr;
+		TUniquePtr<FCollisionSources> CollisionSources;
+		TSharedPtr<FClothSimulationProxy> ClothSimulationProxy;
+		
+		using FAssetClothSimulationProperties = TPair<const UChaosClothAssetSKMClothingAsset*, FChaosClothSimulationProperties>;
+		TArray<FAssetClothSimulationProperties> AssetClothSimulationProperties;  // Asset properties, matches CreateActor's SimDataIndex, some assets may appear multiple times, some entries might have a null asset
+	
+		bool bNeedsResetRestLengths = false;
+		FString RestLengthsMorphTargetName;
+	};
+}  // namespace UE::Chaos::ClothAsset
+
+UCLASS(MinimalAPI)
+class UChaosClothAssetSKMClothingSimulationFactory final : public UClothingSimulationFactory
+{
+	GENERATED_BODY()
+
+public:
+	UChaosClothAssetSKMClothingSimulationFactory();
+	virtual ~UChaosClothAssetSKMClothingSimulationFactory() override;
+
+private:
+	//~ Begin UClothingSimulationFactory Interface
+	virtual IClothingSimulation* CreateSimulation() override
+	{
+		return static_cast<IClothingSimulation*>(new UE::Chaos::ClothAsset::FSKMClothingSimulation());
+	}
+
+	virtual void DestroySimulation(IClothingSimulation* Simulation) override
+	{
+		delete Simulation;
+	}
+
+	virtual bool SupportsAsset(const UClothingAssetBase* Asset) const override;
+
+	virtual bool SupportsRuntimeInteraction() const override
+	{
+		return true;
+	}
+
+	virtual UClothingSimulationInteractor* CreateInteractor() override;
+
+	virtual TArrayView<const TSubclassOf<UClothConfigBase>> GetClothConfigClasses() const override
+	{
+		return TArrayView<const TSubclassOf<UClothConfigBase>>();
+	}
+
+	const UEnum* GetWeightMapTargetEnum() const override
+	{
+		return nullptr;
+	}
+	//~ End UClothingSimulationFactory Interface
+};
diff --git a/Engine/Plugins/ChaosClothAsset/Source/ChaosClothAssetEngine/Public/ChaosClothAsset/ClothSimulationProxy.h b/Engine/Plugins/ChaosClothAsset/Source/ChaosClothAssetEngine/Public/ChaosClothAsset/ClothSimulationProxy.h
--- a/Engine/Plugins/ChaosClothAsset/Source/ChaosClothAssetEngine/Public/ChaosClothAsset/ClothSimulationProxy.h
+++ b/Engine/Plugins/ChaosClothAsset/Source/ChaosClothAssetEngine/Public/ChaosClothAsset/ClothSimulationProxy.h
@@ -14,6 +14,10 @@
 #if UE_ENABLE_INCLUDE_ORDER_DEPRECATED_IN_5_7
 class UChaosClothAsset;
 class UChaosClothComponent;
+namespace Chaos
+{
+	class FClothVisualization;
+}
 #endif
 
 namespace Chaos
@@ -23,7 +27,6 @@ namespace Chaos
 	class FClothingSimulationCloth;
 	class FClothingSimulationConfig;
 	class FClothingSimulationCollider;
-	class FClothVisualization;
 	class FClothVisualizationNoGC;
 }
 
@@ -87,11 +90,7 @@ namespace UE::Chaos::ClothAsset
 
 		UE_API FBoxSphereBounds CalculateBounds_AnyThread() const;
 
-		UE_DEPRECATED(5.5, "Use GetClothVisualization instead.")
-PRAGMA_DISABLE_DEPRECATION_WARNINGS
-		const ::Chaos::FClothVisualization* GetVisualization() const { return Visualization.Get(); }
-UE_API PRAGMA_ENABLE_DEPRECATION_WARNINGS
-		const ::Chaos::FClothVisualizationNoGC* GetClothVisualization() const;
+		UE_API const ::Chaos::FClothVisualizationNoGC* GetClothVisualization() const;
 
 		int32 GetNumCloths() const { return NumCloths; }
 		int32 GetNumKinematicParticles() const { return NumKinematicParticles; }
@@ -169,9 +168,8 @@ UE_API PRAGMA_ENABLE_DEPRECATION_WARNINGS
 		TArray<TUniquePtr<::Chaos::FClothingSimulationCloth>> Cloths;
 		TArray<TUniquePtr<::Chaos::FClothingSimulationConfig>> Configs;
 		TArray<TUniquePtr<::Chaos::FClothingSimulationCollider>> Colliders;
-PRAGMA_DISABLE_DEPRECATION_WARNINGS
-		TUniquePtr<::Chaos::FClothVisualization> Visualization;
-PRAGMA_ENABLE_DEPRECATION_WARNINGS
+
+		TUniquePtr<::Chaos::FClothVisualizationNoGC> Visualization;
 
 		// Chaos Cache needs to have access to the solver.
 		friend class FClothComponentCacheAdapter;
diff --git a/Engine/Plugins/ChaosClothAssetEditor/Source/ChaosClothAssetTools/ChaosClothAssetTools.Build.cs b/Engine/Plugins/ChaosClothAssetEditor/Source/ChaosClothAssetTools/ChaosClothAssetTools.Build.cs
--- a/Engine/Plugins/ChaosClothAssetEditor/Source/ChaosClothAssetTools/ChaosClothAssetTools.Build.cs
+++ b/Engine/Plugins/ChaosClothAssetEditor/Source/ChaosClothAssetTools/ChaosClothAssetTools.Build.cs
@@ -22,6 +22,8 @@ public class ChaosClothAssetTools : ModuleRules
 				"RenderCore",
 				"RHI",
 				"Chaos",
+				"ChaosCloth",
+				"ChaosClothEditor",  // For Chaos::FSimulationEditorExtender
 				"ChaosClothAsset",
 				"ChaosClothAssetEngine",
 				"DataflowEngine",
@@ -35,7 +37,7 @@ public class ChaosClothAssetTools : ModuleRules
 				"MeshConversion",
 				"MeshDescription",
 				"SkeletalMeshDescription",		// For FSkeletalMeshAttributes::DefaultSkinWeightProfileName
-				"SkeletalMeshUtilitiesCommon"           
+				"SkeletalMeshUtilitiesCommon"
 			}
 		);
 	}
diff --git a/Engine/Plugins/ChaosClothAssetEditor/Source/ChaosClothAssetTools/Private/ChaosClothAsset/ChaosClothAssetToolsModule.cpp b/Engine/Plugins/ChaosClothAssetEditor/Source/ChaosClothAssetTools/Private/ChaosClothAsset/ChaosClothAssetToolsModule.cpp
--- a/Engine/Plugins/ChaosClothAssetEditor/Source/ChaosClothAssetTools/Private/ChaosClothAsset/ChaosClothAssetToolsModule.cpp
+++ b/Engine/Plugins/ChaosClothAssetEditor/Source/ChaosClothAssetTools/Private/ChaosClothAsset/ChaosClothAssetToolsModule.cpp
@@ -7,6 +7,8 @@
 #include "ChaosClothAsset/ClothComponentEditorStyle.h"
 #include "ChaosClothAsset/ClothingAssetToClothAssetExporter.h"
 #include "ChaosClothAsset/SkeletalMeshConverter.h"
+#include "ChaosClothAsset/ClothAssetSKMSimulationEditorExtender.h"
+#include "ClothingSystemEditorInterfaceModule.h"
 
 namespace UE::Chaos::ClothAsset
 {
@@ -29,6 +31,8 @@ namespace UE::Chaos::ClothAsset
 			IModularFeatures::Get().RegisterModularFeature(IClothingAssetExporterClassProvider::FeatureName, static_cast<IClothingAssetExporterClassProvider*>(this));
 			IModularFeatures::Get().RegisterModularFeature(IClothAssetSkeletalMeshConverterClassProvider::FeatureName, static_cast<IClothAssetSkeletalMeshConverterClassProvider*>(this));
 			IModularFeatures::Get().RegisterModularFeature(IClothingAssetFactoryProvider::FeatureName, static_cast<IClothingAssetFactoryProvider*>(this));
+			IModularFeatures::Get().RegisterModularFeature(FClothingSystemEditorInterfaceModule::ExtenderFeatureName, &SKMSimulationEditorExtender);
+
 		}
 
 		virtual void ShutdownModule() override
@@ -39,6 +43,7 @@ namespace UE::Chaos::ClothAsset
 				IModularFeatures::Get().UnregisterModularFeature(IClothingAssetExporterClassProvider::FeatureName, static_cast<IClothingAssetExporterClassProvider*>(this));
 				IModularFeatures::Get().UnregisterModularFeature(IClothAssetSkeletalMeshConverterClassProvider::FeatureName, static_cast<IClothAssetSkeletalMeshConverterClassProvider*>(this));
 				IModularFeatures::Get().UnregisterModularFeature(IClothingAssetFactoryProvider::FeatureName, static_cast<IClothingAssetFactoryProvider*>(this));
+				IModularFeatures::Get().UnregisterModularFeature(FClothingSystemEditorInterfaceModule::ExtenderFeatureName, &SKMSimulationEditorExtender);
 			}
 		}
 
@@ -62,6 +67,9 @@ namespace UE::Chaos::ClothAsset
 			return UChaosClothAssetSKMClothingAssetFactory::StaticClass()->GetDefaultObject<UClothingAssetFactoryBase>();
 		}
 		//~ End IClothingAssetFactoryProvider
+
+	private:
+		UE::Chaos::ClothAsset::FSKMSimulationEditorExtender SKMSimulationEditorExtender;
 	};
 }
 
diff --git a/Engine/Plugins/ChaosClothAssetEditor/Source/ChaosClothAssetTools/Private/ChaosClothAsset/ClothAssetSKMSimulationEditorExtender.cpp b/Engine/Plugins/ChaosClothAssetEditor/Source/ChaosClothAssetTools/Private/ChaosClothAsset/ClothAssetSKMSimulationEditorExtender.cpp
new file mode 100644
--- /dev/null
+++ b/Engine/Plugins/ChaosClothAssetEditor/Source/ChaosClothAssetTools/Private/ChaosClothAsset/ClothAssetSKMSimulationEditorExtender.cpp
@@ -0,0 +1,38 @@
+// Copyright Epic Games, Inc. All Rights Reserved.
+
+#include "ChaosClothAsset/ClothAssetSKMSimulationEditorExtender.h"
+#include "ChaosClothAsset/ClothAssetSKMClothingSimulation.h"
+
+#define LOCTEXT_NAMESPACE "ChaosSimulationEditorExtender"
+
+namespace UE::Chaos::ClothAsset
+{
+	FSKMSimulationEditorExtender::FSKMSimulationEditorExtender()
+		: ::Chaos::FSimulationEditorExtender()
+	{
+	}
+
+	FSKMSimulationEditorExtender::~FSKMSimulationEditorExtender() = default;
+
+	UClass* FSKMSimulationEditorExtender::GetSupportedSimulationFactoryClass()
+	{
+		return UChaosClothAssetSKMClothingSimulationFactory::StaticClass();
+	}
+
+	FName FSKMSimulationEditorExtender::GetSectionName() const
+	{
+		return FName("ClothAsset_Visualizations");
+	}
+
+	FText FSKMSimulationEditorExtender::GetSectionHeadingText() const
+	{
+		return FText(LOCTEXT("ClothAssetVisualizationSection", "Cloth Asset Visualization"));
+	}
+
+	const ::Chaos::FClothVisualizationNoGC* FSKMSimulationEditorExtender::GetClothVisualization(const IClothingSimulation* Simulation) const
+	{
+		return Simulation ? static_cast<const FSKMClothingSimulation*>(Simulation)->GetClothVisualization() : nullptr;
+	}
+}  // namespace UE::Chaos::ClothAsset
+
+#undef LOCTEXT_NAMESPACE
diff --git a/Engine/Plugins/ChaosClothAssetEditor/Source/ChaosClothAssetTools/Private/ChaosClothAsset/ClothAssetSKMSimulationEditorExtender.h b/Engine/Plugins/ChaosClothAssetEditor/Source/ChaosClothAssetTools/Private/ChaosClothAsset/ClothAssetSKMSimulationEditorExtender.h
new file mode 100644
--- /dev/null
+++ b/Engine/Plugins/ChaosClothAssetEditor/Source/ChaosClothAssetTools/Private/ChaosClothAsset/ClothAssetSKMSimulationEditorExtender.h
@@ -0,0 +1,32 @@
+// Copyright Epic Games, Inc. All Rights Reserved.
+
+#pragma once
+
+#include "ChaosClothEditor/ChaosSimulationEditorExtender.h"
+
+namespace Chaos
+{
+	class FClothVisualizationNoGC;
+}
+
+namespace UE::Chaos::ClothAsset
+{
+	/** Chaos extension to the asset editor. */
+	class FSKMSimulationEditorExtender final : public ::Chaos::FSimulationEditorExtender
+	{
+	public:
+		FSKMSimulationEditorExtender();
+		virtual ~FSKMSimulationEditorExtender() override;
+
+	private:
+		//~ Begin ISimulationEditorExtender Interface
+		virtual UClass* GetSupportedSimulationFactoryClass() override;
+		//~ End ISimulationEditorExtender Interface
+
+		//~ Begin ::Chaos::FSimulationEditorExtender Interface
+		virtual FName GetSectionName() const override;
+		virtual FText GetSectionHeadingText() const override;
+		virtual const ::Chaos::FClothVisualizationNoGC* GetClothVisualization(const IClothingSimulation* Simulation) const override;
+		//~ End ::Chaos::FSimulationEditorExtender Interface
+	};
+}
