Cloth - Fixed some uninitialized function pointers in the editor extender visualization options.

#rnx

diff --git a/Engine/Plugins/ChaosCloth/Source/ChaosClothEditor/Private/ChaosClothEditor/ChaosSimulationEditorExtender.cpp b/Engine/Plugins/ChaosCloth/Source/ChaosClothEditor/Private/ChaosClothEditor/ChaosSimulationEditorExtender.cpp
--- a/Engine/Plugins/ChaosCloth/Source/ChaosClothEditor/Private/ChaosClothEditor/ChaosSimulationEditorExtender.cpp
+++ b/Engine/Plugins/ChaosCloth/Source/ChaosClothEditor/Private/ChaosClothEditor/ChaosSimulationEditorExtender.cpp
@@ -37,6 +37,7 @@ struct FVisualizationOption
 
 	FVisualizationOption(FDebugDrawFunction InDebugDrawFunction, const TCHAR* ConsoleName, const FText& InDisplayName, const FText& InToolTip, bool bInDisablesSimulation = false, bool bInHidesClothSections = false)
 		: DebugDrawFunction(InDebugDrawFunction)
+		, DebugDrawTextsFunction(nullptr)
 		, DisplayName(InDisplayName)
 		, ToolTip(InToolTip)
 		, bDisablesSimulation(bInDisablesSimulation)
@@ -45,7 +46,8 @@ struct FVisualizationOption
 	{}
 
 	FVisualizationOption(FDebugDrawTextsFunction InDebugDrawTextsFunction, const TCHAR* ConsoleName, const FText& InDisplayName, const FText& InToolTip, bool bInDisablesSimulation = false, bool bInHidesClothSections = false)
-		: DebugDrawTextsFunction(InDebugDrawTextsFunction)
+		: DebugDrawFunction(nullptr)
+		, DebugDrawTextsFunction(InDebugDrawTextsFunction)
 		, DisplayName(InDisplayName)
 		, ToolTip(InToolTip)
 		, bDisablesSimulation(bInDisablesSimulation)
