Update Chaos Cloth to better handle non-simulating LODs.

* Fix an issue where ChaosClothComponent-based cloth would not correctly enable/disable particles when going from LODA -> LODB -> LODC when LODB does not have simulation data but A and C do. This was because ClothingSimulationCloth would not Update when going between unsimulated and simulated LODs.
* Switch ChaosClothComponent-based cloth to use the existence of simulation model data rather than cloth-bound render sections to determine if a LOD needs to simulate. These two conditions should be equivalent, but the solver will only LOD-switch correctly if we skip simulation when simulation model data is missing.
* Update SKM-based cloth to have a concept of whether or not an LOD should simulate. Previously, if any cloth LODs were present, all LODs would do the work to simulate (FillContext, run a Simulate task). Unlike ChaosClothComponent-based cloth, use render mesh being bound to cloth to determine whether to simulate. It's very easy and likely to create LODs of SKM-based cloth that will not render, and we don't want to pay for it. To ensure LODs switch correctly, manually force the cloth to teleport and reset when it is not simulated.
* Reworked GetSimulationData to be AppendSimulationData so multiple types of cloth can run at the same time.
* Fixed FSKMClothingSimulation to call GameThread-only methods from the GameThread.

diff --git a/Engine/Plugins/ChaosCloth/Source/ChaosCloth/Private/ChaosCloth/ChaosClothingSimulation.cpp b/Engine/Plugins/ChaosCloth/Source/ChaosCloth/Private/ChaosCloth/ChaosClothingSimulation.cpp
--- a/Engine/Plugins/ChaosCloth/Source/ChaosCloth/Private/ChaosCloth/ChaosClothingSimulation.cpp
+++ b/Engine/Plugins/ChaosCloth/Source/ChaosCloth/Private/ChaosCloth/ChaosClothingSimulation.cpp
@@ -17,6 +17,7 @@
 #include "PhysicsProxy/PerSolverFieldSystem.h"
 
 #include "HAL/IConsoleManager.h"
+#include "Rendering/SkeletalMeshRenderData.h"
 #include "PrimitiveDrawingUtils.h"
 
 #if !UE_BUILD_SHIPPING
@@ -411,6 +412,47 @@ void FClothingSimulation::CreateActor(USkeletalMeshComponent* InOwnerComponent,
 	UE_LOG(LogChaosCloth, Verbose, TEXT("Added Cloth asset to %s in sim slot %d"), InOwnerComponent->GetOwner() ? *InOwnerComponent->GetOwner()->GetName() : TEXT("None"), InSimDataIndex);
 }
 
+void FClothingSimulation::EndCreateActor()
+{
+	LODHasAnyRenderClothMappingData.Reset();
+	for (const TUniquePtr<FClothingSimulationCloth>& Cloth : Cloths)
+	{
+		const int32 AssetIndex = Cloth->GetGroupId();
+		if (const FClothingSimulationSkeletalMesh* const SimulationMesh = static_cast<const FClothingSimulationSkeletalMesh*>(Cloth->GetMesh()))
+		{
+			if (const USkeletalMeshComponent* const OwnerComponent = SimulationMesh->GetSkeletalMeshComponent())
+			{
+				if (const USkeletalMesh* const SkeletalMesh = OwnerComponent->GetSkeletalMeshAsset())
+				{
+					if (const FSkeletalMeshRenderData* const SkeletalMeshRenderData = SkeletalMesh->GetResourceForRendering())
+					{
+						if (LODHasAnyRenderClothMappingData.Num() < SkeletalMeshRenderData->LODRenderData.Num())
+						{
+							LODHasAnyRenderClothMappingData.SetNum(SkeletalMeshRenderData->LODRenderData.Num(), false);
+						}
+						for (int32 LodIndex = 0; LodIndex < SkeletalMeshRenderData->LODRenderData.Num(); ++LodIndex)
+						{
+							if (LODHasAnyRenderClothMappingData[LodIndex])
+							{
+								continue;
+							}
+							for (const FSkelMeshRenderSection& SkelMeshRenderSection : SkeletalMeshRenderData->LODRenderData[LodIndex].RenderSections)
+							{
+								if (SkelMeshRenderSection.HasClothingData() &&
+									SkelMeshRenderSection.CorrespondClothAssetIndex == AssetIndex)
+								{
+									LODHasAnyRenderClothMappingData[LodIndex] = true;
+									continue;
+								}
+							}
+						}
+					}
+				}
+			}
+		}
+	}
+}
+
 void FClothingSimulation::UpdateWorldForces(const USkeletalMeshComponent* OwnerComponent)
 {
 	if (OwnerComponent)
@@ -483,18 +525,46 @@ void FClothingSimulation::SetNumSubsteps(int32 InNumSubsteps)
 	Solver->GetConfig()->GetProperties(Solver->GetSolverLOD()).SetValue(FClothingSimulationSolver::NumSubstepsName, InNumSubsteps);
 }
 
-bool FClothingSimulation::ShouldSimulate() const
+void FClothingSimulation::FillContextAndPrepareTick(const USkeletalMeshComponent* InComponent, float InDeltaTime, IClothingSimulationContext* InOutContext, bool bIsInitialization, bool bForceTeleportResetOnly)
 {
+	if (bForceTeleportResetOnly)
+	{
+		// Nothing to do in this case.
+		return;
+	}
+
+	FClothingSimulationCommon::FillContextAndPrepareTick(InComponent, InDeltaTime, InOutContext, bIsInitialization, bForceTeleportResetOnly);
+}
+
+bool FClothingSimulation::ShouldSimulateLOD(int32 OwnerLODIndex) const
+{
+	if (!LODHasAnyRenderClothMappingData.IsValidIndex(OwnerLODIndex) || !LODHasAnyRenderClothMappingData[OwnerLODIndex])
+	{
+		return false;
+	}
+
 	for (const TUniquePtr<FClothingSimulationCloth>& Cloth : Cloths)
 	{
-		if (Cloth->GetLODIndex(Solver.Get()) != INDEX_NONE && Cloth->GetParticleRangeId(Solver.Get()) != INDEX_NONE)
+		if (const FClothingSimulationMesh* const SimulationMesh = Cloth->GetMesh())
 		{
-			return true;
+			const int32 ClothLODIndex = SimulationMesh->GetLODIndexFromOwnerLODIndex(OwnerLODIndex);
+			if (SimulationMesh->IsValidLODIndex(ClothLODIndex) && SimulationMesh->GetNumPoints(ClothLODIndex) > 0)
+			{
+				return true;
+			}
 		}
 	}
 	return false;
 }
 
+void FClothingSimulation::ForceClothNextUpdateTeleportAndReset()
+{
+	for (const TUniquePtr<FClothingSimulationCloth>& Cloth : Cloths)
+	{
+		Cloth->Reset();
+	}
+}
+
 void FClothingSimulation::Simulate(const IClothingSimulationContext* InContext)
 {
 	TRACE_CPUPROFILER_EVENT_SCOPE(FClothingSimulation_Simulate);
@@ -611,7 +681,7 @@ void FClothingSimulation::Simulate(const IClothingSimulationContext* InContext)
 #endif  // #if CHAOS_DEBUG_DRAW
 }
 
-void FClothingSimulation::GetSimulationData(
+void FClothingSimulation::AppendSimulationData(
 	TMap<int32, FClothSimulData>& OutData,
 	const USkeletalMeshComponent* InOwnerComponent,
 	const USkinnedMeshComponent* InOverrideComponent) const
@@ -619,17 +689,6 @@ void FClothingSimulation::GetSimulationData(
 	TRACE_CPUPROFILER_EVENT_SCOPE(FClothingSimulation_GetSimulationData);
 	SCOPE_CYCLE_COUNTER(STAT_ChaosClothGetSimulationData);
 
-	if (!Cloths.Num() || !InOwnerComponent)
-	{
-		OutData.Reset();
-		return;
-	}
-
-	// Reset map when new cloths have appeared
-	if (OutData.Num() != Cloths.Num())
-	{
-		OutData.Reset();
-	}
 
 	// Get the solver's local space
 	const FVec3& LocalSpaceLocation = Solver->GetLocalSpaceLocation(); // Note: Since the ReferenceSpaceTransform can be suspended with the simulation, it is important that the suspended local space location is used too in order to get the simulation data back into reference space
@@ -672,7 +731,6 @@ void FClothingSimulation::GetSimulationData(
 		{
 			UE_CLOG(!bHasInvalidReferenceBoneTransforms, LogSkeletalMesh, Warning, TEXT("Failed to write back clothing simulation data for component %s as bone transforms are invalid."), *InOwnerComponent->GetName());
 			bHasInvalidReferenceBoneTransforms = true;
-			OutData.Reset();
 			return;
 		}
 
diff --git a/Engine/Plugins/ChaosCloth/Source/ChaosCloth/Private/ChaosCloth/ChaosClothingSimulationCloth.cpp b/Engine/Plugins/ChaosCloth/Source/ChaosCloth/Private/ChaosCloth/ChaosClothingSimulationCloth.cpp
--- a/Engine/Plugins/ChaosCloth/Source/ChaosCloth/Private/ChaosCloth/ChaosClothingSimulationCloth.cpp
+++ b/Engine/Plugins/ChaosCloth/Source/ChaosCloth/Private/ChaosCloth/ChaosClothingSimulationCloth.cpp
@@ -423,28 +423,32 @@ void FClothingSimulationCloth::FLODData::Update(FClothingSimulationSolver* Solve
 	const FSolverData& SolverDatum = SolverData.FindChecked(Solver);
 
 	const int32 ParticleRangeId = SolverDatum.ParticleRangeId;
-	check(ParticleRangeId != INDEX_NONE);
 
-	// Update the animatable constraint parameters
-	FClothConstraints& ClothConstraints = Solver->GetClothConstraints(ParticleRangeId);
+	if (ParticleRangeId != INDEX_NONE)
+	{
+		// Update the animatable constraint parameters
+		FClothConstraints& ClothConstraints = Solver->GetClothConstraints(ParticleRangeId);
 
-	check(Cloth->Config);
-	const FReal LocalSpaceScale = Solver->GetLocalSpaceScale();
-	check(LocalSpaceScale > UE_SMALL_NUMBER);
-	const FReal LocalSpaceScaleInv = 1. / LocalSpaceScale;
-	const Softs::FSolverReal MeshScale = Cloth->Mesh->GetScale() * LocalSpaceScaleInv;
-	const Softs::FSolverReal MaxDistancesScale = (Softs::FSolverReal)Cloth->MaxDistancesMultiplier;
-	ClothConstraints.Update(Cloth->Config->GetProperties(SolverDatum.LODIndex), WeightMaps, VertexSets, FaceSets, FaceIntMaps, MeshScale, MaxDistancesScale, Solver->GetLocalSpaceRotation(), Cloth->ReferenceSpaceTransform.GetRotation());
+		check(Cloth->Config);
+		const FReal LocalSpaceScale = Solver->GetLocalSpaceScale();
+		check(LocalSpaceScale > UE_SMALL_NUMBER);
+		const FReal LocalSpaceScaleInv = 1. / LocalSpaceScale;
+		const Softs::FSolverReal MeshScale = Cloth->Mesh->GetScale() * LocalSpaceScaleInv;
+		const Softs::FSolverReal MaxDistancesScale = (Softs::FSolverReal)Cloth->MaxDistancesMultiplier;
+		ClothConstraints.Update(Cloth->Config->GetProperties(SolverDatum.LODIndex), WeightMaps, VertexSets, FaceSets, FaceIntMaps, MeshScale, MaxDistancesScale, Solver->GetLocalSpaceRotation(), Cloth->ReferenceSpaceTransform.GetRotation());
+	}
 }
 
 void FClothingSimulationCloth::FLODData::Enable(FClothingSimulationSolver* Solver, bool bEnable) const
 {
 	check(Solver);
 	const int32 ParticleRangeId = SolverData.FindChecked(Solver).ParticleRangeId;
-	check(ParticleRangeId != INDEX_NONE);
-	
-	// Enable particles (and related constraints)
-	Solver->EnableParticles(ParticleRangeId, bEnable);
+
+	if (ParticleRangeId != INDEX_NONE)
+	{
+		// Enable particles (and related constraints)
+		Solver->EnableParticles(ParticleRangeId, bEnable);
+	}
 }
 
 void FClothingSimulationCloth::FLODData::ResetStartPose(FClothingSimulationSolver* Solver)
@@ -453,10 +457,11 @@ void FClothingSimulationCloth::FLODData::ResetStartPose(FClothingSimulationSolve
 
 	FSolverData& SolverDatum = SolverData.FindChecked(Solver);
 	const int32 ParticleRangeId = SolverDatum.ParticleRangeId;
-	check(ParticleRangeId != INDEX_NONE);
-
-	Solver->ResetStartPose(ParticleRangeId, NumParticles);
-	ResetAccessoryMeshes(Solver);
+	if (ParticleRangeId != INDEX_NONE)
+	{
+		Solver->ResetStartPose(ParticleRangeId, NumParticles);
+		ResetAccessoryMeshes(Solver);
+	}
 }
 
 void FClothingSimulationCloth::FLODData::UpdateAccessoryMeshes(FClothingSimulationSolver* Solver)
@@ -761,7 +766,7 @@ int32 FClothingSimulationCloth::GetNumParticles(int32 InLODIndex) const
 
 int32 FClothingSimulationCloth::GetParticleRangeId(const FClothingSimulationSolver* Solver, int32 InLODIndex) const
 {
-	return LODData.IsValidIndex(InLODIndex) ? LODData[InLODIndex]->SolverData.FindChecked(Solver).ParticleRangeId : 0;
+	return LODData.IsValidIndex(InLODIndex) ? LODData[InLODIndex]->SolverData.FindChecked(Solver).ParticleRangeId : INDEX_NONE;
 }
 
 TVec3<FRealSingle> FClothingSimulationCloth::GetGravity(const FClothingSimulationSolver* Solver) const
@@ -973,19 +978,19 @@ void FClothingSimulationCloth::Update(FClothingSimulationSolver* Solver)
 	// Update the source mesh skinned positions
 	const int32 PrevParticleRangeId = GetParticleRangeId(Solver, PrevLODIndex);
 	const int32 ParticleRangeId = GetParticleRangeId(Solver, LODIndex);
-	if (PrevParticleRangeId == INDEX_NONE || ParticleRangeId == INDEX_NONE)
+	if (PrevParticleRangeId == INDEX_NONE && ParticleRangeId == INDEX_NONE)
 	{
 		return;
 	}
 
 	// Retrieve config
 	check(Config);
-	const Softs::FCollectionPropertyFacade& ConfigProperties = Config->GetProperties(LODIndex);
+	const Softs::FCollectionPropertyFacade* const ConfigProperties = LODIndex == INDEX_NONE ? nullptr : &Config->GetProperties(LODIndex);
 	int32 CurrentLODMorphTargetIndex = INDEX_NONE;
 	float CurrentLODMorphTargetWeight = 0.f;
 	if (LODIndex != INDEX_NONE)
 	{
-		LODData[LODIndex]->UpdateCachedProperties(*Mesh, LODIndex, ConfigProperties);
+		LODData[LODIndex]->UpdateCachedProperties(*Mesh, LODIndex, *ConfigProperties);
 		CurrentLODMorphTargetIndex = LODData[LODIndex]->ActiveMorphTarget;
 		CurrentLODMorphTargetWeight = LODData[LODIndex]->ActiveMorphTargetWeight;
 		LODData[LODIndex]->UpdateAccessoryMeshes(Solver);
@@ -995,7 +1000,7 @@ void FClothingSimulationCloth::Update(FClothingSimulationSolver* Solver)
 
 	const int32 CoarseLODIndex = LODIndex != INDEX_NONE ? LODData[LODIndex]->SolverData.FindChecked(Solver).MultiResCoarseLODIndex : INDEX_NONE;
 	const int32 CoarseParticleRangeId = CoarseLODIndex != INDEX_NONE ? GetParticleRangeId(Solver, CoarseLODIndex) : INDEX_NONE;
-	if (CoarseLODIndex != INDEX_NONE)
+	if (CoarseParticleRangeId != INDEX_NONE)
 	{
 		LODData[CoarseLODIndex]->UpdateCachedProperties(*Mesh, CoarseLODIndex, Config->GetProperties(CoarseLODIndex));
 		LODData[CoarseLODIndex]->Enable(Solver, true);
@@ -1029,15 +1034,16 @@ void FClothingSimulationCloth::Update(FClothingSimulationSolver* Solver)
 			NumActiveDynamicParticles = LODData[LODIndex]->NumDynamicParticles;
 
 			// Wrap new LOD based on previous LOD if possible (can only do 1 level LOD at a time, and if previous LOD exists)
-			bNeedsReset = bNeedsReset || !Mesh->WrapDeformLOD(
-				PrevLODIndex,
-				LODIndex,
-				Solver->GetNormals(PrevParticleRangeId),
-				Solver->GetParticlePandInvMs(PrevParticleRangeId),
-				Solver->GetParticleVs(PrevParticleRangeId),
-				Solver->GetParticlePandInvMs(ParticleRangeId),
-				Solver->GetParticleXs(ParticleRangeId),
-				Solver->GetParticleVs(ParticleRangeId));
+			bNeedsReset = bNeedsReset || ParticleRangeId == INDEX_NONE || PrevParticleRangeId == INDEX_NONE ||
+				!Mesh->WrapDeformLOD(
+					PrevLODIndex,
+					LODIndex,
+					Solver->GetNormals(PrevParticleRangeId),
+					Solver->GetParticlePandInvMs(PrevParticleRangeId),
+					Solver->GetParticleVs(PrevParticleRangeId),
+					Solver->GetParticlePandInvMs(ParticleRangeId),
+					Solver->GetParticleXs(ParticleRangeId),
+					Solver->GetParticleVs(ParticleRangeId));
 
 			if (!bNeedsReset)
 			{
@@ -1047,7 +1053,7 @@ void FClothingSimulationCloth::Update(FClothingSimulationSolver* Solver)
 			if (Solver->IsLegacySolver())
 			{
 				// Update the wind velocity field for the new LOD mesh
-				Solver->SetWindAndPressureGeometry(GroupId, GetTriangleMesh(Solver), ConfigProperties, LODData[LODIndex]->WeightMaps);
+				Solver->SetWindAndPressureGeometry(GroupId, GetTriangleMesh(Solver), *ConfigProperties, LODData[LODIndex]->WeightMaps);
 			}
 		}
 		else
@@ -1058,7 +1064,7 @@ void FClothingSimulationCloth::Update(FClothingSimulationSolver* Solver)
 	}
 
 	// Update Cloth group parameters  TODO: Cloth groups should exist as their own node object so that they can be used by several cloth objects
-	if (LODIndex != INDEX_NONE)
+	if (LODIndex != INDEX_NONE && ParticleRangeId != INDEX_NONE)
 	{
 		// TODO: Move all groupID updates out of the cloth update to allow to use of the same GroupId with different cloths
 
@@ -1090,7 +1096,7 @@ void FClothingSimulationCloth::Update(FClothingSimulationSolver* Solver)
 				Collider->ResetStartPose(Solver, this);
 			}
 
-			if (CoarseLODIndex != INDEX_NONE)
+			if (CoarseLODIndex != INDEX_NONE && CoarseParticleRangeId != INDEX_NONE)
 			{
 				LODData[CoarseLODIndex]->ResetStartPose(Solver);
 			}
@@ -1110,18 +1116,18 @@ void FClothingSimulationCloth::Update(FClothingSimulationSolver* Solver)
 		else
 		{
 			// Use the cloth config parameters
-			VelocityScaleSpace = (EChaosSoftsSimulationSpace)GetVelocityScaleSpace(ConfigProperties, (int32)ClothingSimulationClothDefault::VelocityScaleSpace);
-			OutLinearVelocityScale = GetLinearVelocityScale(ConfigProperties, FVector3f(ClothingSimulationClothDefault::VelocityScale));
-			OutAngularVelocityScale = GetAngularVelocityScale(ConfigProperties, ClothingSimulationClothDefault::VelocityScale);
-			OutMaxVelocityScale = GetMaxVelocityScale(ConfigProperties, ClothingSimulationClothDefault::MaxVelocityScale);
-			MaxLinearVelocity = GetMaxLinearVelocity(ConfigProperties, MaxLinearVelocity);
-			MaxLinearAcceleration = GetMaxLinearAcceleration(ConfigProperties, MaxLinearAcceleration);
-			MaxAngularVelocity = GetMaxAngularVelocity(ConfigProperties, MaxAngularVelocity);
-			MaxAngularAcceleration = GetMaxAngularAcceleration(ConfigProperties, MaxAngularAcceleration);
+			VelocityScaleSpace = (EChaosSoftsSimulationSpace)GetVelocityScaleSpace(*ConfigProperties, (int32)ClothingSimulationClothDefault::VelocityScaleSpace);
+			OutLinearVelocityScale = GetLinearVelocityScale(*ConfigProperties, FVector3f(ClothingSimulationClothDefault::VelocityScale));
+			OutAngularVelocityScale = GetAngularVelocityScale(*ConfigProperties, ClothingSimulationClothDefault::VelocityScale);
+			OutMaxVelocityScale = GetMaxVelocityScale(*ConfigProperties, ClothingSimulationClothDefault::MaxVelocityScale);
+			MaxLinearVelocity = GetMaxLinearVelocity(*ConfigProperties, MaxLinearVelocity);
+			MaxLinearAcceleration = GetMaxLinearAcceleration(*ConfigProperties, MaxLinearAcceleration);
+			MaxAngularVelocity = GetMaxAngularVelocity(*ConfigProperties, MaxAngularVelocity);
+			MaxAngularAcceleration = GetMaxAngularAcceleration(*ConfigProperties, MaxAngularAcceleration);
 		}
 
 		// NOTE: Force-based solver doesn't actually use FictitiousAngularScale here. It gets it from the property collection directly.
-		const FRealSingle FictitiousAngularScale = GetFictitiousAngularScale(ConfigProperties, ClothingSimulationClothDefault::FictitiousAngularScale);
+		const FRealSingle FictitiousAngularScale = GetFictitiousAngularScale(*ConfigProperties, ClothingSimulationClothDefault::FictitiousAngularScale);
 		Solver->SetReferenceVelocityScale(
 			GroupId,
 			OldReferenceSpaceTransform,
@@ -1147,8 +1153,8 @@ void FClothingSimulationCloth::Update(FClothingSimulationSolver* Solver)
 		}
 		if (!Solver->IsLegacySolver())
 		{
-			Solver->SetProperties(ParticleRangeId, ConfigProperties, LODData[LODIndex]->WeightMaps);
-			if (CoarseLODIndex != INDEX_NONE)
+			Solver->SetProperties(ParticleRangeId, *ConfigProperties, LODData[LODIndex]->WeightMaps);
+			if (CoarseLODIndex != INDEX_NONE && CoarseParticleRangeId != INDEX_NONE)
 			{
 				Solver->SetProperties(CoarseParticleRangeId, Config->GetProperties(CoarseLODIndex), LODData[CoarseLODIndex]->WeightMaps);
 			}
@@ -1162,15 +1168,15 @@ void FClothingSimulationCloth::Update(FClothingSimulationSolver* Solver)
 			Solver->SetGravity(GroupId, GetGravity(Solver));
 
 			// External forces (legacy wind+field)
-			const bool bUsePointBasedWindModel = GetUsePointBasedWindModel(ConfigProperties, false);
+			const bool bUsePointBasedWindModel = GetUsePointBasedWindModel(*ConfigProperties, false);
 			Solver->AddExternalForces(GroupId, bUsePointBasedWindModel);
 
 			const bool bPointBasedWindDisablesAccurateWind = ClothingSimulationClothConsoleVariables::CVarLegacyDisablesAccurateWind.GetValueOnAnyThread();
 			const bool bEnableAerodynamics = !(bUsePointBasedWindModel && bPointBasedWindDisablesAccurateWind);
-			Solver->SetWindAndPressureProperties(GroupId, ConfigProperties, LODData[LODIndex]->WeightMaps, bEnableAerodynamics);
+			Solver->SetWindAndPressureProperties(GroupId, *ConfigProperties, LODData[LODIndex]->WeightMaps, bEnableAerodynamics);
 
 			constexpr float WorldScale = 100.f;  // VelocityField wind is in m/s in the config (same as the wind unit), but cm/s in the solver  TODO: Cleanup the Solver SetWindVelocity functions to be consistent with the unit
-			const FVec3f WindVelocity = Softs::FVelocityAndPressureField::GetWindVelocity(ConfigProperties, FVector3f(0.f)) * WorldScale;
+			const FVec3f WindVelocity = Softs::FVelocityAndPressureField::GetWindVelocity(*ConfigProperties, FVector3f(0.f)) * WorldScale;
 			Solver->SetWindVelocity(GroupId, WindVelocity + Solver->GetWindVelocity());
 
 			// Update general solver properties
@@ -1179,14 +1185,14 @@ void FClothingSimulationCloth::Update(FClothingSimulationSolver* Solver)
 			const FReal LocalSpaceScaleInv = 1. / LocalSpaceScale;
 			const Softs::FSolverReal MeshScale = Mesh->GetScale() * LocalSpaceScaleInv;
 
-			const FRealSingle DampingCoefficient = GetDampingCoefficient(ConfigProperties, ClothingSimulationClothDefault::DampingCoefficient);
-			const FRealSingle LocalDampingCoefficient = GetLocalDampingCoefficient(ConfigProperties, ClothingSimulationClothDefault::LocalDampingCoefficient);
-			const FRealSingle CollisionThickness = GetCollisionThickness(ConfigProperties, ClothingSimulationClothDefault::CollisionThickness);
-			const FRealSingle FrictionCoefficient = GetFrictionCoefficient(ConfigProperties, ClothingSimulationClothDefault::FrictionCoefficient);
+			const FRealSingle DampingCoefficient = GetDampingCoefficient(*ConfigProperties, ClothingSimulationClothDefault::DampingCoefficient);
+			const FRealSingle LocalDampingCoefficient = GetLocalDampingCoefficient(*ConfigProperties, ClothingSimulationClothDefault::LocalDampingCoefficient);
+			const FRealSingle CollisionThickness = GetCollisionThickness(*ConfigProperties, ClothingSimulationClothDefault::CollisionThickness);
+			const FRealSingle FrictionCoefficient = GetFrictionCoefficient(*ConfigProperties, ClothingSimulationClothDefault::FrictionCoefficient);
 			Solver->SetProperties(GroupId, DampingCoefficient, LocalDampingCoefficient, CollisionThickness * MeshScale, FrictionCoefficient);
 
 			// Update use of continuous collision detection
-			const bool bUseCCD = GetUseCCD(ConfigProperties, false);
+			const bool bUseCCD = GetUseCCD(*ConfigProperties, false);
 			Solver->SetUseCCD(GroupId, bUseCCD);
 		}
 
diff --git a/Engine/Plugins/ChaosCloth/Source/ChaosCloth/Private/ChaosCloth/ChaosClothingSimulationMesh.cpp b/Engine/Plugins/ChaosCloth/Source/ChaosCloth/Private/ChaosCloth/ChaosClothingSimulationMesh.cpp
--- a/Engine/Plugins/ChaosCloth/Source/ChaosCloth/Private/ChaosCloth/ChaosClothingSimulationMesh.cpp
+++ b/Engine/Plugins/ChaosCloth/Source/ChaosCloth/Private/ChaosCloth/ChaosClothingSimulationMesh.cpp
@@ -55,6 +55,11 @@ bool FClothingSimulationMesh::WrapDeformLOD(
 		GetTransitionUpSkinData(LODIndex) :
 		GetTransitionDownSkinData(LODIndex);
 
+	if (SkinData.Num() != NumPoints)
+	{
+		return false;
+	}
+
 	for (int32 Index = 0; Index < NumPoints; ++Index)  // TODO: Profile for parallel for
 	{
 		const FMeshToMeshVertData& VertData = SkinData[Index];
@@ -96,6 +101,11 @@ bool FClothingSimulationMesh::WrapDeformLOD(
 		GetTransitionUpSkinData(LODIndex) :
 		GetTransitionDownSkinData(LODIndex);
 
+	if (SkinData.Num() != NumPoints)
+	{
+		return false;
+	}
+
 	for (int32 Index = 0; Index < NumPoints; ++Index)  // TODO: Profile for parallel for
 	{
 		const FMeshToMeshVertData& VertData = SkinData[Index];
@@ -142,6 +152,11 @@ bool FClothingSimulationMesh::WrapDeformLOD(
 		GetTransitionUpSkinData(LODIndex) :
 		GetTransitionDownSkinData(LODIndex);
 
+	if (SkinData.Num() != NumPoints)
+	{
+		return false;
+	}
+
 	for (int32 Index = 0; Index < NumPoints; ++Index)  // TODO: Profile for parallel for
 	{
 		const FMeshToMeshVertData& VertData = SkinData[Index];
@@ -194,7 +209,7 @@ void FClothingSimulationMesh::Update(
 	check(Solver);
 
 	// Exit if any inputs are missing or not ready, and if the LOD is invalid
-	if (!IsValidLODIndex(LODIndex))
+	if (!IsValidLODIndex(LODIndex) || GetNumPoints(LODIndex) == 0)
 	{
 		return;
 	}
@@ -208,13 +223,17 @@ void FClothingSimulationMesh::Update(
 	// Update old positions after LOD Switching
 	if (LODIndex != PrevLODIndex)
 	{
-		// TODO: Using the more accurate skinning method here would require double buffering the context at the skeletal mesh level
-		const TConstArrayView<FSolverVec3> SrcWrapPositions = Solver->GetOldAnimationPositionsView(PrevParticleRangeId);
-		const TConstArrayView<FSolverVec3> SrcWrapNormals = Solver->GetOldAnimationNormalsView(PrevParticleRangeId);
+		bool bValidWrap = false;
 		TArrayView<FSolverVec3> OutOldPositions = Solver->GetOldAnimationPositionsView(ParticleRangeId);
 		TArrayView<FSolverVec3> OutOldNormals = Solver->GetOldAnimationNormalsView(ParticleRangeId);
+		if (IsValidLODIndex(PrevLODIndex) && GetNumPoints(PrevLODIndex) > 0)
+		{
+			// TODO: Using the more accurate skinning method here would require double buffering the context at the skeletal mesh level
+			const TConstArrayView<FSolverVec3> SrcWrapPositions = Solver->GetOldAnimationPositionsView(PrevParticleRangeId);
+			const TConstArrayView<FSolverVec3> SrcWrapNormals = Solver->GetOldAnimationNormalsView(PrevParticleRangeId);
 
-		const bool bValidWrap = WrapDeformLOD(PrevLODIndex, LODIndex, SrcWrapPositions, SrcWrapNormals, OutOldPositions, OutOldNormals);
+			bValidWrap = WrapDeformLOD(PrevLODIndex, LODIndex, SrcWrapPositions, SrcWrapNormals, OutOldPositions, OutOldNormals);
+		}
 		if (!bValidWrap)
 		{
 			// The previous LOD is invalid, reset old positions with the new LOD
diff --git a/Engine/Plugins/ChaosCloth/Source/ChaosCloth/Private/ChaosCloth/ChaosClothingSimulationSkeletalMesh.cpp b/Engine/Plugins/ChaosCloth/Source/ChaosCloth/Private/ChaosCloth/ChaosClothingSimulationSkeletalMesh.cpp
--- a/Engine/Plugins/ChaosCloth/Source/ChaosCloth/Private/ChaosCloth/ChaosClothingSimulationSkeletalMesh.cpp
+++ b/Engine/Plugins/ChaosCloth/Source/ChaosCloth/Private/ChaosCloth/ChaosClothingSimulationSkeletalMesh.cpp
@@ -37,9 +37,20 @@ namespace Chaos
 	{
 		int32 LODIndex = INDEX_NONE;
 
-		if (Asset && SkeletalMeshComponent)
+		if (SkeletalMeshComponent)
 		{
 			const int32 OwnerLODIndex = SkeletalMeshComponent->GetPredictedLODLevel();
+			LODIndex = GetLODIndexFromOwnerLODIndex(OwnerLODIndex);
+		}
+		return LODIndex;
+	}
+
+	int32 FClothingSimulationSkeletalMesh::GetLODIndexFromOwnerLODIndex(int32 OwnerLODIndex) const
+	{
+		int32 LODIndex = INDEX_NONE;
+
+		if (Asset)
+		{
 			if (Asset->LodMap.IsValidIndex(OwnerLODIndex))
 			{
 				const int32 ClothLODIndex = Asset->LodMap[OwnerLODIndex];
@@ -50,6 +61,7 @@ namespace Chaos
 			}
 		}
 		return LODIndex;
+
 	}
 
 	int32 FClothingSimulationSkeletalMesh::GetOwnerLODIndex(int32 LODIndex) const
diff --git a/Engine/Plugins/ChaosCloth/Source/ChaosCloth/Private/ChaosCloth/ChaosClothingSimulationSolver.cpp b/Engine/Plugins/ChaosCloth/Source/ChaosCloth/Private/ChaosCloth/ChaosClothingSimulationSolver.cpp
--- a/Engine/Plugins/ChaosCloth/Source/ChaosCloth/Private/ChaosCloth/ChaosClothingSimulationSolver.cpp
+++ b/Engine/Plugins/ChaosCloth/Source/ChaosCloth/Private/ChaosCloth/ChaosClothingSimulationSolver.cpp
@@ -121,57 +121,57 @@ namespace Private
 template<typename T>
 TConstArrayView<T> GetPBDParticleConstArrayView(const TUniquePtr<Softs::FPBDEvolution>& PBDEvolution, const int32 ParticleRangeId, const TArray<T>& Array)
 {
-	return TConstArrayView<T>(Array.GetData() + ParticleRangeId, PBDEvolution->GetParticleRangeSize(ParticleRangeId));
+	return ParticleRangeId == INDEX_NONE ? TConstArrayView<T>() : TConstArrayView<T>(Array.GetData() + ParticleRangeId, PBDEvolution->GetParticleRangeSize(ParticleRangeId));
 }
 
 template<typename T>
 TArrayView<T> GetPBDParticleArrayView(TUniquePtr<Softs::FPBDEvolution>& PBDEvolution, const int32 ParticleRangeId, TArray<T>& Array)
 {
-	return TArrayView<T>(Array.GetData() + ParticleRangeId, PBDEvolution->GetParticleRangeSize(ParticleRangeId));
+	return ParticleRangeId == INDEX_NONE ? TArrayView<T>() : TArrayView<T>(Array.GetData() + ParticleRangeId, PBDEvolution->GetParticleRangeSize(ParticleRangeId));
 }
 
 template<typename T>
 TConstArrayView<T> GetParticleConstArrayView(const TUniquePtr<Softs::FEvolution>& Evolution, const TUniquePtr<Softs::FPBDEvolution>& PBDEvolution,
 	const int32 ParticleRangeId, const TArray<T>& Array)
 {
-	return Evolution ? Evolution->GetSoftBodyParticles(ParticleRangeId).GetConstArrayView(Array) :
-		TConstArrayView<T>(Array.GetData() + ParticleRangeId, PBDEvolution->GetParticleRangeSize(ParticleRangeId));
+	return ParticleRangeId == INDEX_NONE ? TConstArrayView<T>() : (Evolution ? Evolution->GetSoftBodyParticles(ParticleRangeId).GetConstArrayView(Array) :
+		TConstArrayView<T>(Array.GetData() + ParticleRangeId, PBDEvolution->GetParticleRangeSize(ParticleRangeId)));
 }
 
 template<typename T>
 TArrayView<T> GetParticleArrayView(TUniquePtr<Softs::FEvolution>& Evolution, TUniquePtr<Softs::FPBDEvolution>& PBDEvolution,
 	const int32 ParticleRangeId, TArray<T>& Array)
 {
-	return Evolution ? Evolution->GetSoftBodyParticles(ParticleRangeId).GetArrayView(Array) :
-		TArrayView<T>(Array.GetData() + ParticleRangeId, PBDEvolution->GetParticleRangeSize(ParticleRangeId));
+	return ParticleRangeId == INDEX_NONE ? TArrayView<T>() : (Evolution ? Evolution->GetSoftBodyParticles(ParticleRangeId).GetArrayView(Array) :
+		TArrayView<T>(Array.GetData() + ParticleRangeId, PBDEvolution->GetParticleRangeSize(ParticleRangeId)));
 }
 
 template<typename T>
 TConstArrayView<T> GetPBDCollisionParticleConstArrayView(const TUniquePtr<Softs::FPBDEvolution>& PBDEvolution, const int32 ParticleRangeId, const TArray<T>& Array)
 {
-	return TConstArrayView<T>(Array.GetData() + ParticleRangeId, PBDEvolution->GetCollisionParticleRangeSize(ParticleRangeId));
+	return ParticleRangeId == INDEX_NONE ? TConstArrayView<T>() : TConstArrayView<T>(Array.GetData() + ParticleRangeId, PBDEvolution->GetCollisionParticleRangeSize(ParticleRangeId));
 }
 
 template<typename T>
 TArrayView<T> GetPBDCollisionParticleArrayView(TUniquePtr<Softs::FPBDEvolution>& PBDEvolution, const int32 ParticleRangeId, TArray<T>& Array)
 {
-	return TArrayView<T>(Array.GetData() + ParticleRangeId, PBDEvolution->GetCollisionParticleRangeSize(ParticleRangeId));
+	return ParticleRangeId == INDEX_NONE ? TArrayView<T>() : TArrayView<T>(Array.GetData() + ParticleRangeId, PBDEvolution->GetCollisionParticleRangeSize(ParticleRangeId));
 }
 
 template<typename T>
 TConstArrayView<T> GetCollisionParticleConstArrayView(const TUniquePtr<Softs::FEvolution>& Evolution, const TUniquePtr<Softs::FPBDEvolution>& PBDEvolution,
 	const int32 ParticleRangeId, const TArray<T>& Array)
 {
-	return Evolution ? Evolution->GetCollisionParticleRange(ParticleRangeId).GetConstArrayView(Array) :
-		TConstArrayView<T>(Array.GetData() + ParticleRangeId, PBDEvolution->GetCollisionParticleRangeSize(ParticleRangeId));
+	return ParticleRangeId == INDEX_NONE ? TConstArrayView<T>() : (Evolution ? Evolution->GetCollisionParticleRange(ParticleRangeId).GetConstArrayView(Array) :
+		TConstArrayView<T>(Array.GetData() + ParticleRangeId, PBDEvolution->GetCollisionParticleRangeSize(ParticleRangeId)));
 }
 
 template<typename T>
 TArrayView<T> GetCollisionParticleArrayView(TUniquePtr<Softs::FEvolution>& Evolution, TUniquePtr<Softs::FPBDEvolution>& PBDEvolution,
 	const int32 ParticleRangeId, TArray<T>& Array)
 {
-	return Evolution ? Evolution->GetCollisionParticleRange(ParticleRangeId).GetArrayView(Array) :
-		TArrayView<T>(Array.GetData() + ParticleRangeId, PBDEvolution->GetCollisionParticleRangeSize(ParticleRangeId));
+	return ParticleRangeId == INDEX_NONE ? TArrayView<T>() : (Evolution ? Evolution->GetCollisionParticleRange(ParticleRangeId).GetArrayView(Array) :
+		TArrayView<T>(Array.GetData() + ParticleRangeId, PBDEvolution->GetCollisionParticleRangeSize(ParticleRangeId)));
 }
 }
 
diff --git a/Engine/Plugins/ChaosCloth/Source/ChaosCloth/Public/ChaosCloth/ChaosClothingSimulation.h b/Engine/Plugins/ChaosCloth/Source/ChaosCloth/Public/ChaosCloth/ChaosClothingSimulation.h
--- a/Engine/Plugins/ChaosCloth/Source/ChaosCloth/Public/ChaosCloth/ChaosClothingSimulation.h
+++ b/Engine/Plugins/ChaosCloth/Source/ChaosCloth/Public/ChaosCloth/ChaosClothingSimulation.h
@@ -41,11 +41,18 @@ namespace Chaos
 		virtual void DestroyContext(IClothingSimulationContext* InContext) override { delete InContext; }
 
 		virtual void CreateActor(USkeletalMeshComponent* InOwnerComponent, const UClothingAssetBase* InAsset, int32 SimDataIndex) override;
+		virtual void EndCreateActor() override;
 		virtual void DestroyActors() override;
 
-		virtual bool ShouldSimulate() const override;
+		virtual void FillContextAndPrepareTick(const USkeletalMeshComponent* InComponent, float InDeltaTime, IClothingSimulationContext* InOutContext, bool bIsInitialization, bool bForceTeleportResetOnly) override;
+		virtual bool ShouldSimulateLOD(int32 OwnerLODIndex) const override;
 		virtual void Simulate(const IClothingSimulationContext* InContext) override;
-		virtual void GetSimulationData(TMap<int32, FClothSimulData>& OutData, const USkeletalMeshComponent* InOwnerComponent, const USkinnedMeshComponent* InOverrideComponent) const override;
+		virtual void ForceClothNextUpdateTeleportAndReset() override;
+		virtual void HardResetSimulation(const IClothingSimulationContext* InContext) override
+		{
+			RefreshClothConfig(InContext);
+		}
+		virtual void AppendSimulationData(TMap<int32, FClothSimulData>& OutData, const USkeletalMeshComponent* InOwnerComponent, const USkinnedMeshComponent* InOverrideComponent) const override;
 
 		// Return bounds in local space (or in world space if InOwnerComponent is null).
 		virtual FBoxSphereBounds GetBounds(const USkeletalMeshComponent* InOwnerComponent) const override;
@@ -149,6 +156,8 @@ namespace Chaos
 		// Shared cloth config
 		const UChaosClothSharedSimConfig* ClothSharedSimConfig;
 
+		TBitArray<> LODHasAnyRenderClothMappingData; // Used by ShouldSimulateLOD
+
 		// Properties that must be readable from all threads
 		TAtomic<int32> NumCloths;
 		TAtomic<int32> NumKinematicParticles;
diff --git a/Engine/Plugins/ChaosCloth/Source/ChaosCloth/Public/ChaosCloth/ChaosClothingSimulationMesh.h b/Engine/Plugins/ChaosCloth/Source/ChaosCloth/Public/ChaosCloth/ChaosClothingSimulationMesh.h
--- a/Engine/Plugins/ChaosCloth/Source/ChaosCloth/Public/ChaosCloth/ChaosClothingSimulationMesh.h
+++ b/Engine/Plugins/ChaosCloth/Source/ChaosCloth/Public/ChaosCloth/ChaosClothingSimulationMesh.h
@@ -53,6 +53,14 @@ namespace Chaos
 		/* Return the cloth mesh LOD Index. The returned value can then be used to switch LODs on the simulation object. */
 		virtual int32 GetLODIndex() const = 0;
 
+		/* Return the cloth mesh LOD Index for a given owner LOD index. */
+		virtual int32 GetLODIndexFromOwnerLODIndex(int32 OwnerLODIndex) const
+		{
+			// This default implementation assumes OwnerLODIndex == Cloth mesh LOD Index. 
+			// Override this method if you have some other mapping.
+			return IsValidLODIndex(OwnerLODIndex) ? OwnerLODIndex : INDEX_NONE;
+		}
+
 		/* Return the owner component LOD Index from the specified cloth mesh LOD index, or 0 if the owner LOD cannot be determined. The mapping between the cloth mesh LOD and the owner component LOD is not necessarily one to one. */
 		virtual int32 GetOwnerLODIndex(int32 LODIndex) const = 0;
 
diff --git a/Engine/Plugins/ChaosCloth/Source/ChaosCloth/Public/ChaosCloth/ChaosClothingSimulationSkeletalMesh.h b/Engine/Plugins/ChaosCloth/Source/ChaosCloth/Public/ChaosCloth/ChaosClothingSimulationSkeletalMesh.h
--- a/Engine/Plugins/ChaosCloth/Source/ChaosCloth/Public/ChaosCloth/ChaosClothingSimulationSkeletalMesh.h
+++ b/Engine/Plugins/ChaosCloth/Source/ChaosCloth/Public/ChaosCloth/ChaosClothingSimulationSkeletalMesh.h
@@ -25,6 +25,7 @@ namespace Chaos
 		//~ Begin FClothingSimulationMesh Interface
 		virtual int32 GetNumLODs() const override;
 		virtual int32 GetLODIndex() const override;
+		virtual int32 GetLODIndexFromOwnerLODIndex(int32 OwnerLODIndex) const override;
 		virtual int32 GetOwnerLODIndex(int32 LODIndex) const override;
 		virtual bool IsValidLODIndex(int32 LODIndex) const override;
 		virtual int32 GetNumPoints(int32 LODIndex) const override;
diff --git a/Engine/Plugins/ChaosClothAsset/Source/ChaosClothAssetEngine/Private/ChaosClothAsset/ClothAssetSKMClothingAsset.cpp b/Engine/Plugins/ChaosClothAsset/Source/ChaosClothAssetEngine/Private/ChaosClothAsset/ClothAssetSKMClothingAsset.cpp
--- a/Engine/Plugins/ChaosClothAsset/Source/ChaosClothAssetEngine/Private/ChaosClothAsset/ClothAssetSKMClothingAsset.cpp
+++ b/Engine/Plugins/ChaosClothAsset/Source/ChaosClothAssetEngine/Private/ChaosClothAsset/ClothAssetSKMClothingAsset.cpp
@@ -123,6 +123,13 @@ UChaosClothAssetSKMClothingAsset::UChaosClothAssetSKMClothingAsset(const FObject
 	UClothingAssetBase::AssetGuid = FGuid::NewGuid();
 }
 
+void UChaosClothAssetSKMClothingAsset::PostLoad()
+{
+	Super::PostLoad();
+
+	CalculateLODHasAnySimulationMeshData();
+}
+
 #if WITH_EDITOR
 void UChaosClothAssetSKMClothingAsset::SetAsset(const UChaosClothAssetBase* InAsset)
 {
@@ -181,6 +188,10 @@ void UChaosClothAssetSKMClothingAsset::OnAssetChanged()
 
 		OnModelChanged();
 	}
+	else
+	{
+		CalculateLODHasAnySimulationMeshData();
+	}
 
 	// Reregister components
 	if (const USkeletalMesh* const OwnerMesh = Cast<USkeletalMesh>(GetOuter()))
@@ -221,7 +232,8 @@ bool UChaosClothAssetSKMClothingAsset::BindToSkeletalMesh(USkeletalMesh* Skeleta
 		!ensure(SkeletalMesh == GetOuter()) ||
 		!ensure(SkeletalMesh->GetImportedModel()) ||
 		!ensure(SkeletalMesh->GetImportedModel()->LODModels.IsValidIndex(MeshLodIndex)) ||
-		!ensure(SkeletalMesh->GetImportedModel()->LODModels[MeshLodIndex].Sections.IsValidIndex(SectionIndex)))
+		!ensure(SkeletalMesh->GetImportedModel()->LODModels[MeshLodIndex].Sections.IsValidIndex(SectionIndex)) ||
+		!ensure(MeshLodIndex == AssetLodIndex))
 	{
 		return false;
 	}
@@ -419,6 +431,7 @@ void UChaosClothAssetSKMClothingAsset::OnModelChanged()
 		UnbindClothingAndBackup(*SkeletalMesh, *this, ClothingBindings);
 		RestoreClothingFromBackup(*SkeletalMesh, *this, ClothingBindings);
 	}
+	CalculateLODHasAnySimulationMeshData();
 }
 
 FGuid UChaosClothAssetSKMClothingAsset::GetClothSimulationModelGuid() const
@@ -458,4 +471,20 @@ TArray<FString> UChaosClothAssetSKMClothingAsset::GetClothSimulationModelIds() c
 	return ClothSimulationModelIds;
 }
 
+void UChaosClothAssetSKMClothingAsset::CalculateLODHasAnySimulationMeshData()
+{
+	LODHasAnySimulationMeshData.Reset();
+	if (Asset && ClothSimulationModelIndex != INDEX_NONE)
+	{
+		if (const TSharedPtr<const FChaosClothSimulationModel> ClothSimulationModel = Asset->GetClothSimulationModel(ClothSimulationModelIndex))
+		{
+			LODHasAnySimulationMeshData.SetNumUninitialized(ClothSimulationModel->GetNumLods());
+			for (int32 LODIndex = 0; LODIndex < ClothSimulationModel->GetNumLods(); ++LODIndex)
+			{
+				LODHasAnySimulationMeshData[LODIndex] = ClothSimulationModel->GetNumVertices(LODIndex) > 0 && ClothSimulationModel->GetNumTriangles(LODIndex) > 0;
+			}
+		}
+	}
+}
+
 #undef LOCTEXT_NAMESPACE
diff --git a/Engine/Plugins/ChaosClothAsset/Source/ChaosClothAssetEngine/Private/ChaosClothAsset/ClothAssetSKMClothingSimulation.cpp b/Engine/Plugins/ChaosClothAsset/Source/ChaosClothAssetEngine/Private/ChaosClothAsset/ClothAssetSKMClothingSimulation.cpp
--- a/Engine/Plugins/ChaosClothAsset/Source/ChaosClothAssetEngine/Private/ChaosClothAsset/ClothAssetSKMClothingSimulation.cpp
+++ b/Engine/Plugins/ChaosClothAsset/Source/ChaosClothAssetEngine/Private/ChaosClothAsset/ClothAssetSKMClothingSimulation.cpp
@@ -7,6 +7,7 @@
 #include "ChaosClothAsset/ClothAssetBase.h"
 #include "ChaosClothAsset/ClothAssetSKMClothingAsset.h"
 #include "ChaosClothAsset/ClothComponent.h"  // For FChaosClothSimulationProperties  TODO: Move FChaosClothSimulationProperties to its own header file?
+#include "ChaosClothAsset/ClothSimulationModel.h"
 #include "ChaosClothAsset/ClothSimulationProxy.h"
 #include "ChaosClothAsset/CollisionSources.h"
 #include "Components/SkeletalMeshComponent.h"
@@ -74,12 +75,31 @@ namespace UE::Chaos::ClothAsset
 		ClothSimulationProxy = nullptr;
 	}
 
-	void FSKMClothingSimulation::FillContext(const USkeletalMeshComponent* /*InComponent*/, float InDeltaTime, IClothingSimulationContext* InOutContext, bool /*bIsInitialization*/) const
+	void FSKMClothingSimulation::FillContextAndPrepareTick(const USkeletalMeshComponent* /*InComponent*/, float InDeltaTime, IClothingSimulationContext* InOutContext, bool bIsInitialization, bool bForceTeleportResetOnly)
 	{
 		if (FSKMClothingSimulationContext* const Context = static_cast<FSKMClothingSimulationContext*>(InOutContext))
 		{
 			Context->DeltaTime = InDeltaTime;
 		}
+		if (!bIsInitialization) // Actual context and cloths will be initialized as part of EndCreateActor when initializing
+		{
+			if (ClothSimulationProxy.IsValid())
+			{
+				if (bForceTeleportResetOnly)
+				{
+					ClothSimulationProxy->ForcePendingReset_GameThread();
+				}
+				else
+				{
+					ClothSimulationProxy->PreProcess_GameThread(InDeltaTime);
+					ClothSimulationProxy->PreSimulate_GameThread(InDeltaTime);
+
+					// Reset the bNeedsResetRestLengths at every pass
+					bNeedsResetRestLengths = false;
+				}
+			}
+		}
+
 	}
 
 	void FSKMClothingSimulation::CreateActor(USkeletalMeshComponent* InOwnerComponent, const UClothingAssetBase* ClothingAsset, int32 SimDataIndex)
@@ -142,6 +162,8 @@ namespace UE::Chaos::ClothAsset
 
 	void FSKMClothingSimulation::EndCreateActor()
 	{
+		CalculateLODHasAnyRenderClothMappingData();
+
 		if (ClothSimulationProxy.IsValid())
 		{
 			// Add all assets to the proxy
@@ -152,25 +174,16 @@ namespace UE::Chaos::ClothAsset
 	void FSKMClothingSimulation::DestroyActors()
 	{
 		AssetClothSimulationProperties.Reset();
+		LODHasAnyRenderClothMappingData.Reset();
 	}
 
-	bool FSKMClothingSimulation::ShouldSimulate() const
+	bool FSKMClothingSimulation::ShouldSimulateLOD(int32 OwnerLODIndex) const
 	{
 		if (!OwnerComponent)
 		{
 			return false;
 		}
-		for (const FAssetClothSimulationProperties& Properties : AssetClothSimulationProperties)
-		{
-			if (const UChaosClothAssetSKMClothingAsset* const ClothAssetSKMClothingAsset = Properties.Key)
-			{
-				if (ClothAssetSKMClothingAsset->GetClothSimulationModelIndex() != INDEX_NONE)
-				{
-					return true;
-				}
-			}
-		}
-		return false;
+		return LODHasAnyRenderClothMappingData.IsValidIndex(OwnerLODIndex) && LODHasAnyRenderClothMappingData[OwnerLODIndex] && HasAnySimulationMeshData(OwnerLODIndex);
 	}
 
 	void FSKMClothingSimulation::Simulate(const IClothingSimulationContext* Context)
@@ -179,23 +192,37 @@ namespace UE::Chaos::ClothAsset
 		{
 			if (ClothSimulationProxy.IsValid())
 			{
-				constexpr bool bUseParallelTask = false;  // The SKM component already deals with calling Simulate() in its own parallel task
-				ClothSimulationProxy->Tick_GameThread(SKMClothingSimulationContext->DeltaTime, bUseParallelTask);
+				ClothSimulationProxy->Tick();
 			}
 		}
-		// Reset the bNeedsResetRestLengths at every pass
-		bNeedsResetRestLengths = false;
+	}
+	
+	void FSKMClothingSimulation::ForceClothNextUpdateTeleportAndReset()
+	{
+		// Do nothing here. This is handled in FillContextAndPrepareTick
 	}
 
-	void FSKMClothingSimulation::GetSimulationData(TMap<int32, FClothSimulData>& OutData, const USkeletalMeshComponent* /*InOwnerComponent*/, const USkinnedMeshComponent* /*OverrideComponent*/) const
+	void FSKMClothingSimulation::HardResetSimulation(const IClothingSimulationContext* InContext)
 	{
+		ResetConfigProperties();
 		if (ClothSimulationProxy.IsValid())
 		{
-			OutData = ClothSimulationProxy->GetCurrentSimulationData_AnyThread();  // TODO: Add a way to get the simulationdata without the unnecessary copy: GetCurrentSimulationData_GameThread(TMap<int32, FClothSimulData>& OutData)?
+			ClothSimulationProxy->HardResetSimulation_GameThread();
 		}
-		else
+	}
+
+	void FSKMClothingSimulation::AppendSimulationData(TMap<int32, FClothSimulData>& InOutData, const USkeletalMeshComponent* /*InOwnerComponent*/, const USkinnedMeshComponent* /*OverrideComponent*/) const
+	{
+		if (ClothSimulationProxy.IsValid())
 		{
-			OutData.Reset();
+			ClothSimulationProxy->PostSimulate_GameThread();
+			ClothSimulationProxy->PostProcess_GameThread();
+
+			// TODO: Add a way to get the simulationdata without the unnecessary copy: GetCurrentSimulationData_GameThread(TMap<int32, FClothSimulData>& OutData)?
+			for (TMap<int32, FClothSimulData>::TConstIterator CurrentDataIter = ClothSimulationProxy->GetCurrentSimulationData_AnyThread().CreateConstIterator(); CurrentDataIter; ++CurrentDataIter)
+			{
+				InOutData.FindOrAdd(CurrentDataIter.Key()) = CurrentDataIter.Value();
+			}
 		}
 	}
 
@@ -346,29 +373,16 @@ namespace UE::Chaos::ClothAsset
 		return EmptyManagedArrayCollectionArray;
 	}
 
-	bool FSKMClothingSimulation::HasAnyRenderClothMappingData(int32 LodIndex) const
+	bool FSKMClothingSimulation::HasAnySimulationMeshData(int32 LodIndex) const
 	{
-		if (const USkeletalMesh* const SkeletalMesh = OwnerComponent->GetSkeletalMeshAsset())
+		// Note that for ClothAsset-based simulation, OwnerLODIndex == ClothAssetLODIndex
+		for (const FAssetClothSimulationProperties& Properties : AssetClothSimulationProperties)
 		{
-			if (const FSkeletalMeshRenderData* const SkeletalMeshRenderData = SkeletalMesh->GetResourceForRendering())
+			if (const UChaosClothAssetSKMClothingAsset* const ClothAssetSKMClothingAsset = Properties.Key)
 			{
-				if (SkeletalMeshRenderData->LODRenderData.IsValidIndex(LodIndex))
+				if (ClothAssetSKMClothingAsset->HasAnySimulationMeshData(LodIndex))
 				{
-					for (const FSkelMeshRenderSection& SkelMeshRenderSection : SkeletalMeshRenderData->LODRenderData[LodIndex].RenderSections)
-					{
-						if (SkelMeshRenderSection.HasClothingData() &&
-							AssetClothSimulationProperties.IsValidIndex(SkelMeshRenderSection.CorrespondClothAssetIndex))
-						{
-							if (const UChaosClothAssetSKMClothingAsset* const ClothAssetSKMClothingAsset = AssetClothSimulationProperties[SkelMeshRenderSection.CorrespondClothAssetIndex].Key)
-							{
-								if (SkelMeshRenderSection.ClothingData.AssetGuid == ClothAssetSKMClothingAsset->GetAssetGuid() &&
-									SkelMeshRenderSection.ClothingData.AssetLodIndex == LodIndex)
-								{
-									return true;
-								}
-							}
-						}
-					}
+					return true;
 				}
 			}
 		}
@@ -415,6 +429,40 @@ namespace UE::Chaos::ClothAsset
 		static const TArray<TSharedPtr<::Chaos::Softs::FCollectionPropertyFacade>> EmptyManagedArrayCollectionArray;
 		return EmptyManagedArrayCollectionArray;
 	}
+
+	void FSKMClothingSimulation::CalculateLODHasAnyRenderClothMappingData()
+	{
+		LODHasAnyRenderClothMappingData.Reset();
+		if (OwnerComponent)
+		{
+			if (const USkeletalMesh* const SkeletalMesh = OwnerComponent->GetSkeletalMeshAsset())
+			{
+				if (const FSkeletalMeshRenderData* const SkeletalMeshRenderData = SkeletalMesh->GetResourceForRendering())
+				{
+					LODHasAnyRenderClothMappingData.Init(false, SkeletalMeshRenderData->LODRenderData.Num());
+					for (int32 LodIndex = 0; LodIndex < SkeletalMeshRenderData->LODRenderData.Num(); ++LodIndex)
+					{
+						for (const FSkelMeshRenderSection& SkelMeshRenderSection : SkeletalMeshRenderData->LODRenderData[LodIndex].RenderSections)
+						{
+							if (SkelMeshRenderSection.HasClothingData() &&
+								AssetClothSimulationProperties.IsValidIndex(SkelMeshRenderSection.CorrespondClothAssetIndex))
+							{
+								if (const UChaosClothAssetSKMClothingAsset* const ClothAssetSKMClothingAsset = AssetClothSimulationProperties[SkelMeshRenderSection.CorrespondClothAssetIndex].Key)
+								{
+									if (SkelMeshRenderSection.ClothingData.AssetGuid == ClothAssetSKMClothingAsset->GetAssetGuid() &&
+										SkelMeshRenderSection.ClothingData.AssetLodIndex == LodIndex)
+									{
+										LODHasAnyRenderClothMappingData[LodIndex] = true;
+										continue;
+									}
+								}
+							}
+						}
+					}
+				}
+			}
+		}
+	}
 }  // namespace UE::Chaos::ClothAsset
 
 UChaosClothAssetSKMClothingSimulationFactory::UChaosClothAssetSKMClothingSimulationFactory() = default;
diff --git a/Engine/Plugins/ChaosClothAsset/Source/ChaosClothAssetEngine/Private/ChaosClothAsset/ClothComponent.cpp b/Engine/Plugins/ChaosClothAsset/Source/ChaosClothAssetEngine/Private/ChaosClothAsset/ClothComponent.cpp
--- a/Engine/Plugins/ChaosClothAsset/Source/ChaosClothAssetEngine/Private/ChaosClothAsset/ClothComponent.cpp
+++ b/Engine/Plugins/ChaosClothAsset/Source/ChaosClothAssetEngine/Private/ChaosClothAsset/ClothComponent.cpp
@@ -626,15 +626,21 @@ void UChaosClothComponent::UnregisterOnBoneTransformsFinalizedDelegate(const FDe
 	OnBoneTransformsFinalizedMC.Remove(DelegateHandle);
 }
 
-bool UChaosClothComponent::HasAnyRenderClothMappingData(int32 LODIndex) const
+bool UChaosClothComponent::HasAnySimulationMeshData(int32 LODIndex) const
 {
-	if (const UChaosClothAssetBase* const AssetToRender = GetAsset())
+	if (const UChaosClothAssetBase* const AssetToSimulate = GetAsset())
 	{
-		if (const FSkeletalMeshRenderData* const SkeletalMeshRenderData = AssetToRender->GetResourceForRendering())
+		const int32 NumSimulationModels = AssetToSimulate->GetNumClothSimulationModels();
+		for (int32 ModelIndex = 0; ModelIndex < NumSimulationModels; ++ModelIndex)
 		{
-			return
-				SkeletalMeshRenderData->LODRenderData.IsValidIndex(LODIndex) &&
-				SkeletalMeshRenderData->LODRenderData[LODIndex].HasClothData();
+			if (const TSharedPtr<const FChaosClothSimulationModel> ClothSimulationModel = AssetToSimulate->GetClothSimulationModel(ModelIndex))
+			{
+				if (ClothSimulationModel->IsValidLodIndex(LODIndex)
+					&& ClothSimulationModel->GetNumVertices(LODIndex) > 0)
+				{
+					return true;
+				}
+			}
 		}
 	}
 	return false;
diff --git a/Engine/Plugins/ChaosClothAsset/Source/ChaosClothAssetEngine/Private/ChaosClothAsset/ClothSimulationContext.cpp b/Engine/Plugins/ChaosClothAsset/Source/ChaosClothAssetEngine/Private/ChaosClothAsset/ClothSimulationContext.cpp
--- a/Engine/Plugins/ChaosClothAsset/Source/ChaosClothAssetEngine/Private/ChaosClothAsset/ClothSimulationContext.cpp
+++ b/Engine/Plugins/ChaosClothAsset/Source/ChaosClothAssetEngine/Private/ChaosClothAsset/ClothSimulationContext.cpp
@@ -19,18 +19,17 @@ namespace UE::Chaos::ClothAsset
 		// Set the time
 		DeltaTime = FMath::Min(InDeltaTime, MaxDeltaTime);
 
-		// Set the current LOD index
+		// Set the current LOD index. Do this before exiting early if there is no simulation data.
 		LodIndex = ClothComponentAdapter.GetOwnerComponent().GetPredictedLODLevel();
 
-		// Retrieve the list of assets to simulate
-		const TArray<const UChaosClothAssetBase*> Assets = ClothComponentAdapter.GetAssets();
-
-		// Do not simulate if there's no asset to render
-		if (!bIsInitialization && !ClothComponentAdapter.HasAnyRenderClothMappingData(LodIndex))
+		if (!bIsInitialization && !ClothComponentAdapter.HasAnySimulationMeshData(LodIndex))
 		{
 			return;  // This Lod will not simulate, so no need to update the context.
 		}
 
+		// Retrieve the list of assets to simulate
+		const TArray<const UChaosClothAssetBase*> Assets = ClothComponentAdapter.GetAssets();
+
 		// Set the teleport mode
 		static const IConsoleVariable* const CVarMaxDeltaTimeTeleportMultiplier = IConsoleManager::Get().FindConsoleVariable(TEXT("p.Cloth.MaxDeltaTimeTeleportMultiplier"));
 		constexpr float MaxDeltaTimeTeleportMultiplierDefault = 1.5f;
diff --git a/Engine/Plugins/ChaosClothAsset/Source/ChaosClothAssetEngine/Private/ChaosClothAsset/ClothSimulationProxy.cpp b/Engine/Plugins/ChaosClothAsset/Source/ChaosClothAssetEngine/Private/ChaosClothAsset/ClothSimulationProxy.cpp
--- a/Engine/Plugins/ChaosClothAsset/Source/ChaosClothAssetEngine/Private/ChaosClothAsset/ClothSimulationProxy.cpp
+++ b/Engine/Plugins/ChaosClothAsset/Source/ChaosClothAssetEngine/Private/ChaosClothAsset/ClothSimulationProxy.cpp
@@ -191,6 +191,7 @@ namespace UE::Chaos::ClothAsset
 		Solver->SetLocalSpaceRotation((FQuat)ClothSimulationContext->ComponentTransform.GetRotation());
 
 		// Create mesh simulation thread objects
+		// HardResetSimulation_GameThread expects to find Cloths added in this order, so update it if you change Cloths initialization order here.
 		for (const UChaosClothAssetBase* const Asset : ClothComponentAdapter.GetAssets())
 		{
 			const int32 NumClothSimulationModels = Asset ? Asset->GetNumClothSimulationModels() : 0;
@@ -298,6 +299,81 @@ namespace UE::Chaos::ClothAsset
 		return bIsSimulating;
 	}
 
+	void FClothSimulationProxy::HardResetSimulation_GameThread()
+	{
+		// Reset solver, cloths, and configs, but keep existing meshes
+		using namespace ::Chaos;
+
+		check(IsInGameThread());
+
+		// Need a valid context to initialize startup transforms
+		constexpr bool bIsInitialization = true;
+		constexpr Softs::FSolverReal NoAdvanceDt = 0.f;
+		ClothSimulationContext->Fill(ClothComponentAdapter, NoAdvanceDt, MaxDeltaTime, bIsInitialization);
+
+		// Setup startup transforms
+		constexpr bool bNeedsReset = true;
+		const FReal LocalSpaceScale = 1. / FMath::Max(ClothSimulationContext->SolverGeometryScale, UE_SMALL_NUMBER);
+		Solver->SetLocalSpaceScale(LocalSpaceScale, bNeedsReset);
+		Solver->SetLocalSpaceLocation((FVec3)ClothSimulationContext->ComponentTransform.GetLocation(), bNeedsReset);
+		Solver->SetLocalSpaceRotation((FQuat)ClothSimulationContext->ComponentTransform.GetRotation());
+
+		// Reset stats
+		int32 LocalNumCloths = 0;
+		int32 LocalNumKinematicParticles = 0;
+		int32 LocalNumDynamicParticles = 0;
+
+		// Clear all cloths from the solver
+		Solver->RemoveCloths();
+
+		// Recreate all cloths.
+		// This loop matches the loop in BeginInitialization_GameThread and expects Cloths to have been added in that order (and for the assets to have not changed since last initialization).
+		int32 ClothIndex = 0;
+		for (const UChaosClothAssetBase* const Asset : ClothComponentAdapter.GetAssets())
+		{
+			const int32 NumClothSimulationModels = Asset ? Asset->GetNumClothSimulationModels() : 0;
+			for (int32 ModelIndex = 0; ModelIndex < NumClothSimulationModels; ++ModelIndex)
+			{
+				const int32 SimulationGroupId = ClothComponentAdapter.GetSimulationGroupId(Asset, ModelIndex);
+				if (SimulationGroupId != INDEX_NONE)  // INDEX_NONE means this model shouldn't be added to the simulation
+				{
+					check(Cloths.IsValidIndex(ClothIndex));
+
+					TUniquePtr<FClothingSimulationCloth>& Cloth = Cloths[ClothIndex++];
+					check(SimulationGroupId == Cloth->GetGroupId());
+
+					FClothSimulationMesh* const Mesh = static_cast<FClothSimulationMesh*>(Cloth->GetMesh());
+					TArray<FClothingSimulationCollider*> ClothColliders = Cloth->GetColliders();
+
+					// Update cloth config runtime simulation object
+					FClothingSimulationConfig* const Config = Cloth->GetConfig();
+					Config->Initialize(ClothComponentAdapter.GetPropertyCollections(Asset, ModelIndex));
+
+					// Recreate cloth runtime simulation object
+					Cloth = MakeUnique<FClothingSimulationCloth>(
+						Config,
+						Mesh,
+						MoveTemp(ClothColliders),
+						SimulationGroupId);
+
+					// Re-add cloth to the solver
+					Solver->AddCloth(Cloth.Get());
+
+					// Update stats
+					++LocalNumCloths;
+					LocalNumKinematicParticles += Cloth->GetNumActiveKinematicParticles();
+					LocalNumDynamicParticles += Cloth->GetNumActiveDynamicParticles();
+				}
+			}
+		}
+		check(ClothIndex == Cloths.Num());
+
+		// Update cloth stats
+		NumCloths = LocalNumCloths;
+		NumKinematicParticles = LocalNumKinematicParticles;
+		NumDynamicParticles = LocalNumDynamicParticles;
+	}
+
 	void FClothSimulationProxy::PreProcess_GameThread(float DeltaTime, bool bForceWaitForInitialization)
 	{
 		TRACE_CPUPROFILER_EVENT_SCOPE(FClothSimulationProxy_PreProcess_GameThread);
@@ -360,8 +436,8 @@ namespace UE::Chaos::ClothAsset
 			return false;
 		}
 
-		// Check that the render mesh current LOD isn't just fully skinned
-		if (!ClothComponentAdapter.HasAnyRenderClothMappingData(ClothSimulationContext->LodIndex))
+		// Check that the there's any simulation data for this LOD
+		if (!ClothComponentAdapter.HasAnySimulationMeshData(ClothSimulationContext->LodIndex))
 		{
 			return false;
 		}
@@ -426,7 +502,30 @@ namespace UE::Chaos::ClothAsset
 		}
 	}
 
-	bool FClothSimulationProxy::Tick_GameThread(float DeltaTime, bool bUseParallelTask)
+	void FClothSimulationProxy::ForcePendingReset_GameThread()
+	{
+		// Check for any in-flight work.
+		if (IsValidRef(ParallelInitializationTask))
+		{
+			if (!bWaitForAsyncClothInitialization && !ParallelInitializationTask->IsComplete())
+			{
+				return;
+			}
+			CompleteInitialization_GameThread();
+		}
+
+		check(bIsInitialized);
+
+		CompleteParallelSimulation_GameThread();
+
+		using namespace ::Chaos;
+		for (const TUniquePtr<FClothingSimulationCloth>& Cloth : Cloths)
+		{
+			Cloth->Reset();
+		}
+	}
+
+	bool FClothSimulationProxy::Tick_GameThread(float DeltaTime)
 	{
 		SCOPE_CYCLE_COUNTER(STAT_ClothSimulationProxy_TickGame);
 
@@ -435,14 +534,9 @@ namespace UE::Chaos::ClothAsset
 
 		if (bIsSimulating)
 		{
-			if (bUseParallelTask)
-			{
-				ParallelTask = TGraphTask<FClothSimulationProxyParallelTask>::CreateTask(nullptr, ENamedThreads::GameThread).ConstructAndDispatchWhenReady(*this);
-
-				return true;  // Note: Post simulate and post process will be handled when the parallel task completes
-			}
+			ParallelTask = TGraphTask<FClothSimulationProxyParallelTask>::CreateTask(nullptr, ENamedThreads::GameThread).ConstructAndDispatchWhenReady(*this);
 
-			Tick();
+			return true;  // Note: Post simulate and post process will be handled when the parallel task completes
 		}
 
 		PostSimulate_GameThread();
diff --git a/Engine/Plugins/ChaosClothAsset/Source/ChaosClothAssetEngine/Private/ChaosClothAsset/CollisionSources.cpp b/Engine/Plugins/ChaosClothAsset/Source/ChaosClothAssetEngine/Private/ChaosClothAsset/CollisionSources.cpp
--- a/Engine/Plugins/ChaosClothAsset/Source/ChaosClothAssetEngine/Private/ChaosClothAsset/CollisionSources.cpp
+++ b/Engine/Plugins/ChaosClothAsset/Source/ChaosClothAssetEngine/Private/ChaosClothAsset/CollisionSources.cpp
@@ -86,7 +86,7 @@ namespace UE::Chaos::ClothAsset
 
 	void FCollisionSources::AddExternalCollisions(const FClothCollisionData& InExternalCollisions)
 	{
-		ExternalCollisions = InExternalCollisions;
+		ExternalCollisions.Append(InExternalCollisions);
 	}
 
 	void FCollisionSources::ClearExternalCollisions()
diff --git a/Engine/Plugins/ChaosClothAsset/Source/ChaosClothAssetEngine/Public/ChaosClothAsset/ClothAssetSKMClothingAsset.h b/Engine/Plugins/ChaosClothAsset/Source/ChaosClothAssetEngine/Public/ChaosClothAsset/ClothAssetSKMClothingAsset.h
--- a/Engine/Plugins/ChaosClothAsset/Source/ChaosClothAssetEngine/Public/ChaosClothAsset/ClothAssetSKMClothingAsset.h
+++ b/Engine/Plugins/ChaosClothAsset/Source/ChaosClothAssetEngine/Public/ChaosClothAsset/ClothAssetSKMClothingAsset.h
@@ -43,11 +43,17 @@ public:
 		return ClothSimulationModelIndex;
 	}
 
+	bool HasAnySimulationMeshData(const int32 LODIndex) const
+	{
+		return LODHasAnySimulationMeshData.IsValidIndex(LODIndex) && LODHasAnySimulationMeshData[LODIndex];
+	}
+
 private:
-#if WITH_EDITOR
 	//~ Begin UObject Interface
+	virtual void PostLoad() override;
+#if WITH_EDITOR
 	virtual void PostEditChangeProperty(FPropertyChangedEvent& PropertyChangedEvent) override;
-	//~ EnUObject Interface
+	//~ End UObject Interface
 
 	//~ Begin UClothingAssetBase Interface
 	virtual bool BindToSkeletalMesh(USkeletalMesh* SkeletalMesh, const int32 MeshLodIndex, const int32 SectionIndex, const int32 AssetLodIndex) override;
@@ -63,6 +69,8 @@ private:
 	virtual void PostUpdateAllAssets() override {}
 	//~ End UClothingAssetBase Interface
 
+	void CalculateLODHasAnySimulationMeshData();
+
 #if WITH_EDITOR
 	void OnModelChanged();
 
@@ -88,6 +96,9 @@ private:
 
 	UPROPERTY()
 	int32 ClothSimulationModelIndex = INDEX_NONE;
+
+	// Calculated from Asset during PostLoad or OnAssetChanged
+	TBitArray<> LODHasAnySimulationMeshData;
 };
 
 #undef UE_API
diff --git a/Engine/Plugins/ChaosClothAsset/Source/ChaosClothAssetEngine/Public/ChaosClothAsset/ClothAssetSKMClothingSimulation.h b/Engine/Plugins/ChaosClothAsset/Source/ChaosClothAssetEngine/Public/ChaosClothAsset/ClothAssetSKMClothingSimulation.h
--- a/Engine/Plugins/ChaosClothAsset/Source/ChaosClothAssetEngine/Public/ChaosClothAsset/ClothAssetSKMClothingSimulation.h
+++ b/Engine/Plugins/ChaosClothAsset/Source/ChaosClothAssetEngine/Public/ChaosClothAsset/ClothAssetSKMClothingSimulation.h
@@ -73,7 +73,7 @@ namespace UE::Chaos::ClothAsset
 		{
 			return new FSKMClothingSimulationContext();
 		}
-		virtual void FillContext(const USkeletalMeshComponent* InComponent, float InDeltaTime, IClothingSimulationContext* InOutContext, bool bIsInitialization) const override;
+		virtual void FillContextAndPrepareTick(const USkeletalMeshComponent* InComponent, float InDeltaTime, IClothingSimulationContext* InOutContext, bool bIsInitialization, bool bForceTeleportResetOnly) override;
 		virtual void DestroyContext(IClothingSimulationContext* Context) override
 		{
 			delete Context;
@@ -83,9 +83,11 @@ namespace UE::Chaos::ClothAsset
 		virtual void EndCreateActor() override;
 		virtual void DestroyActors() override;
 
-		virtual bool ShouldSimulate() const override;
+		virtual bool ShouldSimulateLOD(int32 OwnerLODIndex) const override;
 		virtual void Simulate(const IClothingSimulationContext* Context) override;
-		virtual void GetSimulationData(TMap<int32, FClothSimulData>& OutData, const USkeletalMeshComponent* OwnerComponent, const USkinnedMeshComponent* OverrideComponent) const override;
+		virtual void ForceClothNextUpdateTeleportAndReset() override;
+		virtual void HardResetSimulation(const IClothingSimulationContext* InContext) override;
+		virtual void AppendSimulationData(TMap<int32, FClothSimulData>& OutData, const USkeletalMeshComponent* OwnerComponent, const USkinnedMeshComponent* OverrideComponent) const override;
 
 		virtual FBoxSphereBounds GetBounds(const USkeletalMeshComponent* OwnerComponent) const override;
 
@@ -129,7 +131,7 @@ namespace UE::Chaos::ClothAsset
 		virtual float GetClothGeometryScale() const override;
 		virtual const TArray<TSharedPtr<const FManagedArrayCollection>>& GetPropertyCollections(const UChaosClothAssetBase* Asset, int32 ModelIndex) const override;
 		virtual const TArray<TSharedPtr<const FManagedArrayCollection>>& GetSolverPropertyCollections() const override;
-		virtual bool HasAnyRenderClothMappingData(int32 LodIndex) const override;
+		virtual bool HasAnySimulationMeshData(int32 LodIndex) const override;
 		//~ End IClothComponentAdapter Interface
 
 		//~ Begin FGCObject Interface
@@ -140,12 +142,16 @@ namespace UE::Chaos::ClothAsset
 		const FChaosClothSimulationProperties* GetSolverProperties() const;
 		const TArray<TSharedPtr<::Chaos::Softs::FCollectionPropertyFacade>>& GetSolverPropertyFacades() const;
 
+		void CalculateLODHasAnyRenderClothMappingData();
+
 		USkeletalMeshComponent* OwnerComponent = nullptr;
 		TUniquePtr<FCollisionSources> CollisionSources;
 		TSharedPtr<FClothSimulationProxy> ClothSimulationProxy;
 		
 		using FAssetClothSimulationProperties = TPair<const UChaosClothAssetSKMClothingAsset*, FChaosClothSimulationProperties>;
 		TArray<FAssetClothSimulationProperties> AssetClothSimulationProperties;  // Asset properties, matches CreateActor's SimDataIndex, some assets may appear multiple times, some entries might have a null asset
+		
+		TBitArray<> LODHasAnyRenderClothMappingData;
 	
 		bool bNeedsResetRestLengths = false;
 		FString RestLengthsMorphTargetName;
diff --git a/Engine/Plugins/ChaosClothAsset/Source/ChaosClothAssetEngine/Public/ChaosClothAsset/ClothComponent.h b/Engine/Plugins/ChaosClothAsset/Source/ChaosClothAssetEngine/Public/ChaosClothAsset/ClothComponent.h
--- a/Engine/Plugins/ChaosClothAsset/Source/ChaosClothAssetEngine/Public/ChaosClothAsset/ClothComponent.h
+++ b/Engine/Plugins/ChaosClothAsset/Source/ChaosClothAssetEngine/Public/ChaosClothAsset/ClothComponent.h
@@ -373,7 +373,7 @@ protected:
 		checkf(GetSimulationGroupId(InAsset, ModelIndex) != INDEX_NONE, TEXT("Invalid arguments"));
 		return GetPropertyCollections(ModelIndex);  // There's only this one asset for now, therefore this is the same as calling GetPropertyCollections(ModelIndex)
 	}
-	UE_API virtual bool HasAnyRenderClothMappingData(int32 LODIndex) const override;
+	UE_API virtual bool HasAnySimulationMeshData(int32 LODIndex) const override;
 	//~ End IClothComponentAdapter Interface
 
 	/** Override this function for setting up custom simulation proxies when the component is registered. */
diff --git a/Engine/Plugins/ChaosClothAsset/Source/ChaosClothAssetEngine/Public/ChaosClothAsset/ClothComponentAdapter.h b/Engine/Plugins/ChaosClothAsset/Source/ChaosClothAssetEngine/Public/ChaosClothAsset/ClothComponentAdapter.h
--- a/Engine/Plugins/ChaosClothAsset/Source/ChaosClothAssetEngine/Public/ChaosClothAsset/ClothComponentAdapter.h
+++ b/Engine/Plugins/ChaosClothAsset/Source/ChaosClothAssetEngine/Public/ChaosClothAsset/ClothComponentAdapter.h
@@ -60,8 +60,8 @@ namespace UE::Chaos::ClothAsset
 		/** Return the runtime solver property collections. */
 		virtual const TArray<TSharedPtr<const FManagedArrayCollection>>& GetSolverPropertyCollections() const = 0;
 
-		/** Return true when any of the rendered assets LOD sections has valid Cloth Mapping Data, or false otherwise. */
-		virtual bool HasAnyRenderClothMappingData(int32 LodIndex) const = 0;
+		/** Return true when any of the simulated assets LOD models have valid simulation data, or false otherwise. */
+		virtual bool HasAnySimulationMeshData(int32 LodIndex) const = 0;
 
 	protected:
 		virtual ~IClothComponentAdapter() = default;
diff --git a/Engine/Plugins/ChaosClothAsset/Source/ChaosClothAssetEngine/Public/ChaosClothAsset/ClothSimulationProxy.h b/Engine/Plugins/ChaosClothAsset/Source/ChaosClothAssetEngine/Public/ChaosClothAsset/ClothSimulationProxy.h
--- a/Engine/Plugins/ChaosClothAsset/Source/ChaosClothAssetEngine/Public/ChaosClothAsset/ClothSimulationProxy.h
+++ b/Engine/Plugins/ChaosClothAsset/Source/ChaosClothAssetEngine/Public/ChaosClothAsset/ClothSimulationProxy.h
@@ -60,7 +60,7 @@ namespace UE::Chaos::ClothAsset
 		UE_API virtual void PostConstructor();
 
 		/** Start the simulation if it isn't paused or suspended and return true, or simply update the existing simulation data and return false otherwise. */
-		UE_API bool Tick_GameThread(float DeltaTime, bool bUseParallelTask = true);
+		UE_API bool Tick_GameThread(float DeltaTime);
 
 		/** Wait for the parallel task to complete if one was running, and update the simulation data. */
 		UE_API void CompleteParallelSimulation_GameThread();
@@ -81,6 +81,12 @@ namespace UE::Chaos::ClothAsset
 		/** Post setup, required for when the simulation didn't run. For specialized usage when not calling Tick_GameThread. */
 		UE_API void PostProcess_GameThread();
 
+		/* Force a pending cloth reset. For specialized usage when not calling Tick_GameThread. */
+		UE_API void ForcePendingReset_GameThread();
+
+		/** Hard reset the simulation after config changes. */
+		UE_API void HardResetSimulation_GameThread();
+
 		/**
 		 * Return a map of all simulation data as used by the skeletal rendering code.
 		 * The map key is the rendering section's cloth index as set in FSkelMeshRenderSection::CorrespondClothAssetIndex,
@@ -141,6 +147,9 @@ namespace UE::Chaos::ClothAsset
 		// Internal physics thread object
 		friend class FClothSimulationProxyParallelTask;
 
+		// Allow FSKMClothingSimulation to call private methods like Tick()
+		friend class FSKMClothingSimulation;
+
 		// Reference for the cloth parallel task, to detect whether or not a simulation is running
 		FGraphEventRef ParallelTask;
 
diff --git a/Engine/Source/Runtime/ClothingSystemRuntimeCommon/Private/ClothingSimulation.cpp b/Engine/Source/Runtime/ClothingSystemRuntimeCommon/Private/ClothingSimulation.cpp
--- a/Engine/Source/Runtime/ClothingSystemRuntimeCommon/Private/ClothingSimulation.cpp
+++ b/Engine/Source/Runtime/ClothingSystemRuntimeCommon/Private/ClothingSimulation.cpp
@@ -203,7 +203,7 @@ FClothingSimulationCommon::FClothingSimulationCommon()
 FClothingSimulationCommon::~FClothingSimulationCommon()
 {}
 
-void FClothingSimulationCommon::FillContext(const USkeletalMeshComponent* InComponent, float InDeltaTime, IClothingSimulationContext* InOutContext, bool bIsInitialization) const
+void FClothingSimulationCommon::FillContextAndPrepareTick(const USkeletalMeshComponent* InComponent, float InDeltaTime, IClothingSimulationContext* InOutContext, bool bIsInitialization, bool bForceTeleportResetOnly)
 {
 	check(InOutContext);
 	FClothingSimulationContextCommon* const Context = static_cast<FClothingSimulationContextCommon*>(InOutContext);
diff --git a/Engine/Source/Runtime/ClothingSystemRuntimeCommon/Public/ClothingSimulation.h b/Engine/Source/Runtime/ClothingSystemRuntimeCommon/Public/ClothingSimulation.h
--- a/Engine/Source/Runtime/ClothingSystemRuntimeCommon/Public/ClothingSimulation.h
+++ b/Engine/Source/Runtime/ClothingSystemRuntimeCommon/Public/ClothingSimulation.h
@@ -94,7 +94,7 @@ public:
 
 protected:
 	/** Fills in the base data for a clothing simulation */
-	UE_API virtual void FillContext(const USkeletalMeshComponent* InComponent, float InDeltaTime, IClothingSimulationContext* InOutContext, bool bIsInitialization) const override;
+	UE_API virtual void FillContextAndPrepareTick(const USkeletalMeshComponent* InComponent, float InDeltaTime, IClothingSimulationContext* InOutContext, bool bIsInitialization, bool bForceTeleportResetOnly) override;
 
 protected:
 	/** Maximum physics time, incoming deltas will be clamped down to this value on long frames */
diff --git a/Engine/Source/Runtime/ClothingSystemRuntimeInterface/Private/ClothingSimulationInstance.cpp b/Engine/Source/Runtime/ClothingSystemRuntimeInterface/Private/ClothingSimulationInstance.cpp
--- a/Engine/Source/Runtime/ClothingSystemRuntimeInterface/Private/ClothingSimulationInstance.cpp
+++ b/Engine/Source/Runtime/ClothingSystemRuntimeInterface/Private/ClothingSimulationInstance.cpp
@@ -4,6 +4,7 @@
 #include "ClothCollisionData.h"
 #include "ClothingSimulationInteractor.h"
 #include "ClothingSimulationFactory.h"
+#include "Components/SkeletalMeshComponent.h"
 
 #include UE_INLINE_GENERATED_CPP_BY_NAME(ClothingSimulationInstance)
 
@@ -16,8 +17,8 @@ FClothingSimulationInstance::FClothingSimulationInstance(USkeletalMeshComponent*
 	check(ClothingSimulation);
 	ClothingSimulation->Initialize();
 
-	const bool bIsInitialization = true;
-	FillContext(SkeletalMeshComponent, 0.f, bIsInitialization);
+	constexpr bool bIsInitialization = true;
+	FillContextAndPrepareTick(SkeletalMeshComponent, 0.f, bIsInitialization);
 
 	for (int32 ClothingAssetIndex = 0; ClothingAssetIndex < FactoryAssetGroup.ClothingAssets.Num(); ++ClothingAssetIndex)
 	{
@@ -63,19 +64,24 @@ void FClothingSimulationInstance::RemoveAllClothingActors()
 	ClothingSimulation->DestroyActors();
 }
 
-void FClothingSimulationInstance::SyncClothingInteractor(bool bNeedsHardReset)
+void FClothingSimulationInstance::SyncClothingInteractor()
 {
-	if (ClothingSimulationInteractor)
+	if (ClothingSimulationInteractor && ClothingSimulation->ShouldSimulateLOD(CurrentOwnerLODIndex))
 	{
-		if (bNeedsHardReset)
-		{
-			ClothingSimulationInteractor->ClothConfigUpdated();
-		}
-
 		ClothingSimulationInteractor->Sync(ClothingSimulation, ClothingSimulationContext);
 	}
 }
 
+void FClothingSimulationInstance::FillContextAndPrepareTick(const USkeletalMeshComponent* OwnerComponent, float DeltaTime, bool bIsInitialization)
+{
+	if (!bIsInitialization)
+	{
+		CurrentOwnerLODIndex = OwnerComponent ? OwnerComponent->GetPredictedLODLevel() : INDEX_NONE;
+	}
+	const bool bForceTeleportResetOnly = !bIsInitialization && !ClothingSimulation->ShouldSimulateLOD(CurrentOwnerLODIndex);
+	ClothingSimulation->FillContextAndPrepareTick(OwnerComponent, DeltaTime, ClothingSimulationContext, bIsInitialization, bForceTeleportResetOnly);
+}
+
 void FClothingSimulationInstance::AppendUniqueCollisions(FClothCollisionData& OutCollisions, bool bIncludeExternal) const
 {
 	FClothCollisionData CollisionsToAppend;
diff --git a/Engine/Source/Runtime/ClothingSystemRuntimeInterface/Public/ClothingSimulationInstance.h b/Engine/Source/Runtime/ClothingSystemRuntimeInterface/Public/ClothingSimulationInstance.h
--- a/Engine/Source/Runtime/ClothingSystemRuntimeInterface/Public/ClothingSimulationInstance.h
+++ b/Engine/Source/Runtime/ClothingSystemRuntimeInterface/Public/ClothingSimulationInstance.h
@@ -110,7 +110,13 @@ public:
 	UE_API void RemoveAllClothingActors();
 
 	/** Synchronize the simulation with the changes requested by the interactor. */
-	UE_API void SyncClothingInteractor(bool bNeedsHardReset);
+	UE_API void SyncClothingInteractor();
+
+	/** Hard reset the cloth simulation */
+	void HardResetSimulation()
+	{
+		ClothingSimulation->HardResetSimulation(static_cast<const IClothingSimulationContext*>(ClothingSimulationContext));
+	}
 
 	/** Return the collision data from this simulation instance. */
 	void GetCollisions(FClothCollisionData& Collisions, bool bIncludeExternal) const
@@ -139,34 +145,40 @@ public:
 		ClothingSimulation->UpdateWorldForces(OwnerComponent);
 	}
 
-	/** Fill the simulation context. */
-	void FillContext(const USkeletalMeshComponent* OwnerComponent, float DeltaTime, bool bIsInitialization)
-	{
-		ClothingSimulation->FillContext(OwnerComponent, DeltaTime, ClothingSimulationContext, bIsInitialization);
-	}
+	/** Fill the simulation context and do any other related game-thread initialization. */
+	UE_API void FillContextAndPrepareTick(const USkeletalMeshComponent* OwnerComponent, float DeltaTime, bool bIsInitialization);
 
-	/** Advance the simulation. */
+	/** Advance the simulation. This method may be called asynchronously. */
 	void Simulate()
 	{
-		ClothingSimulation->Simulate(static_cast<const IClothingSimulationContext*>(ClothingSimulationContext));
+		if (ClothingSimulation->ShouldSimulateLOD(CurrentOwnerLODIndex))
+		{
+			ClothingSimulation->Simulate(static_cast<const IClothingSimulationContext*>(ClothingSimulationContext));
+		}
+		else
+		{
+			ClothingSimulation->ForceClothNextUpdateTeleportAndReset();
+		}
 	}
 
-	/** Return the simulation data. */
-	void GetSimulationData(TMap<int32, FClothSimulData>& CurrentSimulationData, const USkeletalMeshComponent* OwnerComponent, const USkinnedMeshComponent* LeaderPoseComponent) const
+	/** Append the simulation data. */
+	void AppendSimulationData(TMap<int32, FClothSimulData>& CurrentSimulationData, const USkeletalMeshComponent* OwnerComponent, const USkinnedMeshComponent* LeaderPoseComponent) const
 	{
-		ClothingSimulation->GetSimulationData(CurrentSimulationData, OwnerComponent, LeaderPoseComponent);
+		if (ClothingSimulation->ShouldSimulateLOD(CurrentOwnerLODIndex))
+		{
+			ClothingSimulation->AppendSimulationData(CurrentSimulationData, OwnerComponent, LeaderPoseComponent);
+		}
 	}
 
-	/** Return whether this simulation has anything to simulate. */
-	bool ShouldSimulate() const
-	{
-		return ClothingSimulation->ShouldSimulate();
-	}
 
 	/** Return this simulation bounds. */
 	FBoxSphereBounds GetBounds(const USkeletalMeshComponent* OwnerComponent) const
 	{
-		return ClothingSimulation->GetBounds(OwnerComponent);
+		if (ClothingSimulation->ShouldSimulateLOD(CurrentOwnerLODIndex))
+		{
+			return ClothingSimulation->GetBounds(OwnerComponent);
+		}
+		return FBoxSphereBounds(ForceInit);
 	}
 
 private:
@@ -187,6 +199,8 @@ private:
 	/** Simulation factory. */
 	UPROPERTY(Transient)
 	TObjectPtr<UClothingSimulationFactory> ClothingSimulationFactory;
+
+	int32 CurrentOwnerLODIndex = INDEX_NONE; // This is populated by FillContextAndPrepareTick
 };
 
 template<>
diff --git a/Engine/Source/Runtime/ClothingSystemRuntimeInterface/Public/ClothingSimulationInterface.h b/Engine/Source/Runtime/ClothingSystemRuntimeInterface/Public/ClothingSimulationInterface.h
--- a/Engine/Source/Runtime/ClothingSystemRuntimeInterface/Public/ClothingSimulationInterface.h
+++ b/Engine/Source/Runtime/ClothingSystemRuntimeInterface/Public/ClothingSimulationInterface.h
@@ -67,26 +67,28 @@ PRAGMA_ENABLE_DEPRECATION_WARNINGS
 	virtual IClothingSimulationContext* CreateContext() = 0;
 
 	/**
-	 * Fills an existing context for a single simulation step, called by the engine on the game thread prior to simulation 
+	 * Fills an existing context for a single simulation step and do any other work that needs to be called by the engine on the game thread prior to simulation.
 	 * @param InComponent - The component to fill the context for
 	 * @param InOutContext - The context to fill
 	 * @param bIsInitialization - Whether this fill is occurring as part of the actor creation stage
+	 * @param bForceTeleportResetOnly - Whether ForceClothNextUpdateTeleportAndReset instead of Simulate and AppendSimulationData will be called.
 	 */
-	virtual void FillContext(const USkeletalMeshComponent* InComponent, float InDeltaTime, IClothingSimulationContext* InOutContext, bool bIsInitialization) const
+	// TODO: make this pure virtual in the new interface.
+	virtual void FillContextAndPrepareTick(const USkeletalMeshComponent* InComponent, float InDeltaTime, IClothingSimulationContext* InOutContext, bool bIsInitialization, bool bForceTeleportResetOnly)
 	{
 PRAGMA_DISABLE_DEPRECATION_WARNINGS
-		return const_cast<IClothingSimulation*>(this)->FillContext(InComponent, InDeltaTime, InOutContext, bIsInitialization);
+		return FillContext(const_cast<USkeletalMeshComponent*>(InComponent), InDeltaTime, InOutContext, bIsInitialization);
 PRAGMA_ENABLE_DEPRECATION_WARNINGS
 	}
 
-	UE_DEPRECATED(5.7, "Use const version instead.")
+	UE_DEPRECATED(5.7, "Use FillContextAndPrepareTick instead.")
 	virtual void FillContext(USkeletalMeshComponent* InComponent, float InDeltaTime, IClothingSimulationContext* InOutContext, bool bIsInitialization)
 	{
 PRAGMA_DISABLE_DEPRECATION_WARNINGS
 		return FillContext(InComponent, InDeltaTime, InOutContext);
 PRAGMA_ENABLE_DEPRECATION_WARNINGS
 	}
-	UE_DEPRECATED(5.7, "Use const version instead.")
+	UE_DEPRECATED(5.7, "Use FillContextAndPrepareTick instead.")
 	virtual void FillContext(USkeletalMeshComponent* InComponent, float InDeltaTime, IClothingSimulationContext* InOutContext) {}
 
 	/** Initialize the simulation, will be called before any Simulate calls */
@@ -102,15 +104,31 @@ PRAGMA_ENABLE_DEPRECATION_WARNINGS
 	 * Called by the engine to detect whether or not the simulation should run, essentially
 	 * are there any actors that need to simulate in this simulation
 	 */
-	virtual bool ShouldSimulate() const = 0;
+	UE_DEPRECATED(5.7, "This method has not been used consistently and has been replaced with ShouldSimulateLOD")
+	virtual bool ShouldSimulate() const
+	{
+		return true;
+	}
+
+	/** Called by the engine to determine if this simulation can run this tick at this LOD.
+	 *  ForceClothNextUpdateTeleportAndReset rather than Simulate and AppendSimulationData will be called this tick.
+	 */
+	 // TODO: make this pure virtual in the new interface.
+	virtual bool ShouldSimulateLOD(int32 OwnerLODIndex) const
+	{
+PRAGMA_DISABLE_DEPRECATION_WARNINGS
+		return ShouldSimulate();
+PRAGMA_ENABLE_DEPRECATION_WARNINGS
+	}
 
 	/**
 	 * Run a single tick of the simulation.
 	 * The pointer InContext is guaranteed (when called by the engine) to be the context allocated in
 	 * CreateContext and can be assumed to be safely castable to any derived type allocated there.
 	 * New callers should make care to make sure only the correct context is ever passed through
-	 * @param InContext - The context to use during simulation, will have been filled in FillContext
+	 * @param InContext - The context to use during simulation, will have been filled in FillContextAndPrepareTick
 	 */
+	 // TODO: make this pure virtual in the new interface. Indicate that this can be called asynchronously
 	virtual void Simulate(const IClothingSimulationContext* InContext)
 	{
 PRAGMA_DISABLE_DEPRECATION_WARNINGS
@@ -121,6 +139,14 @@ PRAGMA_ENABLE_DEPRECATION_WARNINGS
 	UE_DEPRECATED(5.7, "Use const argument version instead,")
 	virtual void Simulate(IClothingSimulationContext* InContext) {}
 
+	/** Called instead of Simulate when not ticking.*/
+	 // TODO: make this pure virtual in the new interface. Indicate that this can be called asynchronously
+	virtual void ForceClothNextUpdateTeleportAndReset() {};
+
+	/** Hard reset the simulation without necessarily recreating cloth actors. */
+	// TODO: make this pure virtual in the new interface.
+	virtual void HardResetSimulation(const IClothingSimulationContext* InContext) {};
+
 	/** Simulation should remove all of it's actors when next possible and free them */
 	virtual void DestroyActors() = 0;
 
@@ -131,19 +157,21 @@ PRAGMA_ENABLE_DEPRECATION_WARNINGS
 	virtual void DestroyContext(IClothingSimulationContext* InContext) = 0;
 
 	/** 
-	 * Fill FClothSimulData map for the clothing simulation. Should fill a map pair per-actor 
-	 * @param OutData - The simulation data to write to
+	 * Fill FClothSimulData map for the clothing simulation. Should fill a map pair per-actor. 
+	 * Do not remove InOutData elements which you are not responsible for: they may have been written by a different instance.
+	 * @param InOutData - The simulation data to write to
 	 * @param OwnerComponent - the component that owns the simulation
 	 * @param OverrideComponent - An override component if bound to a leader pose component
 	 */
-	virtual void GetSimulationData(TMap<int32, FClothSimulData>& OutData, const USkeletalMeshComponent* OwnerComponent, const USkinnedMeshComponent* OverrideComponent) const
+	// TODO: make this pure virtual in the new interface.
+	virtual void AppendSimulationData(TMap<int32, FClothSimulData>& InOutData, const USkeletalMeshComponent* OwnerComponent, const USkinnedMeshComponent* OverrideComponent) const 
 	{
 PRAGMA_DISABLE_DEPRECATION_WARNINGS
-		return GetSimulationData(OutData, const_cast<USkeletalMeshComponent*>(OwnerComponent), const_cast<USkinnedMeshComponent*>(OverrideComponent));
+		GetSimulationData(InOutData, const_cast<USkeletalMeshComponent*>(OwnerComponent), const_cast<USkinnedMeshComponent*>(OverrideComponent));
 PRAGMA_ENABLE_DEPRECATION_WARNINGS
 	}
 
-	UE_DEPRECATED(5.7, "Use const version instead.")
+	UE_DEPRECATED(5.7, "Use AppendSimulationData instead.")
 	virtual void GetSimulationData(TMap<int32, FClothSimulData>& OutData, USkeletalMeshComponent* InOwnerComponent, USkinnedMeshComponent* InOverrideComponent) const {}
 
 	/** Get the bounds of the simulation mesh in local simulation space */
diff --git a/Engine/Source/Runtime/ClothingSystemRuntimeInterface/Public/ClothingSystemRuntimeTypes.h b/Engine/Source/Runtime/ClothingSystemRuntimeInterface/Public/ClothingSystemRuntimeTypes.h
--- a/Engine/Source/Runtime/ClothingSystemRuntimeInterface/Public/ClothingSystemRuntimeTypes.h
+++ b/Engine/Source/Runtime/ClothingSystemRuntimeInterface/Public/ClothingSystemRuntimeTypes.h
@@ -20,6 +20,11 @@ struct FClothSimulData
 		LODIndex = INDEX_NONE;
 	}
 
+	bool IsReset() const
+	{
+		return Positions.IsEmpty() && Normals.IsEmpty() && LODIndex == INDEX_NONE;
+	}
+
 	// Positions of the simulation mesh particles (aligned for SIMD loads)
 	TArray<FVector3f, TAlignedHeapAllocator<16>> Positions;
 
diff --git a/Engine/Source/Runtime/Engine/Classes/Components/SkeletalMeshComponent.h b/Engine/Source/Runtime/Engine/Classes/Components/SkeletalMeshComponent.h
--- a/Engine/Source/Runtime/Engine/Classes/Components/SkeletalMeshComponent.h
+++ b/Engine/Source/Runtime/Engine/Classes/Components/SkeletalMeshComponent.h
@@ -1688,10 +1688,6 @@ public:
 	/** Tick Clothing Animation , basically this is called inside TickComponent */
 	ENGINE_API void TickClothing(float DeltaTime, FTickFunction& ThisTickFunction);
 
-	/** Store cloth simulation data into OutClothSimData */
-	UE_DEPRECATED(5.2, "Use GetUpdateClothSimulationData_AnyThread instead.")
-	ENGINE_API void GetUpdateClothSimulationData(TMap<int32, FClothSimulData>& OutClothSimData, USkeletalMeshComponent* OverrideLocalRootComponent = nullptr);
-
 	/** Store cloth simulation data into OutClothSimulData. Override USkinnedMeshComponent. */
 	ENGINE_API virtual void GetUpdateClothSimulationData_AnyThread(TMap<int32, FClothSimulData>& OutClothSimulData, FMatrix& OutLocalToWorld, float& OutClothBlendWeight) const override;
 
diff --git a/Engine/Source/Runtime/Engine/Private/Components/SkeletalMeshComponent.cpp b/Engine/Source/Runtime/Engine/Private/Components/SkeletalMeshComponent.cpp
--- a/Engine/Source/Runtime/Engine/Private/Components/SkeletalMeshComponent.cpp
+++ b/Engine/Source/Runtime/Engine/Private/Components/SkeletalMeshComponent.cpp
@@ -793,7 +793,12 @@ bool USkeletalMeshComponent::ShouldRunClothTick() const
 		return false;
 	}
 
-	if(CanSimulateClothing())
+	if(!CanSimulateClothing())
+	{
+		return false;
+	}
+
+	if (GetSkeletalMeshAsset()->HasActiveClothingAssetsForLOD(GetPredictedLODLevel()))
 	{
 		return true;
 	}
@@ -1865,38 +1870,41 @@ void USkeletalMeshComponent::TickComponent(float DeltaTime, enum ELevelTick Tick
 	}
 #endif
 
+	Super::TickComponent(DeltaTime, TickType, ThisTickFunction);
+
 	if(!PrimaryComponentTick.bRunOnAnyThread)
 	{
-		for (FClothingSimulationInstance& ClothingSimulationInstance : ClothingSimulationInstances)
-		{
-			ClothingSimulationInstance.UpdateWorldForces(this);
-		}
+		// Cloth work must be done after Super::TickComponent so that it has the correct PredictedLOD
 
 		UpdateEndPhysicsTickRegisteredState();
 		UpdateClothTickRegisteredState();
-
-		// If we are suspended, we will not simulate clothing, but as clothing is simulated in local space
-		// relative to a root bone we need to extract simulation positions as this bone could be animated.
-		if (!CVarEnableClothPhysics.GetValueOnGameThread() || bClothingSimulationSuspended)
+		if(!ClothTickFunction.IsTickFunctionRegistered())
 		{
+			// Do any cloth work that needs to be done when it doesn't tick.
 			CSV_SCOPED_TIMING_STAT(Animation, Cloth);
+			const int32 PredictedLOD = GetPredictedLODLevel();
 
-			for (FClothingSimulationInstance& ClothingSimulationInstance : ClothingSimulationInstances)
+			// If we are suspended, we will not simulate clothing, but as clothing is simulated in local space
+			// relative to a root bone we need to extract simulation positions as this bone could be animated.
+			if (!CVarEnableClothPhysics.GetValueOnGameThread() || bClothingSimulationSuspended)
 			{
-				// First update the simulation context, since the simulation isn't ticking
-				// and it is still required to get the correct simulation data and bounds.
-				constexpr bool bIsInitialization = false;
-				ClothingSimulationInstance.FillContext(this, DeltaTime, bIsInitialization);
+				for (FClothingSimulationInstance& ClothingSimulationInstance : ClothingSimulationInstances)
+				{
+					// First update the simulation context, since the simulation isn't ticking
+					// and it is still required to get the correct simulation data and bounds.
+					constexpr bool bIsInitialization = false;
+					ClothingSimulationInstance.FillContextAndPrepareTick(this, DeltaTime, bIsInitialization);
+				}
 
-				ClothingSimulationInstance.GetSimulationData(CurrentSimulationData, this, Cast<USkeletalMeshComponent>(LeaderPoseComponent.Get()));
+				WritebackClothingSimulationData();
+			}
+			else
+			{
+				// Since cloth isn't ticking, force a teleport and reset the next time it does.
+				ForceClothNextUpdateTeleportAndReset();
 			}
 		}
-	}
 
-	Super::TickComponent(DeltaTime, TickType, ThisTickFunction);
-
-	if(!PrimaryComponentTick.bRunOnAnyThread)
-	{
 		PendingRadialForces.Reset();
 
 		// Update bOldForceRefPose
@@ -2513,11 +2521,17 @@ void USkeletalMeshComponent::UpdateClothSimulationContext(float InDeltaTime)
 	}
 
 	// Fill the contexts for the next simulation
+	const int32 PredictedLOD = GetPredictedLODLevel();
 	for (FClothingSimulationInstance& ClothingSimulationInstance : ClothingSimulationInstances)
 	{
-		const bool bIsInitialization = false;
-		ClothingSimulationInstance.FillContext(this, InDeltaTime, bIsInitialization);
-		ClothingSimulationInstance.SyncClothingInteractor(bNeedsHardReset);
+		constexpr bool bIsInitialization = false;
+		ClothingSimulationInstance.FillContextAndPrepareTick(this, InDeltaTime, bIsInitialization);
+		if (bNeedsHardReset)
+		{
+			ClothingSimulationInstance.HardResetSimulation();
+		}
+
+		ClothingSimulationInstance.SyncClothingInteractor();
 	}
 
 	PendingTeleportType = ETeleportType::None;
@@ -2565,9 +2579,23 @@ void USkeletalMeshComponent::WritebackClothingSimulationData()
 			}
 		}
 
+		// Reset all existing CurrentSimulationData so that the ClothingSimulationInstances can append their data.
+		for (TMap<int32, FClothSimulData>::TIterator DataIter = CurrentSimulationData.CreateIterator(); DataIter; ++DataIter)
+		{
+			DataIter.Value().Reset();
+		}
+
 		for (const FClothingSimulationInstance& ClothingSimulationInstance : ClothingSimulationInstances)
 		{
-			ClothingSimulationInstance.GetSimulationData(CurrentSimulationData, this, OverrideComponent);
+			ClothingSimulationInstance.AppendSimulationData(CurrentSimulationData, this, OverrideComponent);
+		}
+
+		for (TMap<int32, FClothSimulData>::TIterator DataIter = CurrentSimulationData.CreateIterator(); DataIter; ++DataIter)
+		{
+			if (DataIter.Value().IsReset())
+			{
+				DataIter.RemoveCurrent();
+			}
 		}
 	}
 }
diff --git a/Engine/Source/Runtime/Engine/Private/SkeletalMeshComponentPhysics.cpp b/Engine/Source/Runtime/Engine/Private/SkeletalMeshComponentPhysics.cpp
--- a/Engine/Source/Runtime/Engine/Private/SkeletalMeshComponentPhysics.cpp
+++ b/Engine/Source/Runtime/Engine/Private/SkeletalMeshComponentPhysics.cpp
@@ -2906,12 +2906,10 @@ bool USkeletalMeshComponent::ComponentOverlapMultiImpl(TArray<struct FOverlapRes
 
 void USkeletalMeshComponent::AddClothingBounds(FBoxSphereBounds& InOutBounds, const FTransform& LocalToWorld) const
 {
+	const int32 PredictedLOD = GetPredictedLODLevel();
 	for (const FClothingSimulationInstance& ClothingSimulationInstance : ClothingSimulationInstances)
 	{
-		if (ClothingSimulationInstance.ShouldSimulate())
-		{
-			InOutBounds = InOutBounds + ClothingSimulationInstance.GetBounds(this).TransformBy(LocalToWorld);
-		}
+		InOutBounds = InOutBounds + ClothingSimulationInstance.GetBounds(this).TransformBy(LocalToWorld);
 	}
 }
 
@@ -3097,11 +3095,8 @@ void USkeletalMeshComponent::CopyClothCollisionsToChildren()
 		{
 			for (const FClothingSimulationInstance& ClothingSimulationInstance : Child->ClothingSimulationInstances)
 			{
-				if (ClothingSimulationInstance.ShouldSimulate())
-				{
-					ClothChildren.Add(Child);
-					break;
-				}
+				ClothChildren.Add(Child);
+				break;
 			}
 		}
 	}
@@ -3286,11 +3281,8 @@ void USkeletalMeshComponent::UpdateClothTransformImp()
 	{
 		for (FClothingSimulationInstance& ClothingSimulationInstance : ClothingSimulationInstances)
 		{
-			if (ClothingSimulationInstance.ShouldSimulate())
-			{
-				// Component has moved, and there is something to simulate, update all external cloth collisions
-				ClothingSimulationInstance.ClearExternalCollisions();
-			}
+			// Component has moved, and there is something to simulate, update all external cloth collisions
+			ClothingSimulationInstance.ClearExternalCollisions();
 		}
 
 		if (bCollideWithAttachedChildren)
@@ -3573,6 +3565,12 @@ void USkeletalMeshComponent::UpdateClothStateAndSimulate(float DeltaTime, FTickF
 	}
 #endif // WITH_CLOTH_COLLISION_DETECTION
 
+
+	for (FClothingSimulationInstance& ClothingSimulationInstance : ClothingSimulationInstances)
+	{
+		ClothingSimulationInstance.UpdateWorldForces(this);
+	}
+
 	UpdateClothSimulationContext(DeltaTime);
 
 	ParallelClothTask = TGraphTask<FParallelClothTask>::CreateTask(nullptr, ENamedThreads::GameThread).ConstructAndDispatchWhenReady(*this, DeltaTime);
@@ -3644,28 +3642,6 @@ void USkeletalMeshComponent::TickClothing(float DeltaTime, FTickFunction& ThisTi
 	}
 }
 
-void USkeletalMeshComponent::GetUpdateClothSimulationData(TMap<int32, FClothSimulData>& OutClothSimData, USkeletalMeshComponent* OverrideLocalRootComponent)
-{
-	if(CVarEnableClothPhysics.GetValueOnAnyThread() == 0)
-	{
-		return;
-	}
-
-	SCOPE_CYCLE_COUNTER(STAT_ClothTotalTime);
-	CSV_SCOPED_TIMING_STAT(Animation, Cloth);
-
-	if(bDisableClothSimulation)
-	{
-		OutClothSimData.Reset();
-		return;
-	}
-
-	for (const FClothingSimulationInstance& ClothingSimulationInstance : ClothingSimulationInstances)
-	{
-		ClothingSimulationInstance.GetSimulationData(OutClothSimData, this, OverrideLocalRootComponent);
-	}
-}
-
 void USkeletalMeshComponent::GetUpdateClothSimulationData_AnyThread(TMap<int32, FClothSimulData>& OutClothSimulData, FMatrix& OutLocalToWorld, float& OutClothBlendWeight) const
 {
 	OutLocalToWorld = GetComponentToWorld().ToMatrixWithScale();
diff --git a/Engine/Source/Runtime/Engine/Private/SkeletalRenderCPUSkin.cpp b/Engine/Source/Runtime/Engine/Private/SkeletalRenderCPUSkin.cpp
--- a/Engine/Source/Runtime/Engine/Private/SkeletalRenderCPUSkin.cpp
+++ b/Engine/Source/Runtime/Engine/Private/SkeletalRenderCPUSkin.cpp
@@ -1147,43 +1147,5 @@ static void CalculateMorphTargetWeights(FFinalSkinVertex* DestVertex, FSkeletalM
 	}
 }
 
-bool FDynamicSkelMeshObjectDataCPUSkin::UpdateClothSimulationData(USkinnedMeshComponent* InMeshComponent)
-{
-	USkeletalMeshComponent* SimMeshComponent = Cast<USkeletalMeshComponent>(InMeshComponent);
-
-	if (InMeshComponent->LeaderPoseComponent.IsValid() && (SimMeshComponent && SimMeshComponent->IsClothBoundToLeaderComponent()))
-	{
-		USkeletalMeshComponent* SrcComponent = SimMeshComponent;
-
-		// if I have Leader, override sim component
-		SimMeshComponent = Cast<USkeletalMeshComponent>(InMeshComponent->LeaderPoseComponent.Get());
-
-		// IF we don't have sim component that is skeletalmeshcomponent, just ignore
-		if (!SimMeshComponent)
-		{
-			return false;
-		}
-
-		WorldToLocal = SrcComponent->GetRenderMatrix().InverseFast();
-		ClothBlendWeight = IsSkeletalMeshClothBlendEnabled() ? SrcComponent->ClothBlendWeight : 0.0f;
-PRAGMA_DISABLE_DEPRECATION_WARNINGS
-		SimMeshComponent->GetUpdateClothSimulationData(ClothSimulUpdateData, SrcComponent);
-PRAGMA_ENABLE_DEPRECATION_WARNINGS
-
-		return true;
-	}
-
-	if (SimMeshComponent)
-	{
-		WorldToLocal = SimMeshComponent->GetRenderMatrix().InverseFast();
-		ClothBlendWeight = IsSkeletalMeshClothBlendEnabled() ? SimMeshComponent->ClothBlendWeight : 0.0f;
-PRAGMA_DISABLE_DEPRECATION_WARNINGS
-		SimMeshComponent->GetUpdateClothSimulationData(ClothSimulUpdateData);
-PRAGMA_ENABLE_DEPRECATION_WARNINGS
-		return true;
-	}
-	return false;
-}
-
 
 MSVC_PRAGMA(warning(pop))
\ No newline at end of file
diff --git a/Engine/Source/Runtime/Engine/Private/SkeletalRenderCPUSkin.h b/Engine/Source/Runtime/Engine/Private/SkeletalRenderCPUSkin.h
--- a/Engine/Source/Runtime/Engine/Private/SkeletalRenderCPUSkin.h
+++ b/Engine/Source/Runtime/Engine/Private/SkeletalRenderCPUSkin.h
@@ -78,10 +78,6 @@ public:
 		CumulativeResourceSize.AddDedicatedSystemMemoryBytes(ReferenceToLocal.GetAllocatedSize());
 		CumulativeResourceSize.AddDedicatedSystemMemoryBytes(ActiveMorphTargets.GetAllocatedSize());
  	}
-
-	/** Update Simulated Positions & Normals from Clothing actor */
-	UE_DEPRECATED(5.2, "Use USkinnedMeshComponent::GetUpdateClothSimulationData_AnyThread() instead.")
-	ENGINE_API bool UpdateClothSimulationData(USkinnedMeshComponent* InMeshComponent);
 };
 
 /**
