Add MinLodSize optimization to SkinnedTriangleMesh physics assets.

diff --git a/Engine/Plugins/ChaosCloth/Source/ChaosCloth/Private/ChaosCloth/ChaosClothComplexColliders.cpp b/Engine/Plugins/ChaosCloth/Source/ChaosCloth/Private/ChaosCloth/ChaosClothComplexColliders.cpp
--- a/Engine/Plugins/ChaosCloth/Source/ChaosCloth/Private/ChaosCloth/ChaosClothComplexColliders.cpp
+++ b/Engine/Plugins/ChaosCloth/Source/ChaosCloth/Private/ChaosCloth/ChaosClothComplexColliders.cpp
@@ -150,7 +150,7 @@ namespace Chaos {
 				SubBoneTransforms.Emplace(TRigidTransform<FReal, 3>(CollisionSubBones.X[SubBoneIndex], CollisionSubBones.R[SubBoneIndex]) * RootTransformInv);
 			}
 			SkinnedLevelSet.DeformPoints(SubBoneTransforms);
-			SkinnedLevelSet.UpdateSpatialHierarchy();
+			SkinnedLevelSet.UpdateSpatialHierarchy(MinLODSize);
 		}
 
 		// Update MLLevelSets
@@ -186,7 +186,7 @@ namespace Chaos {
 					InterpolatedPositions[VtxIndex] = Frame.TransformPositionNoScale(Alpha * Data.SkinnedPositions.Positions[VtxIndex] + OneMinusAlpha * Data.SkinnedPositions.OldPositions[VtxIndex]);
 				}
 				SkinnedTriangleMesh.UpdateLocalBoundingBox();
-				SkinnedTriangleMesh.UpdateSpatialHierarchy();
+				SkinnedTriangleMesh.UpdateSpatialHierarchy(MinLODSize);
 			}
 		}
 	}
diff --git a/Engine/Plugins/ChaosCloth/Source/ChaosCloth/Private/ChaosCloth/ChaosClothingSimulationCloth.cpp b/Engine/Plugins/ChaosCloth/Source/ChaosCloth/Private/ChaosCloth/ChaosClothingSimulationCloth.cpp
--- a/Engine/Plugins/ChaosCloth/Source/ChaosCloth/Private/ChaosCloth/ChaosClothingSimulationCloth.cpp
+++ b/Engine/Plugins/ChaosCloth/Source/ChaosCloth/Private/ChaosCloth/ChaosClothingSimulationCloth.cpp
@@ -8,6 +8,7 @@
 #include "ChaosCloth/ChaosClothingSimulationConfig.h"
 #include "ChaosCloth/ChaosClothingPatternData.h"
 #include "ChaosCloth/ChaosWeightMapTarget.h"
+#include "ChaosCloth/ChaosClothComplexColliders.h"
 #include "ChaosCloth/ChaosClothPrivate.h"
 #include "Chaos/CollectionPropertyFacade.h"
 #include "Chaos/PBDFlatWeightMap.h"
@@ -436,6 +437,24 @@ void FClothingSimulationCloth::FLODData::Update(FClothingSimulationSolver* Solve
 		const Softs::FSolverReal MeshScale = Cloth->Mesh->GetScale() * LocalSpaceScaleInv;
 		const Softs::FSolverReal MaxDistancesScale = (Softs::FSolverReal)Cloth->MaxDistancesMultiplier;
 		ClothConstraints.Update(Cloth->Config->GetProperties(SolverDatum.LODIndex), WeightMaps, VertexSets, FaceSets, FaceIntMaps, MeshScale, MaxDistancesScale, Solver->GetLocalSpaceRotation(), Cloth->ReferenceSpaceTransform.GetRotation());
+
+		if (const TSharedPtr<Softs::FPBDSkinnedTriangleMeshCollisionConstraints>& SkinnedTriangleCollisionConstraint = ClothConstraints.GetSkinnedTriangleCollisionsConstraint())
+		{
+			const Softs::FSolverReal MinLODSize = SkinnedTriangleCollisionConstraint->GetCollisionThickness() + SkinnedTriangleCollisionConstraint->GetMaxClothCollisionThickness();
+
+			for (FClothingSimulationCollider* const Collider : Cloth->GetColliders())
+			{
+				check(Collider);
+				for (FClothingSimulationCollider::ECollisionDataType CollisionType : TEnumRange<FClothingSimulationCollider::ECollisionDataType>())
+				{
+					const int32 CollisionRangeId = Collider->GetCollisionRangeId(Solver, Cloth, CollisionType);
+					if (FClothComplexColliders* const ComplexCollider = Solver->GetComplexColliders(CollisionRangeId))
+					{
+						ComplexCollider->SetMinLODSize(MinLODSize);
+					}
+				}
+			}
+		}
 	}
 }
 
diff --git a/Engine/Plugins/ChaosCloth/Source/ChaosCloth/Public/ChaosCloth/ChaosClothComplexColliders.h b/Engine/Plugins/ChaosCloth/Source/ChaosCloth/Public/ChaosCloth/ChaosClothComplexColliders.h
--- a/Engine/Plugins/ChaosCloth/Source/ChaosCloth/Public/ChaosCloth/ChaosClothComplexColliders.h
+++ b/Engine/Plugins/ChaosCloth/Source/ChaosCloth/Public/ChaosCloth/ChaosClothComplexColliders.h
@@ -45,9 +45,13 @@ namespace Chaos
 		void AddSkinnedLevelSet(int32 Index, const TArray<int32>& MappedSubBones, const FImplicitObjectPtr& SkinnedLevelSet);
 		void AddMLLevelSet(int32 Index, const TArray<int32>& MappedSubBones, const FImplicitObjectPtr& MLLevelSet);
 		void AddSkinnedTriangleMesh(int32 Index, const TArray<int32>& MappedSubBones, const FImplicitObjectPtr& SkinnedTriangleMesh);
-
 		void Update(const Softs::FSolverTransform3& ComponentToLocalSpace, const TArray<FTransform>& BoneTransforms, const TConstArrayView<Softs::FSolverRigidTransform3>& CollisionRangeTransforms);
 		void ResetStartPose();
+
+		void SetMinLODSize(const Softs::FSolverReal InMinLODSize)
+		{
+			MinLODSize = InMinLODSize;
+		}
 		// ---- End of Collider interface ----
 
 		// -- Constraints interface ----
@@ -153,6 +157,7 @@ namespace Chaos
 		};
 		TArray<FSkinnedTriangleMeshData> SkinnedTriangleMeshes;
 		bool bSkipSkinnedTriangleMeshKinematicUpdate = false;
+		Softs::FSolverReal MinLODSize = 0.f;
 
 	};
 } // namespace Chaos
diff --git a/Engine/Source/Runtime/Experimental/Chaos/Private/Chaos/WeightedLatticeImplicitObject.cpp b/Engine/Source/Runtime/Experimental/Chaos/Private/Chaos/WeightedLatticeImplicitObject.cpp
--- a/Engine/Source/Runtime/Experimental/Chaos/Private/Chaos/WeightedLatticeImplicitObject.cpp
+++ b/Engine/Source/Runtime/Experimental/Chaos/Private/Chaos/WeightedLatticeImplicitObject.cpp
@@ -361,7 +361,7 @@ struct FWeightedLatticeBvEntry
 	}
 };
 
-void FWeightedLatticeImplicitObject::UpdateSpatialHierarchy()
+void FWeightedLatticeImplicitObject::UpdateSpatialHierarchy(const FReal MinLodSize)
 {
 	if (!bSpatialDirty)
 	{
@@ -384,7 +384,7 @@ void FWeightedLatticeImplicitObject::UpdateSpatialHierarchy()
 		}
 	}
 
-	Spatial.Initialize(Entries);
+	Spatial.Initialize(Entries, MinLodSize);
 
 	bSpatialDirty = false;
 }
diff --git a/Engine/Source/Runtime/Experimental/Chaos/Public/Chaos/PBDKinematicTriangleMeshCollisions.h b/Engine/Source/Runtime/Experimental/Chaos/Public/Chaos/PBDKinematicTriangleMeshCollisions.h
--- a/Engine/Source/Runtime/Experimental/Chaos/Public/Chaos/PBDKinematicTriangleMeshCollisions.h
+++ b/Engine/Source/Runtime/Experimental/Chaos/Public/Chaos/PBDKinematicTriangleMeshCollisions.h
@@ -257,6 +257,16 @@ namespace Chaos::Softs
 			return bEnableSkinnedTriangleMeshCollisions;
 		}
 
+		FSolverReal GetCollisionThickness() const
+		{
+			return CollisionThickness;
+		}
+
+		FSolverReal GetMaxClothCollisionThickness() const
+		{
+			return FMath::Max(ClothCollisionThickness.GetLow(), ClothCollisionThickness.GetHigh());
+		}
+
 	private:
 		int32 NumParticles;
 		FPBDFlatWeightMap ClothCollisionThickness;
diff --git a/Engine/Source/Runtime/Experimental/Chaos/Public/Chaos/WeightedLatticeImplicitObject.h b/Engine/Source/Runtime/Experimental/Chaos/Public/Chaos/WeightedLatticeImplicitObject.h
--- a/Engine/Source/Runtime/Experimental/Chaos/Public/Chaos/WeightedLatticeImplicitObject.h
+++ b/Engine/Source/Runtime/Experimental/Chaos/Public/Chaos/WeightedLatticeImplicitObject.h
@@ -87,7 +87,7 @@ public:
 	CHAOS_API virtual void Serialize(FChaosArchive& Ar) override;
 
 	CHAOS_API FVec3 GetDeformedPoint(const FVec3& UndeformedPoint) const;
-	CHAOS_API void UpdateSpatialHierarchy();
+	CHAOS_API void UpdateSpatialHierarchy(const FReal MinLodSize = 0.);
 	CHAOS_API void FinalizeConstruction();
 	CHAOS_API void DeformPoints(const TArray<FTransform>& RelativeTransforms);
 
