Add WindVelocity cloth debug draw. Update gravity vector debug draw to use an arrow instead of a line.

diff --git a/Engine/Plugins/ChaosCloth/Source/ChaosCloth/Private/ChaosCloth/ChaosClothVisualization.cpp b/Engine/Plugins/ChaosCloth/Source/ChaosCloth/Private/ChaosCloth/ChaosClothVisualization.cpp
--- a/Engine/Plugins/ChaosCloth/Source/ChaosCloth/Private/ChaosCloth/ChaosClothVisualization.cpp
+++ b/Engine/Plugins/ChaosCloth/Source/ChaosCloth/Private/ChaosCloth/ChaosClothVisualization.cpp
@@ -573,6 +573,34 @@ namespace Chaos
 #endif
 	}
 
+	static void DrawArrow(FPrimitiveDrawInterface* PDI, const FVector& StartLocation, const FVector& EndLocation, const FLinearColor& Color)
+	{
+		// This is just taken from FChaosVDDebugDrawUtils::DrawArrowVector
+		
+		const FVector LineVectorToDraw = EndLocation - StartLocation;
+
+		FVector ArrowDir;
+		float ArrowLength;
+		LineVectorToDraw.ToDirectionAndLength(ArrowDir, ArrowLength);
+
+		FVector YAxis, ZAxis;
+		ArrowDir.FindBestAxisVectors(YAxis, ZAxis);
+		const FMatrix ArrowTransformMatrix(ArrowDir, YAxis, ZAxis, StartLocation);
+
+		constexpr float MinTipOfArrowSize = 0.2f;
+		constexpr float MaxTipOfArrowSize = 10.0f;
+		constexpr float MaxVectorSizeForArrow = 100.0f; // The vector size that is the upper limit after which we just use the max size for the tip of the arrow
+
+		const float ProportionalArrowSize = MaxTipOfArrowSize * (ArrowLength / MaxVectorSizeForArrow);
+		const float ArrowSize = FMath::Clamp(ProportionalArrowSize, MinTipOfArrowSize, MaxTipOfArrowSize);
+
+		DrawLine(PDI, ArrowTransformMatrix.TransformPosition(FVector(ArrowLength, 0, 0)), ArrowTransformMatrix.TransformPosition(FVector::ZeroVector), Color);
+		DrawLine(PDI, ArrowTransformMatrix.TransformPosition(FVector(ArrowLength, 0, 0)), ArrowTransformMatrix.TransformPosition(FVector(ArrowLength - ArrowSize, +ArrowSize, +ArrowSize)), Color);
+		DrawLine(PDI, ArrowTransformMatrix.TransformPosition(FVector(ArrowLength, 0, 0)), ArrowTransformMatrix.TransformPosition(FVector(ArrowLength - ArrowSize, +ArrowSize, -ArrowSize)), Color);
+		DrawLine(PDI, ArrowTransformMatrix.TransformPosition(FVector(ArrowLength, 0, 0)), ArrowTransformMatrix.TransformPosition(FVector(ArrowLength - ArrowSize, -ArrowSize, +ArrowSize)), Color);
+		DrawLine(PDI, ArrowTransformMatrix.TransformPosition(FVector(ArrowLength, 0, 0)), ArrowTransformMatrix.TransformPosition(FVector(ArrowLength - ArrowSize, -ArrowSize, -ArrowSize)), Color);
+	}
+
 	static void DrawArc(FPrimitiveDrawInterface* PDI, const FVector& Base, const FVector& X, const FVector& Y, ::Chaos::FReal MinAngle, float MaxAngle, ::Chaos::FReal Radius, const FLinearColor& Color)
 	{
 		static const int32 Sections = 10;
@@ -1076,25 +1104,23 @@ namespace Chaos
 					{
 						const FVector Pos0 = GetWorldPosition(*Solver, Positions[ParticleIndex]);
 						const FVector Pos1 = Pos0 + GravityVectorLengthMultiplier * GetWorldVector(*Solver, ExternalForces->GetScaledGravity(ParticleIndex));
-						DrawLine(PDI, Pos0, Pos1, FLinearColor::Red);
+						DrawArrow(PDI, Pos0, Pos1, FLinearColor::Red);
 					}
 				}
 				else
 				{
-					const FAABB3 Bounds = Cloth->CalculateBoundingBox(Solver);
-
-					const FVector Pos0 = Bounds.Center();
+					const FRigidTransform3& ReferenceSpaceTransform = Cloth->GetReferenceSpaceTransform();
+					const FVector Pos0 = ReferenceSpaceTransform.GetLocation();
 					const FVector Pos1 = Pos0 + GravityVectorLengthMultiplier * GetWorldVector(*Solver, ExternalForces->GetScaledGravity(0));
-					DrawLine(PDI, Pos0, Pos1, FLinearColor::Red);
+					DrawArrow(PDI, Pos0, Pos1, FLinearColor::Red);
 				}
 			}
 			else
 			{
-				const FAABB3 Bounds = Cloth->CalculateBoundingBox(Solver);
-
-				const FVector Pos0 = Bounds.Center();
+				const FRigidTransform3& ReferenceSpaceTransform = Cloth->GetReferenceSpaceTransform();
+				const FVector Pos0 = ReferenceSpaceTransform.GetLocation();
 				const FVector Pos1 = Pos0 + GravityVectorLengthMultiplier * GetWorldVector(*Solver, Cloth->GetGravity(Solver));
-				DrawLine(PDI, Pos0, Pos1, FLinearColor::Red);
+				DrawArrow(PDI, Pos0, Pos1, FLinearColor::Red);
 			}
 		}
 	}
@@ -2637,6 +2663,42 @@ namespace Chaos
 		}
 	}
 
+	void FClothVisualizationNoGC::DrawWindVelocity(FPrimitiveDrawInterface* PDI, const FReal LengthScale) const
+	{
+		if (!Solver)
+		{
+			return;
+		}
+
+		for (const FClothingSimulationCloth* const Cloth : Solver->GetCloths())
+		{
+			const int32 ParticleRangeId = Cloth->GetParticleRangeId(Solver);
+			if (ParticleRangeId == INDEX_NONE)
+			{
+				continue;
+			}
+
+			const Softs::FVelocityAndPressureField* VelocityField = nullptr;
+			if (!Solver->IsLegacySolver())
+			{
+				const FClothConstraints& ClothConstraints = Solver->GetClothConstraints(ParticleRangeId);
+				VelocityField = ClothConstraints.GetVelocityAndPressureField().Get();
+			}
+			else
+			{
+				VelocityField = &Solver->GetWindVelocityAndPressureField(Cloth->GetGroupId());
+			}
+
+			if (VelocityField)
+			{
+				const FRigidTransform3& ReferenceSpaceTransform = Cloth->GetReferenceSpaceTransform();
+				const FVector Pos0 = ReferenceSpaceTransform.GetLocation();
+				const FVector Pos1 = Pos0 + LengthScale * GetWorldVector(*Solver, VelocityField->GetVelocity());
+				DrawArrow(PDI, Pos0, Pos1, FLinearColor::Blue);
+			}
+		}
+	}
+
 	void FClothVisualizationNoGC::DrawLocalSpace(FPrimitiveDrawInterface* PDI) const
 	{
 		if (!Solver)
diff --git a/Engine/Plugins/ChaosCloth/Source/ChaosCloth/Private/ChaosCloth/ChaosClothingSimulation.cpp b/Engine/Plugins/ChaosCloth/Source/ChaosCloth/Private/ChaosCloth/ChaosClothingSimulation.cpp
--- a/Engine/Plugins/ChaosCloth/Source/ChaosCloth/Private/ChaosCloth/ChaosClothingSimulation.cpp
+++ b/Engine/Plugins/ChaosCloth/Source/ChaosCloth/Private/ChaosCloth/ChaosClothingSimulation.cpp
@@ -93,6 +93,7 @@ namespace ClothingSimulationCVar
 	TAutoConsoleVariable<bool> DebugDrawBendingConstraint   (TEXT("p.ChaosCloth.DebugDrawBendingConstraint"   ), false, TEXT("Whether to debug draw the Chaos Cloth bending constraint"), ECVF_Cheat);
 	TAutoConsoleVariable<bool> DebugDrawLongRangeConstraint (TEXT("p.ChaosCloth.DebugDrawLongRangeConstraint" ), false, TEXT("Whether to debug draw the Chaos Cloth long range constraint (aka tether constraint)"), ECVF_Cheat);
 	TAutoConsoleVariable<bool> DebugDrawWindForces          (TEXT("p.ChaosCloth.DebugDrawWindForces"          ), false, TEXT("Whether to debug draw the Chaos Cloth wind forces"), ECVF_Cheat);
+	TAutoConsoleVariable<bool> DebugDrawWindVelocity        (TEXT("p.ChaosCloth.DebugDrawWindVelocity"        ), false, TEXT("Whether to debug draw the Chaos Cloth wind velocity"), ECVF_Cheat);
 	TAutoConsoleVariable<bool> DebugDrawSelfCollision       (TEXT("p.ChaosCloth.DebugDrawSelfCollision"       ), false, TEXT("Whether to debug draw the Chaos Cloth self collision information"), ECVF_Cheat);
 	TAutoConsoleVariable<bool> DebugDrawSelfIntersection    (TEXT("p.ChaosCloth.DebugDrawSelfIntersection"    ), false, TEXT("Whether to debug draw the Chaos Cloth self intersection information"), ECVF_Cheat);
 	TAutoConsoleVariable<bool> DebugDrawParticleIndices     (TEXT("p.ChaosCloth.DebugDrawParticleIndices"     ), false, TEXT("Whether to debug draw the Chaos Cloth particle indices"), ECVF_Cheat);
@@ -674,6 +675,7 @@ void FClothingSimulation::Simulate(const IClothingSimulationContext* InContext)
 	if (ClothingSimulationCVar::DebugDrawBendingConstraint   .GetValueOnAnyThread()) { DebugDrawBendingConstraint   (); }
 	if (ClothingSimulationCVar::DebugDrawLongRangeConstraint .GetValueOnAnyThread()) { DebugDrawLongRangeConstraint (); }
 	if (ClothingSimulationCVar::DebugDrawWindForces          .GetValueOnAnyThread()) { DebugDrawWindAndPressureForces(); }
+	if (ClothingSimulationCVar::DebugDrawWindVelocity        .GetValueOnAnyThread()) { DebugDrawWindVelocity        (); }
 	if (ClothingSimulationCVar::DebugDrawSelfCollision       .GetValueOnAnyThread()) { DebugDrawSelfCollision       (); }
 	if (ClothingSimulationCVar::DebugDrawSelfIntersection    .GetValueOnAnyThread()) { DebugDrawSelfIntersection    (); }
 	if (ClothingSimulationCVar::DebugDrawTeleportReset       .GetValueOnAnyThread()) { DebugDrawTeleportReset       (); }
diff --git a/Engine/Plugins/ChaosCloth/Source/ChaosCloth/Public/ChaosCloth/ChaosClothVisualization.h b/Engine/Plugins/ChaosCloth/Source/ChaosCloth/Public/ChaosCloth/ChaosClothVisualization.h
--- a/Engine/Plugins/ChaosCloth/Source/ChaosCloth/Public/ChaosCloth/ChaosClothVisualization.h
+++ b/Engine/Plugins/ChaosCloth/Source/ChaosCloth/Public/ChaosCloth/ChaosClothVisualization.h
@@ -50,6 +50,7 @@ namespace Chaos
 		CHAOSCLOTH_API void DrawBendingConstraint(FPrimitiveDrawInterface* PDI = nullptr) const;
 		CHAOSCLOTH_API void DrawLongRangeConstraint(FPrimitiveDrawInterface* PDI = nullptr) const;
 		CHAOSCLOTH_API void DrawWindAndPressureForces(FPrimitiveDrawInterface* PDI, const FReal ForceLengthScale) const;
+		CHAOSCLOTH_API void DrawWindVelocity(FPrimitiveDrawInterface* PDI, const FReal LengthScale) const;
 		CHAOSCLOTH_API void DrawLocalSpace(FPrimitiveDrawInterface* PDI = nullptr) const;
 		CHAOSCLOTH_API void DrawSelfCollision(FPrimitiveDrawInterface* PDI = nullptr) const;
 		CHAOSCLOTH_API void DrawSelfIntersection(FPrimitiveDrawInterface* PDI = nullptr) const;
@@ -86,6 +87,7 @@ namespace Chaos
 		void DrawBendingConstraint(FPrimitiveDrawInterface* /*PDI*/ = nullptr) const {}
 		void DrawLongRangeConstraint(FPrimitiveDrawInterface* /*PDI*/ = nullptr) const {}
 		void DrawWindAndPressureForces(FPrimitiveDrawInterface* /*PDI*/, const FReal /*ForceLengthScale*/) const {}
+		void DrawWindVelocity(FPrimitiveDrawInterface* /*PDI*/, const FReal /*LengthScale*/) const {}
 		void DrawLocalSpace(FPrimitiveDrawInterface* /*PDI*/ = nullptr) const {}
 		void DrawSelfCollision(FPrimitiveDrawInterface* /*PDI*/ = nullptr) const {}
 		void DrawSelfIntersection(FPrimitiveDrawInterface* /*PDI*/ = nullptr) const {}
@@ -121,6 +123,11 @@ namespace Chaos
 			constexpr FReal DefaultForceLengthScale = 10.;
 			DrawWindAndPressureForces(PDI, DefaultForceLengthScale);
 		}
+		void DrawWindVelocity(FPrimitiveDrawInterface* PDI = nullptr) const
+		{
+			constexpr FReal DefaultLengthScale = .1;
+			DrawWindVelocity(PDI, DefaultLengthScale);
+		}
 
 #if WITH_EDITOR && CHAOS_DEBUG_DRAW
 		// Editor only functions
diff --git a/Engine/Plugins/ChaosCloth/Source/ChaosCloth/Public/ChaosCloth/ChaosClothingSimulation.h b/Engine/Plugins/ChaosCloth/Source/ChaosCloth/Public/ChaosCloth/ChaosClothingSimulation.h
--- a/Engine/Plugins/ChaosCloth/Source/ChaosCloth/Public/ChaosCloth/ChaosClothingSimulation.h
+++ b/Engine/Plugins/ChaosCloth/Source/ChaosCloth/Public/ChaosCloth/ChaosClothingSimulation.h
@@ -116,6 +116,7 @@ namespace Chaos
 		void DebugDrawBendingConstraint(FPrimitiveDrawInterface* PDI = nullptr) const { Visualization.DrawBendingConstraint(PDI); }
 		void DebugDrawLongRangeConstraint(FPrimitiveDrawInterface* PDI = nullptr) const { Visualization.DrawLongRangeConstraint(PDI); }
 		void DebugDrawWindAndPressureForces(FPrimitiveDrawInterface* PDI = nullptr) const { Visualization.DrawWindAndPressureForces(PDI); }
+		void DebugDrawWindVelocity(FPrimitiveDrawInterface* PDI = nullptr) const { Visualization.DrawWindVelocity(PDI); }
 		void DebugDrawLocalSpace(FPrimitiveDrawInterface* PDI = nullptr) const { Visualization.DrawLocalSpace(PDI); }
 		void DebugDrawSelfCollision(FPrimitiveDrawInterface* PDI = nullptr) const { Visualization.DrawSelfCollision(PDI); }
 		void DebugDrawSelfIntersection(FPrimitiveDrawInterface* PDI = nullptr) const { Visualization.DrawSelfIntersection(PDI); }
diff --git a/Engine/Plugins/ChaosCloth/Source/ChaosClothEditor/Private/ChaosClothEditor/ChaosSimulationEditorExtender.cpp b/Engine/Plugins/ChaosCloth/Source/ChaosClothEditor/Private/ChaosClothEditor/ChaosSimulationEditorExtender.cpp
--- a/Engine/Plugins/ChaosCloth/Source/ChaosClothEditor/Private/ChaosClothEditor/ChaosSimulationEditorExtender.cpp
+++ b/Engine/Plugins/ChaosCloth/Source/ChaosClothEditor/Private/ChaosClothEditor/ChaosSimulationEditorExtender.cpp
@@ -83,6 +83,7 @@ const FVisualizationOption FVisualizationOption::OptionData[] =
 	FVisualizationOption(&FClothVisualizationNoGC::DrawBendingConstraint   , TEXT("p.ChaosClothEditor.DebugDrawBendingConstraint"   ), LOCTEXT("ChaosVisName_BendingConstraint"   , "Bending Constraint"         ), LOCTEXT("ChaosVisName_BendingConstraint_Tooltip"   , "Draws the bending spring constraints")),
 	FVisualizationOption(&FClothVisualizationNoGC::DrawLongRangeConstraint , TEXT("p.ChaosClothEditor.DebugDrawLongRangeConstraint" ), LOCTEXT("ChaosVisName_LongRangeConstraint" , "Long Range Constraint"      ), LOCTEXT("ChaosVisName_LongRangeConstraint_Tooltip" , "Draws the long range attachment constraint distances")),
 	FVisualizationOption(&FClothVisualizationNoGC::DrawWindAndPressureForces, TEXT("p.ChaosClothEditor.DebugDrawWindAndPressureForces"), LOCTEXT("ChaosVisName_WindAndPressureForces", "Wind Aerodynamic And Pressure Forces"), LOCTEXT("ChaosVisName_WindAndPressure_Tooltip", "Draws the Wind drag and lift and pressure forces")),
+	FVisualizationOption(&FClothVisualizationNoGC::DrawWindVelocity, TEXT("p.ChaosClothEditor.DebugDrawWindVelocity"), LOCTEXT("ChaosVisName_WindVelocity", "Wind Velocity"), LOCTEXT("ChaosVisName_WindVelocity_Tooltip", "Draws the Wind Velocity vector")),
 	FVisualizationOption(&FClothVisualizationNoGC::DrawLocalSpace          , TEXT("p.ChaosClothEditor.DebugDrawLocalSpace"          ), LOCTEXT("ChaosVisName_LocalSpace"          , "Local Space Reference Bone" ), LOCTEXT("ChaosVisName_LocalSpace_Tooltip"          , "Draws the local space reference bone")),
 	FVisualizationOption(&FClothVisualizationNoGC::DrawSelfCollision       , TEXT("p.ChaosClothEditor.DebugDrawSelfCollision"       ), LOCTEXT("ChaosVisName_SelfCollision"       , "Self Collision"             ), LOCTEXT("ChaosVisName_SelfCollision_Tooltip"       , "Draws the self collision thickness/debugging information")),
 	FVisualizationOption(&FClothVisualizationNoGC::DrawSelfIntersection    , TEXT("p.ChaosClothEditor.DebugDrawSelfIntersection"    ), LOCTEXT("ChaosVisName_SelfIntersection"    , "Self Intersection"          ), LOCTEXT("ChaosVisName_SelfIntersection_Tooltip"    , "Draws the self intersection contour/region information")),
diff --git a/Engine/Plugins/ChaosClothAsset/Source/ChaosClothAssetEngine/Private/ChaosClothAsset/ClothSimulationProxy.cpp b/Engine/Plugins/ChaosClothAsset/Source/ChaosClothAssetEngine/Private/ChaosClothAsset/ClothSimulationProxy.cpp
--- a/Engine/Plugins/ChaosClothAsset/Source/ChaosClothAssetEngine/Private/ChaosClothAsset/ClothSimulationProxy.cpp
+++ b/Engine/Plugins/ChaosClothAsset/Source/ChaosClothAssetEngine/Private/ChaosClothAsset/ClothSimulationProxy.cpp
@@ -667,6 +667,7 @@ namespace UE::Chaos::ClothAsset
 		static const TConsoleVariableData<bool>* const DebugDrawBendingConstraintCVar = IConsoleManager::Get().FindTConsoleVariableDataBool(TEXT("p.ChaosCloth.DebugDrawBendingConstraint"));
 		static const TConsoleVariableData<bool>* const DebugDrawLongRangeConstraintCVar = IConsoleManager::Get().FindTConsoleVariableDataBool(TEXT("p.ChaosCloth.DebugDrawLongRangeConstraint"));
 		static const TConsoleVariableData<bool>* const DebugDrawWindForcesCVar = IConsoleManager::Get().FindTConsoleVariableDataBool(TEXT("p.ChaosCloth.DebugDrawWindForces"));
+		static const TConsoleVariableData<bool>* const DebugDrawWindVelocityCVar = IConsoleManager::Get().FindTConsoleVariableDataBool(TEXT("p.ChaosCloth.DebugDrawWindVelocity"));
 		static const TConsoleVariableData<bool>* const DebugDrawSelfCollisionCVar = IConsoleManager::Get().FindTConsoleVariableDataBool(TEXT("p.ChaosCloth.DebugDrawSelfCollision"));
 		static const TConsoleVariableData<bool>* const DebugDrawSelfIntersectionCVar = IConsoleManager::Get().FindTConsoleVariableDataBool(TEXT("p.ChaosCloth.DebugDrawSelfIntersection"));
 		static const TConsoleVariableData<bool>* const DebugDrawParticleIndicesCVar = IConsoleManager::Get().FindTConsoleVariableDataBool(TEXT("p.ChaosCloth.DebugDrawParticleIndices"));
@@ -693,6 +694,7 @@ namespace UE::Chaos::ClothAsset
 		if (DebugDrawBendingConstraintCVar && DebugDrawBendingConstraintCVar->GetValueOnAnyThread()) { Visualization->DrawBendingConstraint(); }
 		if (DebugDrawLongRangeConstraintCVar && DebugDrawLongRangeConstraintCVar->GetValueOnAnyThread()) { Visualization->DrawLongRangeConstraint(); }
 		if (DebugDrawWindForcesCVar && DebugDrawWindForcesCVar->GetValueOnAnyThread()) { Visualization->DrawWindAndPressureForces(); }
+		if (DebugDrawWindVelocityCVar && DebugDrawWindVelocityCVar->GetValueOnAnyThread()) { Visualization->DrawWindVelocity(); }
 		if (DebugDrawSelfCollisionCVar && DebugDrawSelfCollisionCVar->GetValueOnAnyThread()) { Visualization->DrawSelfCollision(); }
 		if (DebugDrawSelfIntersectionCVar && DebugDrawSelfIntersectionCVar->GetValueOnAnyThread()) { Visualization->DrawSelfIntersection(); }
 		if (DebugDrawParticleIndicesCVar && DebugDrawParticleIndicesCVar->GetValueOnAnyThread()) { Visualization->DrawParticleIndices(); }
diff --git a/Engine/Plugins/ChaosClothAssetEditor/Source/ChaosClothAssetEditor/Private/ChaosClothAsset/ClothEditorSimulationVisualization.cpp b/Engine/Plugins/ChaosClothAssetEditor/Source/ChaosClothAssetEditor/Private/ChaosClothAsset/ClothEditorSimulationVisualization.cpp
--- a/Engine/Plugins/ChaosClothAssetEditor/Source/ChaosClothAssetEditor/Private/ChaosClothAsset/ClothEditorSimulationVisualization.cpp
+++ b/Engine/Plugins/ChaosClothAssetEditor/Source/ChaosClothAssetEditor/Private/ChaosClothAsset/ClothEditorSimulationVisualization.cpp
@@ -327,6 +327,18 @@ const FVisualizationOption FVisualizationOption::OptionData[] =
 		{
 			EditorVisualization.ExtendViewportShowMenuAerodynamicsLengthScale(MenuBuilder);
 		})),
+	FVisualizationOption(
+		FClothVisualizationDebugDraw::CreateLambda([](const FClothEditorSimulationVisualization& EditorVisualization, const ::Chaos::FClothVisualizationNoGC& Visualization, FPrimitiveDrawInterface* PDI)
+		{
+			Visualization.DrawWindVelocity(PDI, EditorVisualization.GetWindVelocityLengthScale());
+		}),
+		LOCTEXT("ChaosVisName_WindVelocity", "Wind Velocity"), 
+		LOCTEXT("ChaosVisName_WindVelocity_Tooltip", "Draws the Wind Velocity vector"),
+		/*bDisablesSimulation =*/false, /*bHidesClothSections=*/false, /*bDefaultFlagValue=*/false,
+		FAdditionalMenus::CreateLambda([](FClothEditorSimulationVisualization& EditorVisualization, FMenuBuilder& MenuBuilder)
+		{
+			EditorVisualization.ExtendViewportShowMenuWindVelocityLength(MenuBuilder);
+		})),
 	FVisualizationOption(
 		FClothVisualizationDebugDraw::CreateLambda([](const FClothEditorSimulationVisualization&, const ::Chaos::FClothVisualizationNoGC& Visualization, FPrimitiveDrawInterface* PDI)
 		{
diff --git a/Engine/Plugins/ChaosClothAssetEditor/Source/ChaosClothAssetEditor/Private/ChaosClothAsset/ClothEditorSimulationVisualization.h b/Engine/Plugins/ChaosClothAssetEditor/Source/ChaosClothAssetEditor/Private/ChaosClothAsset/ClothEditorSimulationVisualization.h
--- a/Engine/Plugins/ChaosClothAssetEditor/Source/ChaosClothAssetEditor/Private/ChaosClothAsset/ClothEditorSimulationVisualization.h
+++ b/Engine/Plugins/ChaosClothAssetEditor/Source/ChaosClothAssetEditor/Private/ChaosClothAsset/ClothEditorSimulationVisualization.h
@@ -70,10 +70,15 @@ public:
 	{
 		ExtendViewportShowMenuSpinBox(MenuBuilder, AccessoryMeshNormalLength, 0.f, FLT_MAX, 0.f, 40.f);
 	}
+	void ExtendViewportShowMenuWindVelocityLength(FMenuBuilder& MenuBuilder)
+	{
+		ExtendViewportShowMenuSpinBox(MenuBuilder, WindVelocityLengthScale, 0.f, FLT_MAX, 0.f, 40.f);
+	}
 	float GetPointNormalLength() const { return PointNormalLength; }
 	float GetAnimatedNormalLength() const { return AnimatedNormalLength; }
 	float GetAccessoryMeshNormalLength() const { return AccessoryMeshNormalLength; }
 	float GetAerodynamicsLengthScale() const { return AerodynamicsLengthScale; }
+	float GetWindVelocityLengthScale() const { return WindVelocityLengthScale; }
 
 private:
 	struct FNameSelectionData
@@ -103,6 +108,7 @@ private:
 	float AnimatedNormalLength = 20.f;
 	float AccessoryMeshNormalLength = 20.f;
 	float AerodynamicsLengthScale = 10.f;
+	float WindVelocityLengthScale = .1f;
 	
 };
 } // namespace UE::Chaos::ClothAsset
diff --git a/Engine/Source/Runtime/Experimental/Chaos/Public/Chaos/VelocityField.h b/Engine/Source/Runtime/Experimental/Chaos/Public/Chaos/VelocityField.h
--- a/Engine/Source/Runtime/Experimental/Chaos/Public/Chaos/VelocityField.h
+++ b/Engine/Source/Runtime/Experimental/Chaos/Public/Chaos/VelocityField.h
@@ -192,6 +192,8 @@ public:
 
 	void SetVelocity(const FSolverVec3& InVelocity) { Velocity = InVelocity; }
 
+	const FSolverVec3& GetVelocity() const { return Velocity; }
+
 	TConstArrayView<TVector<int32, 3>> GetElements() const { return TConstArrayView<TVector<int32, 3>>(Elements); }
 	TConstArrayView<FSolverVec3> GetForces() const { return TConstArrayView<FSolverVec3>(Forces); }
 
