Cloth - Deprecated the old Clothing Simulation Interface and added a new updated class to replace it.

#rnx

diff --git a/Engine/Plugins/ChaosCloth/Source/ChaosCloth/Private/ChaosCloth/ChaosClothingSimulation.cpp b/Engine/Plugins/ChaosCloth/Source/ChaosCloth/Private/ChaosCloth/ChaosClothingSimulation.cpp
--- a/Engine/Plugins/ChaosCloth/Source/ChaosCloth/Private/ChaosCloth/ChaosClothingSimulation.cpp
+++ b/Engine/Plugins/ChaosCloth/Source/ChaosCloth/Private/ChaosCloth/ChaosClothingSimulation.cpp
@@ -454,21 +454,6 @@ void FClothingSimulation::EndCreateActor()
 	}
 }
 
-void FClothingSimulation::UpdateWorldForces(const USkeletalMeshComponent* OwnerComponent)
-{
-	if (OwnerComponent)
-	{
-		UWorld* OwnerWorld = OwnerComponent->GetWorld();
-		if (OwnerWorld && OwnerWorld->PhysicsField && Solver)
-		{
-			const FBox BoundingBox = OwnerComponent->CalcBounds(OwnerComponent->GetComponentTransform()).GetBox();
-
-			OwnerWorld->PhysicsField->FillTransientCommands(false, BoundingBox, Solver->GetTime(), Solver->GetPerSolverField().GetTransientCommands());
-			OwnerWorld->PhysicsField->FillPersistentCommands(false, BoundingBox, Solver->GetTime(), Solver->GetPerSolverField().GetPersistentCommands());
-		}
-	}
-}
-
 void FClothingSimulation::ResetStats()
 {
 	check(Solver);
@@ -511,21 +496,6 @@ void FClothingSimulation::UpdateSimulationFromSharedSimConfig()
 	}
 }
 
-void FClothingSimulation::SetNumIterations(int32 InNumIterations)
-{
-	Solver->GetConfig()->GetProperties(Solver->GetSolverLOD()).SetValue(FClothingSimulationSolver::NumIterationsName, InNumIterations);
-}
-
-void FClothingSimulation::SetMaxNumIterations(int32 MaxNumIterations)
-{
-	Solver->GetConfig()->GetProperties(Solver->GetSolverLOD()).SetValue(FClothingSimulationSolver::MaxNumIterationsName, MaxNumIterations);
-}
-
-void FClothingSimulation::SetNumSubsteps(int32 InNumSubsteps)
-{
-	Solver->GetConfig()->GetProperties(Solver->GetSolverLOD()).SetValue(FClothingSimulationSolver::NumSubstepsName, InNumSubsteps);
-}
-
 void FClothingSimulation::FillContextAndPrepareTick(const USkeletalMeshComponent* InComponent, float InDeltaTime, IClothingSimulationContext* InOutContext, bool bIsInitialization, bool bForceTeleportResetOnly)
 {
 	if (bForceTeleportResetOnly)
@@ -534,6 +504,20 @@ void FClothingSimulation::FillContextAndPrepareTick(const USkeletalMeshComponent
 		return;
 	}
 
+	// Update world forces
+	if (InComponent)
+	{
+		UWorld* OwnerWorld = InComponent->GetWorld();
+		if (OwnerWorld && OwnerWorld->PhysicsField && Solver)
+		{
+			const FBox BoundingBox = InComponent->CalcBounds(InComponent->GetComponentTransform()).GetBox();
+
+			OwnerWorld->PhysicsField->FillTransientCommands(false, BoundingBox, Solver->GetTime(), Solver->GetPerSolverField().GetTransientCommands());
+			OwnerWorld->PhysicsField->FillPersistentCommands(false, BoundingBox, Solver->GetTime(), Solver->GetPerSolverField().GetPersistentCommands());
+		}
+	}
+
+	// Fill context
 	FClothingSimulationCommon::FillContextAndPrepareTick(InComponent, InDeltaTime, InOutContext, bIsInitialization, bForceTeleportResetOnly);
 }
 
@@ -558,7 +542,7 @@ bool FClothingSimulation::ShouldSimulateLOD(int32 OwnerLODIndex) const
 	return false;
 }
 
-void FClothingSimulation::ForceClothNextUpdateTeleportAndReset()
+void FClothingSimulation::ForceClothNextUpdateTeleportAndReset_AnyThread()
 {
 	for (const TUniquePtr<FClothingSimulationCloth>& Cloth : Cloths)
 	{
@@ -566,7 +550,7 @@ void FClothingSimulation::ForceClothNextUpdateTeleportAndReset()
 	}
 }
 
-void FClothingSimulation::Simulate(const IClothingSimulationContext* InContext)
+void FClothingSimulation::Simulate_AnyThread(const IClothingSimulationContext* InContext)
 {
 	TRACE_CPUPROFILER_EVENT_SCOPE(FClothingSimulation_Simulate);
 #if !(UE_BUILD_SHIPPING || UE_BUILD_TEST)
diff --git a/Engine/Plugins/ChaosCloth/Source/ChaosCloth/Private/ChaosCloth/ChaosClothingSimulationFactory.cpp b/Engine/Plugins/ChaosCloth/Source/ChaosCloth/Private/ChaosCloth/ChaosClothingSimulationFactory.cpp
--- a/Engine/Plugins/ChaosCloth/Source/ChaosCloth/Private/ChaosCloth/ChaosClothingSimulationFactory.cpp
+++ b/Engine/Plugins/ChaosCloth/Source/ChaosCloth/Private/ChaosCloth/ChaosClothingSimulationFactory.cpp
@@ -9,13 +9,12 @@
 
 #include UE_INLINE_GENERATED_CPP_BY_NAME(ChaosClothingSimulationFactory)
 
-IClothingSimulation* UChaosClothingSimulationFactory::CreateSimulation()
+IClothingSimulationInterface* UChaosClothingSimulationFactory::CreateSimulation() const
 {
-	IClothingSimulation* Simulation = new Chaos::FClothingSimulation();
-	return Simulation;
+	return new Chaos::FClothingSimulation();
 }
 
-void UChaosClothingSimulationFactory::DestroySimulation(IClothingSimulation* InSimulation)
+void UChaosClothingSimulationFactory::DestroySimulation(IClothingSimulationInterface* InSimulation) const
 {
 	delete InSimulation;
 }
diff --git a/Engine/Plugins/ChaosCloth/Source/ChaosCloth/Private/ChaosCloth/ChaosClothingSimulationInteractor.cpp b/Engine/Plugins/ChaosCloth/Source/ChaosCloth/Private/ChaosCloth/ChaosClothingSimulationInteractor.cpp
--- a/Engine/Plugins/ChaosCloth/Source/ChaosCloth/Private/ChaosCloth/ChaosClothingSimulationInteractor.cpp
+++ b/Engine/Plugins/ChaosCloth/Source/ChaosCloth/Private/ChaosCloth/ChaosClothingSimulationInteractor.cpp
@@ -24,19 +24,19 @@ namespace ChaosClothingInteractor
 	static const float InvStiffnessLogBase = 1.f / FMath::Loge(1.e3f);  // Log base for updating old linear stiffnesses to the new stiffness exponentiation
 }
 
-void UChaosClothingInteractor::Sync(IClothingSimulation* Simulation)
+void UChaosClothingInteractor::Sync(IClothingSimulationInterface* Simulation)
 {
 	check(Simulation);
 
 	TRACE_CPUPROFILER_EVENT_SCOPE(UChaosClothingInteractor_Sync);
 	if (Chaos::FClothingSimulationCloth* const Cloth = static_cast<Chaos::FClothingSimulation*>(Simulation)->GetCloth(ClothingId))
 	{
-		for (FChaosClothingInteractorCommand& Command : Commands)
+		for (Chaos::Internal::FClothingInteractorCommand& Command : Commands)
 		{
 			Command.Execute(Cloth);
 		}
 		Commands.Reset();
-		for (FChaosClothingInteractorConfigCommand& ConfigCommand : ConfigCommands)
+		for (Chaos::Internal::FClothingInteractorConfigCommand& ConfigCommand : ConfigCommands)
 		{
 			check(Cloth->GetConfig()->IsLegacySingleLOD());
 			ConfigCommand.Execute(Cloth->GetConfig(), 0);
@@ -54,7 +54,7 @@ void UChaosClothingInteractor::SetMaterialLinear(float EdgeStiffnessLinear, floa
 	const FVector2f BendingStiffness((FMath::Clamp(FMath::Loge(BendingStiffnessLinear) * ChaosClothingInteractor::InvStiffnessLogBase + 1.f, 0.f, 1.f)), 1.f);
 	const FVector2f AreaStiffness((FMath::Clamp(FMath::Loge(AreaStiffnessLinear) * ChaosClothingInteractor::InvStiffnessLogBase + 1.f, 0.f, 1.f)), 1.f);
 
-	ConfigCommands.Add(FChaosClothingInteractorConfigCommand::CreateLambda([EdgeStiffness, BendingStiffness, AreaStiffness](Chaos::FClothingSimulationConfig* Config, int32 LODIndex)
+	ConfigCommands.Add(Chaos::Internal::FClothingInteractorConfigCommand::CreateLambda([EdgeStiffness, BendingStiffness, AreaStiffness](Chaos::FClothingSimulationConfig* Config, int32 LODIndex)
 	{
 		Config->GetProperties(LODIndex).SetWeightedFloatValue(Chaos::Softs::FPBDEdgeSpringConstraints::EdgeSpringStiffnessName, EdgeStiffness);
 		Config->GetProperties(LODIndex).SetWeightedFloatValue(Chaos::Softs::FPBDBendingSpringConstraints::BendingSpringStiffnessName, BendingStiffness);
@@ -64,7 +64,7 @@ void UChaosClothingInteractor::SetMaterialLinear(float EdgeStiffnessLinear, floa
 
 void UChaosClothingInteractor::SetMaterial(FVector2D EdgeStiffness, FVector2D BendingStiffness, FVector2D AreaStiffness)
 {
-	ConfigCommands.Add(FChaosClothingInteractorConfigCommand::CreateLambda([EdgeStiffness, BendingStiffness, AreaStiffness](Chaos::FClothingSimulationConfig* Config, int32 LODIndex)
+	ConfigCommands.Add(Chaos::Internal::FClothingInteractorConfigCommand::CreateLambda([EdgeStiffness, BendingStiffness, AreaStiffness](Chaos::FClothingSimulationConfig* Config, int32 LODIndex)
 	{
 		Config->GetProperties(LODIndex).SetWeightedFloatValue(Chaos::Softs::FPBDEdgeSpringConstraints::EdgeSpringStiffnessName, FVector2f(EdgeStiffness));
 		Config->GetProperties(LODIndex).SetWeightedFloatValue(Chaos::Softs::FPBDBendingSpringConstraints::BendingSpringStiffnessName, FVector2f(BendingStiffness));
@@ -74,7 +74,7 @@ void UChaosClothingInteractor::SetMaterial(FVector2D EdgeStiffness, FVector2D Be
 
 void UChaosClothingInteractor::SetMaterialBuckling(FVector2D BucklingRatio, FVector2D BucklingStiffness)
 {
-	ConfigCommands.Add(FChaosClothingInteractorConfigCommand::CreateLambda([BucklingRatio, BucklingStiffness](Chaos::FClothingSimulationConfig* Config, int32 LODIndex)
+	ConfigCommands.Add(Chaos::Internal::FClothingInteractorConfigCommand::CreateLambda([BucklingRatio, BucklingStiffness](Chaos::FClothingSimulationConfig* Config, int32 LODIndex)
 	{
 		Config->GetProperties(LODIndex).SetValue(Chaos::Softs::FPBDBendingConstraints::BucklingRatioName, (float)BucklingRatio[0]);  // TODO: Make BuckingRatio weighted
 		Config->GetProperties(LODIndex).SetWeightedFloatValue(Chaos::Softs::FPBDBendingConstraints::BucklingStiffnessName, FVector2f(BucklingStiffness));
@@ -85,7 +85,7 @@ void UChaosClothingInteractor::SetLongRangeAttachmentLinear(float TetherStiffnes
 {
 	// Deprecated
 	const FVector2f TetherStiffness((FMath::Clamp(FMath::Loge(TetherStiffnessLinear) * ChaosClothingInteractor::InvStiffnessLogBase + 1.f, 0.f, 1.f)), 1.f);
-	ConfigCommands.Add(FChaosClothingInteractorConfigCommand::CreateLambda([TetherStiffness, TetherScale](Chaos::FClothingSimulationConfig* Config, int32 LODIndex)
+	ConfigCommands.Add(Chaos::Internal::FClothingInteractorConfigCommand::CreateLambda([TetherStiffness, TetherScale](Chaos::FClothingSimulationConfig* Config, int32 LODIndex)
 	{
 		Config->GetProperties(LODIndex).SetWeightedFloatValue(Chaos::Softs::FPBDLongRangeConstraints::TetherStiffnessName, TetherStiffness);
 		Config->GetProperties(LODIndex).SetWeightedFloatValue(Chaos::Softs::FPBDLongRangeConstraints::TetherScaleName, FVector2f(TetherScale));
@@ -94,7 +94,7 @@ void UChaosClothingInteractor::SetLongRangeAttachmentLinear(float TetherStiffnes
 
 void UChaosClothingInteractor::SetLongRangeAttachment(FVector2D TetherStiffness, FVector2D TetherScale)
 {
-	ConfigCommands.Add(FChaosClothingInteractorConfigCommand::CreateLambda([TetherStiffness, TetherScale](Chaos::FClothingSimulationConfig* Config, int32 LODIndex)
+	ConfigCommands.Add(Chaos::Internal::FClothingInteractorConfigCommand::CreateLambda([TetherStiffness, TetherScale](Chaos::FClothingSimulationConfig* Config, int32 LODIndex)
 	{
 		Config->GetProperties(LODIndex).SetWeightedFloatValue(Chaos::Softs::FPBDLongRangeConstraints::TetherStiffnessName, FVector2f(TetherStiffness));
 		Config->GetProperties(LODIndex).SetWeightedFloatValue(Chaos::Softs::FPBDLongRangeConstraints::TetherScaleName, FVector2f(TetherScale));
@@ -103,7 +103,7 @@ void UChaosClothingInteractor::SetLongRangeAttachment(FVector2D TetherStiffness,
 
 void UChaosClothingInteractor::SetCollision(float CollisionThickness, float FrictionCoefficient, bool bUseCCD, float SelfCollisionThickness)
 {
-	ConfigCommands.Add(FChaosClothingInteractorConfigCommand::CreateLambda([CollisionThickness, FrictionCoefficient, bUseCCD, SelfCollisionThickness](Chaos::FClothingSimulationConfig* Config, int32 LODIndex)
+	ConfigCommands.Add(Chaos::Internal::FClothingInteractorConfigCommand::CreateLambda([CollisionThickness, FrictionCoefficient, bUseCCD, SelfCollisionThickness](Chaos::FClothingSimulationConfig* Config, int32 LODIndex)
 	{
 		Config->GetProperties(LODIndex).SetValue(Chaos::FClothingSimulationCloth::CollisionThicknessName, CollisionThickness);
 		Config->GetProperties(LODIndex).SetValue(Chaos::FClothingSimulationCloth::FrictionCoefficientName, FrictionCoefficient);
@@ -114,14 +114,14 @@ void UChaosClothingInteractor::SetCollision(float CollisionThickness, float Fric
 
 void UChaosClothingInteractor::SetBackstop(bool bEnabled)
 {
-	ConfigCommands.Add(FChaosClothingInteractorConfigCommand::CreateLambda([bEnabled](Chaos::FClothingSimulationConfig* Config, int32 LODIndex)
+	ConfigCommands.Add(Chaos::Internal::FClothingInteractorConfigCommand::CreateLambda([bEnabled](Chaos::FClothingSimulationConfig* Config, int32 LODIndex)
 	{
 		Config->GetProperties(LODIndex).SetEnabled(Chaos::Softs::FPBDSphericalBackstopConstraint::BackstopRadiusName, bEnabled);  // BackstopRadius controls whether the backstop is enabled or not
 	}));
 }
 void UChaosClothingInteractor::SetDamping(float DampingCoefficient, float LocalDampingCoefficient)
 {
-	ConfigCommands.Add(FChaosClothingInteractorConfigCommand::CreateLambda([DampingCoefficient, LocalDampingCoefficient](Chaos::FClothingSimulationConfig* Config, int32 LODIndex)
+	ConfigCommands.Add(Chaos::Internal::FClothingInteractorConfigCommand::CreateLambda([DampingCoefficient, LocalDampingCoefficient](Chaos::FClothingSimulationConfig* Config, int32 LODIndex)
 	{
 		Config->GetProperties(LODIndex).SetValue(Chaos::FClothingSimulationCloth::DampingCoefficientName, DampingCoefficient);
 		Config->GetProperties(LODIndex).SetValue(Chaos::FClothingSimulationCloth::LocalDampingCoefficientName, LocalDampingCoefficient);
@@ -130,7 +130,7 @@ void UChaosClothingInteractor::SetDamping(float DampingCoefficient, float LocalD
 
 void UChaosClothingInteractor::SetAerodynamics(float DragCoefficient, float LiftCoefficient, FVector WindVelocity)
 {
-	ConfigCommands.Add(FChaosClothingInteractorConfigCommand::CreateLambda([DragCoefficient, LiftCoefficient, WindVelocity](Chaos::FClothingSimulationConfig* Config, int32 LODIndex)
+	ConfigCommands.Add(Chaos::Internal::FClothingInteractorConfigCommand::CreateLambda([DragCoefficient, LiftCoefficient, WindVelocity](Chaos::FClothingSimulationConfig* Config, int32 LODIndex)
 	{
 		constexpr float AirDensity = 1.225f;
 		constexpr float WorldScale = 100.f;  // Unreal's world unit is the cm
@@ -143,7 +143,7 @@ void UChaosClothingInteractor::SetAerodynamics(float DragCoefficient, float Lift
 
 void UChaosClothingInteractor::SetWind(FVector2D Drag, FVector2D Lift, float AirDensity, FVector WindVelocity, FVector2D OuterDrag, FVector2D OuterLift)
 {
-	ConfigCommands.Add(FChaosClothingInteractorConfigCommand::CreateLambda([Drag, Lift, AirDensity, WindVelocity, OuterDrag, OuterLift](Chaos::FClothingSimulationConfig* Config, int32 LODIndex)
+	ConfigCommands.Add(Chaos::Internal::FClothingInteractorConfigCommand::CreateLambda([Drag, Lift, AirDensity, WindVelocity, OuterDrag, OuterLift](Chaos::FClothingSimulationConfig* Config, int32 LODIndex)
 	{
 		constexpr float WorldScale = 100.f;  // Unreal's world unit is the cm
 		Config->GetProperties(LODIndex).SetWeightedFloatValue(Chaos::Softs::FVelocityAndPressureField::DragName, FVector2f(Drag));
@@ -157,7 +157,7 @@ void UChaosClothingInteractor::SetWind(FVector2D Drag, FVector2D Lift, float Air
 
 void UChaosClothingInteractor::SetPressure(FVector2D Pressure)
 {
-	ConfigCommands.Add(FChaosClothingInteractorConfigCommand::CreateLambda([Pressure](Chaos::FClothingSimulationConfig* Config, int32 LODIndex)
+	ConfigCommands.Add(Chaos::Internal::FClothingInteractorConfigCommand::CreateLambda([Pressure](Chaos::FClothingSimulationConfig* Config, int32 LODIndex)
 	{
 		Config->GetProperties(LODIndex).SetWeightedFloatValue(Chaos::Softs::FVelocityAndPressureField::PressureName, FVector2f(Pressure));
 	}));
@@ -165,7 +165,7 @@ void UChaosClothingInteractor::SetPressure(FVector2D Pressure)
 
 void UChaosClothingInteractor::SetGravity(float GravityScale, bool bIsGravityOverridden, FVector GravityOverride)
 {
-	ConfigCommands.Add(FChaosClothingInteractorConfigCommand::CreateLambda([GravityScale, bIsGravityOverridden, GravityOverride](Chaos::FClothingSimulationConfig* Config, int32 LODIndex)
+	ConfigCommands.Add(Chaos::Internal::FClothingInteractorConfigCommand::CreateLambda([GravityScale, bIsGravityOverridden, GravityOverride](Chaos::FClothingSimulationConfig* Config, int32 LODIndex)
 	{
 		Config->GetProperties(LODIndex).SetValue(Chaos::FClothingSimulationCloth::GravityScaleName, GravityScale);
 		Config->GetProperties(LODIndex).SetValue(Chaos::FClothingSimulationCloth::UseGravityOverrideName, bIsGravityOverridden);
@@ -177,7 +177,7 @@ void UChaosClothingInteractor::SetAnimDriveLinear(float AnimDriveStiffnessLinear
 {
 	// Deprecated
 	const FVector2f AnimDriveStiffness(0.f, FMath::Clamp(FMath::Loge(AnimDriveStiffnessLinear) * ChaosClothingInteractor::InvStiffnessLogBase + 1.f, 0.f, 1.f));
-	ConfigCommands.Add(FChaosClothingInteractorConfigCommand::CreateLambda([AnimDriveStiffness](Chaos::FClothingSimulationConfig* Config, int32 LODIndex)
+	ConfigCommands.Add(Chaos::Internal::FClothingInteractorConfigCommand::CreateLambda([AnimDriveStiffness](Chaos::FClothingSimulationConfig* Config, int32 LODIndex)
 	{
 		// The Anim Drive stiffness Low value needs to be 0 in order to keep backward compatibility with existing mask (this wouldn't be an issue if this property had no legacy mask)
 		static const FVector2f AnimDriveDamping(0.f, 1.f);
@@ -188,7 +188,7 @@ void UChaosClothingInteractor::SetAnimDriveLinear(float AnimDriveStiffnessLinear
 
 void UChaosClothingInteractor::SetAnimDrive(FVector2D AnimDriveStiffness, FVector2D AnimDriveDamping)
 {
-	ConfigCommands.Add(FChaosClothingInteractorConfigCommand::CreateLambda([AnimDriveStiffness, AnimDriveDamping](Chaos::FClothingSimulationConfig* Config, int32 LODIndex)
+	ConfigCommands.Add(Chaos::Internal::FClothingInteractorConfigCommand::CreateLambda([AnimDriveStiffness, AnimDriveDamping](Chaos::FClothingSimulationConfig* Config, int32 LODIndex)
 	{
 		Config->GetProperties(LODIndex).SetWeightedFloatValue(Chaos::Softs::FPBDAnimDriveConstraint::AnimDriveStiffnessName, FVector2f(AnimDriveStiffness));
 		Config->GetProperties(LODIndex).SetWeightedFloatValue(Chaos::Softs::FPBDAnimDriveConstraint::AnimDriveDampingName, FVector2f(AnimDriveDamping));
@@ -197,7 +197,7 @@ void UChaosClothingInteractor::SetAnimDrive(FVector2D AnimDriveStiffness, FVecto
 
 void UChaosClothingInteractor::SetVelocityScale(FVector LinearVelocityScale, float AngularVelocityScale, float FictitiousAngularScale)
 {
-	ConfigCommands.Add(FChaosClothingInteractorConfigCommand::CreateLambda([LinearVelocityScale, AngularVelocityScale, FictitiousAngularScale](Chaos::FClothingSimulationConfig* Config, int32 LODIndex)
+	ConfigCommands.Add(Chaos::Internal::FClothingInteractorConfigCommand::CreateLambda([LinearVelocityScale, AngularVelocityScale, FictitiousAngularScale](Chaos::FClothingSimulationConfig* Config, int32 LODIndex)
 	{
 		Config->GetProperties(LODIndex).SetValue(Chaos::FClothingSimulationCloth::LinearVelocityScaleName, FVector3f(LinearVelocityScale));
 		Config->GetProperties(LODIndex).SetValue(Chaos::FClothingSimulationCloth::AngularVelocityScaleName, AngularVelocityScale);
@@ -208,7 +208,7 @@ void UChaosClothingInteractor::SetVelocityScale(FVector LinearVelocityScale, flo
 void UChaosClothingInteractor::SetVelocityClamps(bool bEnableLinearVelocityClamping, FVector MaxLinearVelocity, bool bEnableLinearAccelerationClamping, FVector MaxLinearAcceleration,
 	bool bEnableAngularVelocityClamping, float MaxAngularVelocity, bool bEnableAngularAccelerationClamping, float MaxAngularAcceleration)
 {
-	ConfigCommands.Add(FChaosClothingInteractorConfigCommand::CreateLambda([bEnableLinearVelocityClamping, MaxLinearVelocity, bEnableLinearAccelerationClamping, MaxLinearAcceleration, bEnableAngularVelocityClamping, MaxAngularVelocity, bEnableAngularAccelerationClamping, MaxAngularAcceleration](Chaos::FClothingSimulationConfig* Config, int32 LODIndex)
+	ConfigCommands.Add(Chaos::Internal::FClothingInteractorConfigCommand::CreateLambda([bEnableLinearVelocityClamping, MaxLinearVelocity, bEnableLinearAccelerationClamping, MaxLinearAcceleration, bEnableAngularVelocityClamping, MaxAngularVelocity, bEnableAngularAccelerationClamping, MaxAngularAcceleration](Chaos::FClothingSimulationConfig* Config, int32 LODIndex)
 		{
 			static const FVector3f DefaultLinearClamp = { TNumericLimits<float>::Max(), TNumericLimits<float>::Max(), TNumericLimits<float>::Max() };
 			Config->GetProperties(LODIndex).SetValue(Chaos::FClothingSimulationCloth::MaxLinearVelocityName, bEnableLinearVelocityClamping ? FVector3f(MaxLinearVelocity) : DefaultLinearClamp);
@@ -222,31 +222,31 @@ void UChaosClothingInteractor::ResetAndTeleport(bool bReset, bool bTeleport)
 {
 	if (bReset)
 	{
-		Commands.Add(FChaosClothingInteractorCommand::CreateLambda([](Chaos::FClothingSimulationCloth* Cloth)
+		Commands.Add(Chaos::Internal::FClothingInteractorCommand::CreateLambda([](Chaos::FClothingSimulationCloth* Cloth)
 		{
 			Cloth->Reset();
 		}));
 	}
 	if (bTeleport)
 	{
-		Commands.Add(FChaosClothingInteractorCommand::CreateLambda([](Chaos::FClothingSimulationCloth* Cloth)
+		Commands.Add(Chaos::Internal::FClothingInteractorCommand::CreateLambda([](Chaos::FClothingSimulationCloth* Cloth)
 		{
 			Cloth->Teleport();
 		}));
 	}
 }
 
-void UChaosClothingSimulationInteractor::Sync(IClothingSimulation* Simulation, IClothingSimulationContext* Context)
+void UChaosClothingSimulationInteractor::Sync(IClothingSimulationInterface* Simulation, const IClothingSimulationContext* Context)
 {
 	check(Simulation);
 	check(Context);
 
-	for (FChaosClothingSimulationInteractorCommand& Command : Commands)
+	for (Chaos::Internal::FClothingSimulationInteractorCommand& Command : Commands)
 	{
-		Command.Execute(static_cast<Chaos::FClothingSimulation*>(Simulation), static_cast<Chaos::FClothingSimulationContext*>(Context));
+		Command.Execute(static_cast<Chaos::FClothingSimulation*>(Simulation), static_cast<const Chaos::FClothingSimulationContext*>(Context));
 	}
 	Commands.Reset();
-	for (FChaosClothingInteractorConfigCommand& ConfigCommand : ConfigCommands)
+	for (Chaos::Internal::FClothingInteractorConfigCommand& ConfigCommand : ConfigCommands)
 	{
 		ensure(static_cast<Chaos::FClothingSimulation*>(Simulation)->GetSolver()->GetConfig()->IsLegacySingleLOD());
 		ConfigCommand.Execute(static_cast<Chaos::FClothingSimulation*>(Simulation)->GetSolver()->GetConfig(), 0);
@@ -259,7 +259,7 @@ void UChaosClothingSimulationInteractor::Sync(IClothingSimulation* Simulation, I
 
 void UChaosClothingSimulationInteractor::PhysicsAssetUpdated()
 {
-	Commands.Add(FChaosClothingSimulationInteractorCommand::CreateLambda([](Chaos::FClothingSimulation* Simulation, Chaos::FClothingSimulationContext* /*Context*/)
+	Commands.Add(Chaos::Internal::FClothingSimulationInteractorCommand::CreateLambda([](Chaos::FClothingSimulation* Simulation, const Chaos::FClothingSimulationContext* /*Context*/)
 	{
 		Simulation->RefreshPhysicsAsset();
 	}));
@@ -267,7 +267,7 @@ void UChaosClothingSimulationInteractor::PhysicsAssetUpdated()
 
 void UChaosClothingSimulationInteractor::ClothConfigUpdated()
 {
-	Commands.Add(FChaosClothingSimulationInteractorCommand::CreateLambda([](Chaos::FClothingSimulation* Simulation, Chaos::FClothingSimulationContext* Context)
+	Commands.Add(Chaos::Internal::FClothingSimulationInteractorCommand::CreateLambda([](Chaos::FClothingSimulation* Simulation, const Chaos::FClothingSimulationContext* Context)
 	{
 		Simulation->RefreshClothConfig(Context);
 	}));
@@ -287,7 +287,7 @@ void UChaosClothingSimulationInteractor::SetAnimDriveSpringStiffness(float Stiff
 
 void UChaosClothingSimulationInteractor::EnableGravityOverride(const FVector& Gravity)
 {
-	Commands.Add(FChaosClothingSimulationInteractorCommand::CreateLambda([Gravity](Chaos::FClothingSimulation* Simulation, Chaos::FClothingSimulationContext* /*Context*/)
+	Commands.Add(Chaos::Internal::FClothingSimulationInteractorCommand::CreateLambda([Gravity](Chaos::FClothingSimulation* Simulation, const Chaos::FClothingSimulationContext* /*Context*/)
 	{
 		Simulation->SetGravityOverride(Gravity);
 	}));
@@ -295,7 +295,7 @@ void UChaosClothingSimulationInteractor::EnableGravityOverride(const FVector& Gr
 
 void UChaosClothingSimulationInteractor::DisableGravityOverride()
 {
-	Commands.Add(FChaosClothingSimulationInteractorCommand::CreateLambda([](Chaos::FClothingSimulation* Simulation, Chaos::FClothingSimulationContext* /*Context*/)
+	Commands.Add(Chaos::Internal::FClothingSimulationInteractorCommand::CreateLambda([](Chaos::FClothingSimulation* Simulation, const Chaos::FClothingSimulationContext* /*Context*/)
 	{
 		Simulation->DisableGravityOverride();
 	}));
@@ -303,7 +303,7 @@ void UChaosClothingSimulationInteractor::DisableGravityOverride()
 
 void UChaosClothingSimulationInteractor::SetNumIterations(int32 NumIterations)
 {
-	ConfigCommands.Add(FChaosClothingInteractorConfigCommand::CreateLambda([NumIterations](Chaos::FClothingSimulationConfig* Config, int32 LODIndex)
+	ConfigCommands.Add(Chaos::Internal::FClothingInteractorConfigCommand::CreateLambda([NumIterations](Chaos::FClothingSimulationConfig* Config, int32 LODIndex)
 	{
 		Config->GetProperties(LODIndex).SetValue(Chaos::FClothingSimulationSolver::NumIterationsName, NumIterations);
 	}));
@@ -311,7 +311,7 @@ void UChaosClothingSimulationInteractor::SetNumIterations(int32 NumIterations)
 
 void UChaosClothingSimulationInteractor::SetMaxNumIterations(int32 MaxNumIterations)
 {
-	ConfigCommands.Add(FChaosClothingInteractorConfigCommand::CreateLambda([MaxNumIterations](Chaos::FClothingSimulationConfig* Config, int32 LODIndex)
+	ConfigCommands.Add(Chaos::Internal::FClothingInteractorConfigCommand::CreateLambda([MaxNumIterations](Chaos::FClothingSimulationConfig* Config, int32 LODIndex)
 	{
 		Config->GetProperties(LODIndex).SetValue(Chaos::FClothingSimulationSolver::MaxNumIterationsName, MaxNumIterations);
 	}));
@@ -319,7 +319,7 @@ void UChaosClothingSimulationInteractor::SetMaxNumIterations(int32 MaxNumIterati
 
 void UChaosClothingSimulationInteractor::SetNumSubsteps(int32 NumSubsteps)
 {
-	ConfigCommands.Add(FChaosClothingInteractorConfigCommand::CreateLambda([NumSubsteps](Chaos::FClothingSimulationConfig* Config, int32 LODIndex)
+	ConfigCommands.Add(Chaos::Internal::FClothingInteractorConfigCommand::CreateLambda([NumSubsteps](Chaos::FClothingSimulationConfig* Config, int32 LODIndex)
 	{
 		Config->GetProperties(LODIndex).SetValue(Chaos::FClothingSimulationSolver::NumSubstepsName, NumSubsteps);
 	}));
diff --git a/Engine/Plugins/ChaosCloth/Source/ChaosCloth/Public/ChaosCloth/ChaosClothingSimulation.h b/Engine/Plugins/ChaosCloth/Source/ChaosCloth/Public/ChaosCloth/ChaosClothingSimulation.h
--- a/Engine/Plugins/ChaosCloth/Source/ChaosCloth/Public/ChaosCloth/ChaosClothingSimulation.h
+++ b/Engine/Plugins/ChaosCloth/Source/ChaosCloth/Public/ChaosCloth/ChaosClothingSimulation.h
@@ -33,7 +33,7 @@ namespace Chaos
 		friend FSkeletalMeshCacheAdapter;
 
 	protected:
-		// IClothingSimulation interface
+		//~ Begin IClothingSimulationInterface Interface
 		virtual void Initialize() override;
 		virtual void Shutdown() override;
 
@@ -46,8 +46,8 @@ namespace Chaos
 
 		virtual void FillContextAndPrepareTick(const USkeletalMeshComponent* InComponent, float InDeltaTime, IClothingSimulationContext* InOutContext, bool bIsInitialization, bool bForceTeleportResetOnly) override;
 		virtual bool ShouldSimulateLOD(int32 OwnerLODIndex) const override;
-		virtual void Simulate(const IClothingSimulationContext* InContext) override;
-		virtual void ForceClothNextUpdateTeleportAndReset() override;
+		virtual void Simulate_AnyThread(const IClothingSimulationContext* InContext) override;
+		virtual void ForceClothNextUpdateTeleportAndReset_AnyThread() override;
 		virtual void HardResetSimulation(const IClothingSimulationContext* InContext) override
 		{
 			RefreshClothConfig(InContext);
@@ -60,7 +60,7 @@ namespace Chaos
 		virtual void AddExternalCollisions(const FClothCollisionData& InData) override;
 		virtual void ClearExternalCollisions() override;
 		virtual void GetCollisions(FClothCollisionData& OutCollisions, bool bIncludeExternal = true) const override;
-		// End of IClothingSimulation interface
+		//~ End IClothingSimulationInterface Interface
 
 	public:
 		void SetGravityOverride(const FVector& InGravityOverride);
@@ -73,10 +73,7 @@ namespace Chaos
 		// bFullClothRefresh set to false during changes created using the viewport manipulators.
 		void RefreshPhysicsAsset();
 
-		// IClothingSimulation interface
-		virtual void SetNumIterations(int32 NumIterations) override;
-		virtual void SetMaxNumIterations(int32 MaxNumIterations) override;
-		virtual void SetNumSubsteps(int32 NumSubsteps) override;
+		//~ Begin IClothingSimulationInterface Interface
 		virtual int32 GetNumCloths() const override { return NumCloths; }
 		virtual int32 GetNumKinematicParticles() const override { return NumKinematicParticles; }
 		virtual int32 GetNumDynamicParticles() const override { return NumDynamicParticles; }
@@ -84,8 +81,7 @@ namespace Chaos
 		virtual int32 GetNumSubsteps() const override { return NumSubsteps; }
 		virtual float GetSimulationTime() const override { return SimulationTime; }
 		virtual bool IsTeleported() const override { return bIsTeleported; }
-		virtual void UpdateWorldForces(const USkeletalMeshComponent* OwnerComponent) override;
-		// End of IClothingSimulation interface
+		//~ End IClothingSimulationInterface Interface
 
 		FClothingSimulationCloth* GetCloth(int32 ClothId);
 
diff --git a/Engine/Plugins/ChaosCloth/Source/ChaosCloth/Public/ChaosCloth/ChaosClothingSimulationFactory.h b/Engine/Plugins/ChaosCloth/Source/ChaosCloth/Public/ChaosCloth/ChaosClothingSimulationFactory.h
--- a/Engine/Plugins/ChaosCloth/Source/ChaosCloth/Public/ChaosCloth/ChaosClothingSimulationFactory.h
+++ b/Engine/Plugins/ChaosCloth/Source/ChaosCloth/Public/ChaosCloth/ChaosClothingSimulationFactory.h
@@ -11,8 +11,8 @@ class UChaosClothingSimulationFactory final : public UClothingSimulationFactory
 {
 	GENERATED_BODY()
 public:
-	CHAOSCLOTH_API virtual IClothingSimulation* CreateSimulation() override;
-	CHAOSCLOTH_API virtual void DestroySimulation(IClothingSimulation* InSimulation) override;
+	CHAOSCLOTH_API virtual IClothingSimulationInterface* CreateSimulation() const override;
+	CHAOSCLOTH_API virtual void DestroySimulation(IClothingSimulationInterface* InSimulation) const override;
 	CHAOSCLOTH_API virtual bool SupportsAsset(const UClothingAssetBase* InAsset) const override;
 
 	CHAOSCLOTH_API virtual bool SupportsRuntimeInteraction() const override;
diff --git a/Engine/Plugins/ChaosCloth/Source/ChaosCloth/Public/ChaosCloth/ChaosClothingSimulationInteractor.h b/Engine/Plugins/ChaosCloth/Source/ChaosCloth/Public/ChaosCloth/ChaosClothingSimulationInteractor.h
--- a/Engine/Plugins/ChaosCloth/Source/ChaosCloth/Public/ChaosCloth/ChaosClothingSimulationInteractor.h
+++ b/Engine/Plugins/ChaosCloth/Source/ChaosCloth/Public/ChaosCloth/ChaosClothingSimulationInteractor.h
@@ -13,10 +13,17 @@ namespace Chaos
 }
 class FClothingSimulationContextCommon;
 
-// Command signature for handling synced command buffer
-DECLARE_DELEGATE_TwoParams(FChaosClothingInteractorConfigCommand, Chaos::FClothingSimulationConfig*, int32)
-DECLARE_DELEGATE_OneParam(FChaosClothingInteractorCommand, Chaos::FClothingSimulationCloth*)
-DECLARE_DELEGATE_TwoParams(FChaosClothingSimulationInteractorCommand, Chaos::FClothingSimulation*, FClothingSimulationContextCommon*)
+namespace Chaos::Internal
+{
+	// Command signature for handling synced command buffer
+	DECLARE_DELEGATE_TwoParams(FClothingInteractorConfigCommand, Chaos::FClothingSimulationConfig*, int32)
+	DECLARE_DELEGATE_OneParam(FClothingInteractorCommand, Chaos::FClothingSimulationCloth*)
+	DECLARE_DELEGATE_TwoParams(FClothingSimulationInteractorCommand, Chaos::FClothingSimulation*, const FClothingSimulationContextCommon*)
+}
+
+using FChaosClothingInteractorConfigCommand UE_DEPRECATED(5.7, "For internal use only. Use Chaos::Internal version instead.") = TDelegate<void(Chaos::FClothingSimulationConfig*, int32)>;
+using FChaosClothingInteractorCommand UE_DEPRECATED(5.7, "For internal use only. Use Chaos::Internal version instead.") = TDelegate<void(Chaos::FClothingSimulationCloth*)>;
+using FChaosClothingSimulationInteractorCommand UE_DEPRECATED(5.7, "For internal use only. Use Chaos::Internal version instead.") = TDelegate<void(Chaos::FClothingSimulation*, FClothingSimulationContextCommon*)>;
 
 UCLASS(BlueprintType, MinimalAPI)
 class UChaosClothingInteractor : public UClothingInteractor
@@ -24,7 +31,7 @@ class UChaosClothingInteractor : public UClothingInteractor
 	GENERATED_BODY()
 
 public:
-	CHAOSCLOTH_API virtual void Sync(IClothingSimulation* Simulation);
+	CHAOSCLOTH_API virtual void Sync(IClothingSimulationInterface* Simulation);
 
 	UFUNCTION(BlueprintCallable, Category = ClothingSimulation, Meta = (Keywords = "Chaos Edge Bending Area Stiffness"))
 	CHAOSCLOTH_API void SetMaterialLinear(float EdgeStiffness = 1.f, float BendingStiffness = 1.f, float AreaStiffness = 1.f);
@@ -83,8 +90,8 @@ public:
 
 private:
 	// Cloth command queue processed when we hit a sync
-	TArray<FChaosClothingInteractorCommand> Commands;
-	TArray<FChaosClothingInteractorConfigCommand> ConfigCommands;
+	TArray<Chaos::Internal::FClothingInteractorCommand> Commands;
+	TArray<Chaos::Internal::FClothingInteractorConfigCommand> ConfigCommands;
 };
 
 UCLASS(BlueprintType, MinimalAPI)
@@ -93,7 +100,7 @@ class UChaosClothingSimulationInteractor : public UClothingSimulationInteractor
 	GENERATED_BODY()
 public:
 	// UClothingSimulationInteractor interface
-	CHAOSCLOTH_API virtual void Sync(IClothingSimulation* InSimulation, IClothingSimulationContext* InContext) override;
+	CHAOSCLOTH_API virtual void Sync(IClothingSimulationInterface* InSimulation, const IClothingSimulationContext* InContext) override;
 
 	CHAOSCLOTH_API virtual void PhysicsAssetUpdated() override;
 	CHAOSCLOTH_API virtual void ClothConfigUpdated() override;
@@ -109,6 +116,6 @@ protected:
 
 private:
 	// Simulation command queue processed when we hit a sync
-	TArray<FChaosClothingSimulationInteractorCommand> Commands;
-	TArray<FChaosClothingInteractorConfigCommand> ConfigCommands;
+	TArray<Chaos::Internal::FClothingSimulationInteractorCommand> Commands;
+	TArray<Chaos::Internal::FClothingInteractorConfigCommand> ConfigCommands;
 };
diff --git a/Engine/Plugins/ChaosCloth/Source/ChaosClothEditor/Private/ChaosClothEditor/ChaosSimulationEditorExtender.cpp b/Engine/Plugins/ChaosCloth/Source/ChaosClothEditor/Private/ChaosClothEditor/ChaosSimulationEditorExtender.cpp
--- a/Engine/Plugins/ChaosCloth/Source/ChaosClothEditor/Private/ChaosClothEditor/ChaosSimulationEditorExtender.cpp
+++ b/Engine/Plugins/ChaosCloth/Source/ChaosClothEditor/Private/ChaosClothEditor/ChaosSimulationEditorExtender.cpp
@@ -145,7 +145,7 @@ void FSimulationEditorExtender::ExtendViewportShowMenu(FMenuBuilder& MenuBuilder
 	MenuBuilder.EndSection();
 }
 
-void FSimulationEditorExtender::DebugDrawSimulation(const IClothingSimulation* Simulation, USkeletalMeshComponent* OwnerComponent, FPrimitiveDrawInterface* PDI)
+void FSimulationEditorExtender::DebugDrawSimulation(const IClothingSimulationInterface* Simulation, USkeletalMeshComponent* OwnerComponent, FPrimitiveDrawInterface* PDI)
 {
 	if (!ensure(Simulation) || (OwnerComponent && OwnerComponent->bDisableClothSimulation))  // Can't debug draw when the simulation is disabled
 	{
@@ -165,7 +165,7 @@ void FSimulationEditorExtender::DebugDrawSimulation(const IClothingSimulation* S
 	}
 }
 
-void FSimulationEditorExtender::DebugDrawSimulationTexts(const IClothingSimulation* Simulation, USkeletalMeshComponent* OwnerComponent, FCanvas* Canvas, const FSceneView* SceneView)
+void FSimulationEditorExtender::DebugDrawSimulationTexts(const IClothingSimulationInterface* Simulation, USkeletalMeshComponent* OwnerComponent, FCanvas* Canvas, const FSceneView* SceneView)
 {
 	if (!ensure(Simulation) || (OwnerComponent && OwnerComponent->bDisableClothSimulation))  // Can't debug draw when the simulation is disabled
 	{
@@ -195,7 +195,7 @@ FText FSimulationEditorExtender::GetSectionHeadingText() const
 	return FText(LOCTEXT("VisualizationSection", "Skeletal Mesh Clothing Visualization"));
 }
 
-const FClothVisualizationNoGC* FSimulationEditorExtender::GetClothVisualization(const IClothingSimulation* Simulation) const
+const FClothVisualizationNoGC* FSimulationEditorExtender::GetClothVisualization(const IClothingSimulationInterface* Simulation) const
 {
 	return Simulation ? static_cast<const FClothingSimulation*>(Simulation)->GetClothVisualization() : nullptr;
 }
diff --git a/Engine/Plugins/ChaosCloth/Source/ChaosClothEditor/Public/ChaosClothEditor/ChaosSimulationEditorExtender.h b/Engine/Plugins/ChaosCloth/Source/ChaosClothEditor/Public/ChaosClothEditor/ChaosSimulationEditorExtender.h
--- a/Engine/Plugins/ChaosCloth/Source/ChaosClothEditor/Public/ChaosClothEditor/ChaosSimulationEditorExtender.h
+++ b/Engine/Plugins/ChaosCloth/Source/ChaosClothEditor/Public/ChaosClothEditor/ChaosSimulationEditorExtender.h
@@ -18,8 +18,8 @@ namespace Chaos
 		// ISimulationEditorExtender Interface
 		virtual UClass* GetSupportedSimulationFactoryClass() override;
 		CHAOSCLOTHEDITOR_API virtual void ExtendViewportShowMenu(FMenuBuilder& MenuBuilder, TSharedRef<IPersonaPreviewScene> PreviewScene) override;
-		CHAOSCLOTHEDITOR_API virtual void DebugDrawSimulation(const IClothingSimulation* Simulation, USkeletalMeshComponent* OwnerComponent, FPrimitiveDrawInterface* PDI) override;
-		CHAOSCLOTHEDITOR_API virtual void DebugDrawSimulationTexts(const IClothingSimulation* Simulation, USkeletalMeshComponent* OwnerComponent, FCanvas* Canvas, const FSceneView* SceneView) override;
+		CHAOSCLOTHEDITOR_API virtual void DebugDrawSimulation(const IClothingSimulationInterface* Simulation, USkeletalMeshComponent* OwnerComponent, FPrimitiveDrawInterface* PDI) override;
+		CHAOSCLOTHEDITOR_API virtual void DebugDrawSimulationTexts(const IClothingSimulationInterface* Simulation, USkeletalMeshComponent* OwnerComponent, FCanvas* Canvas, const FSceneView* SceneView) override;
 		// End ISimulationEditorExtender Interface
 
 	private:
@@ -30,7 +30,7 @@ namespace Chaos
 		virtual FText GetSectionHeadingText() const;
 
 		/** Return the visualization object for the specified simulation. */
-		virtual const FClothVisualizationNoGC* GetClothVisualization(const IClothingSimulation* Simulation) const;
+		virtual const FClothVisualizationNoGC* GetClothVisualization(const IClothingSimulationInterface* Simulation) const;
 
 		/** Return whether or not - given the current enabled options - the simulation should be disabled. */
 		bool ShouldDisableSimulation() const;
diff --git a/Engine/Plugins/ChaosClothAsset/Source/ChaosClothAssetEngine/Private/ChaosClothAsset/ClothAssetSKMClothingSimulation.cpp b/Engine/Plugins/ChaosClothAsset/Source/ChaosClothAssetEngine/Private/ChaosClothAsset/ClothAssetSKMClothingSimulation.cpp
--- a/Engine/Plugins/ChaosClothAsset/Source/ChaosClothAssetEngine/Private/ChaosClothAsset/ClothAssetSKMClothingSimulation.cpp
+++ b/Engine/Plugins/ChaosClothAsset/Source/ChaosClothAssetEngine/Private/ChaosClothAsset/ClothAssetSKMClothingSimulation.cpp
@@ -118,11 +118,11 @@ namespace UE::Chaos::ClothAsset
 				// Keep track of the component
 				if (!OwnerComponent || OwnerComponent != InOwnerComponent)
 				{
-					checkf(!OwnerComponent, TEXT("IClothingSimulation::CreateActor() can't be called from multiple components."));
+					checkf(!OwnerComponent, TEXT("IClothingSimulationInterface::CreateActor() can't be called from multiple components."));
 					OwnerComponent = InOwnerComponent;
 
 					// Set the collision sources when CreateActor is setting the component
-					checkf(CollisionSources, TEXT("IClothingSimulation::Initialize() must be called before calling IClothingSimulation::CreateActor()."));
+					checkf(CollisionSources, TEXT("IClothingSimulationInterface::Initialize() must be called before calling IClothingSimulationInterface::CreateActor()."));
 					CollisionSources->SetOwnerComponent(OwnerComponent);
 					CollisionSources->SetCollideWithEnvironment(OwnerComponent->bCollideWithEnvironment);
 				}
@@ -186,7 +186,7 @@ namespace UE::Chaos::ClothAsset
 		return LODHasAnyRenderClothMappingData.IsValidIndex(OwnerLODIndex) && LODHasAnyRenderClothMappingData[OwnerLODIndex] && HasAnySimulationMeshData(OwnerLODIndex);
 	}
 
-	void FSKMClothingSimulation::Simulate(const IClothingSimulationContext* Context)
+	void FSKMClothingSimulation::Simulate_AnyThread(const IClothingSimulationContext* Context)
 	{
 		if (const FSKMClothingSimulationContext* const SKMClothingSimulationContext = static_cast<const FSKMClothingSimulationContext*>(Context))
 		{
@@ -197,7 +197,7 @@ namespace UE::Chaos::ClothAsset
 		}
 	}
 	
-	void FSKMClothingSimulation::ForceClothNextUpdateTeleportAndReset()
+	void FSKMClothingSimulation::ForceClothNextUpdateTeleportAndReset_AnyThread()
 	{
 		// Do nothing here. This is handled in FillContextAndPrepareTick
 	}
diff --git a/Engine/Plugins/ChaosClothAsset/Source/ChaosClothAssetEngine/Private/ChaosClothAsset/ClothAssetSKMClothingSimulationInteractor.cpp b/Engine/Plugins/ChaosClothAsset/Source/ChaosClothAssetEngine/Private/ChaosClothAsset/ClothAssetSKMClothingSimulationInteractor.cpp
--- a/Engine/Plugins/ChaosClothAsset/Source/ChaosClothAssetEngine/Private/ChaosClothAsset/ClothAssetSKMClothingSimulationInteractor.cpp
+++ b/Engine/Plugins/ChaosClothAsset/Source/ChaosClothAssetEngine/Private/ChaosClothAsset/ClothAssetSKMClothingSimulationInteractor.cpp
@@ -15,7 +15,7 @@ void UChaosClothAssetSKMClothingInteractor::Initialize(const UChaosClothAssetSKM
 	ClothingAsset = InClothingAsset;
 }
 
-void UChaosClothAssetSKMClothingSimulationInteractor::Sync(IClothingSimulation* Simulation, IClothingSimulationContext* Context)
+void UChaosClothAssetSKMClothingSimulationInteractor::Sync(IClothingSimulationInterface* Simulation, const IClothingSimulationContext* Context)
 {
 	using namespace UE::Chaos::ClothAsset;
 
@@ -97,7 +97,10 @@ void UChaosClothAssetSKMClothingSimulationInteractor::SetNumIterations(int32 Num
 	using namespace UE::Chaos::ClothAsset;
 	Commands.Add(FCommand::CreateLambda([NumIterations](FSKMClothingSimulation* SKMClothingSimulation)
 		{
-			static_cast<IClothingSimulation*>(SKMClothingSimulation)->SetNumIterations(NumIterations);
+			if (SKMClothingSimulation)
+			{
+				SKMClothingSimulation->SetNumIterations(NumIterations);
+			}
 		}));
 }
 
@@ -106,7 +109,10 @@ void UChaosClothAssetSKMClothingSimulationInteractor::SetMaxNumIterations(int32
 	using namespace UE::Chaos::ClothAsset;
 	Commands.Add(FCommand::CreateLambda([MaxNumIterations](FSKMClothingSimulation* SKMClothingSimulation)
 		{
-			static_cast<IClothingSimulation*>(SKMClothingSimulation)->SetMaxNumIterations(MaxNumIterations);
+			if (SKMClothingSimulation)
+			{
+				SKMClothingSimulation->SetMaxNumIterations(MaxNumIterations);
+			}
 		}));
 }
 
@@ -115,7 +121,10 @@ void UChaosClothAssetSKMClothingSimulationInteractor::SetNumSubsteps(int32 NumSu
 	using namespace UE::Chaos::ClothAsset;
 	Commands.Add(FCommand::CreateLambda([NumSubsteps](FSKMClothingSimulation* SKMClothingSimulation)
 		{
-			static_cast<IClothingSimulation*>(SKMClothingSimulation)->SetNumSubsteps(NumSubsteps);
+			if (SKMClothingSimulation)
+			{
+				SKMClothingSimulation->SetNumSubsteps(NumSubsteps);
+			}
 		}));
 }
 
diff --git a/Engine/Plugins/ChaosClothAsset/Source/ChaosClothAssetEngine/Private/ChaosClothAsset/ClothAssetSKMClothingSimulationInteractor.h b/Engine/Plugins/ChaosClothAsset/Source/ChaosClothAssetEngine/Private/ChaosClothAsset/ClothAssetSKMClothingSimulationInteractor.h
--- a/Engine/Plugins/ChaosClothAsset/Source/ChaosClothAssetEngine/Private/ChaosClothAsset/ClothAssetSKMClothingSimulationInteractor.h
+++ b/Engine/Plugins/ChaosClothAsset/Source/ChaosClothAssetEngine/Private/ChaosClothAsset/ClothAssetSKMClothingSimulationInteractor.h
@@ -37,7 +37,7 @@ private:
 
 	//~ Begin UClothingSimulationInteractor Interface
 	virtual void CreateClothingInteractor(const UClothingAssetBase* ClothingAsset, int32 ClothingId) override;
-	virtual void Sync(IClothingSimulation* InSimulation, IClothingSimulationContext* InContext) override;
+	virtual void Sync(IClothingSimulationInterface* InSimulation, const IClothingSimulationContext* InContext) override;
 
 	virtual void PhysicsAssetUpdated() override;
 	virtual void ClothConfigUpdated() override;
diff --git a/Engine/Plugins/ChaosClothAsset/Source/ChaosClothAssetEngine/Public/ChaosClothAsset/ClothAssetSKMClothingSimulation.h b/Engine/Plugins/ChaosClothAsset/Source/ChaosClothAssetEngine/Public/ChaosClothAsset/ClothAssetSKMClothingSimulation.h
--- a/Engine/Plugins/ChaosClothAsset/Source/ChaosClothAssetEngine/Public/ChaosClothAsset/ClothAssetSKMClothingSimulation.h
+++ b/Engine/Plugins/ChaosClothAsset/Source/ChaosClothAssetEngine/Public/ChaosClothAsset/ClothAssetSKMClothingSimulation.h
@@ -12,6 +12,7 @@
 class USkinnedMeshComponent;
 class UChaosClothAssetInteractor;
 class UChaosClothAssetSKMClothingAsset;
+class UChaosClothAssetSKMClothingSimulationInteractor;
 struct FChaosClothSimulationProperties;
 
 namespace Chaos
@@ -41,7 +42,7 @@ namespace UE::Chaos::ClothAsset
 
 	/** Skeletal Mesh simulation class that interfaces between the Skeletal Mesh Component and the Cloth Simulation Proxy. */
 	class FSKMClothingSimulation final
-		: public IClothingSimulation
+		: public IClothingSimulationInterface
 		, public IClothComponentAdapter
 		, public FGCObject
 	{
@@ -65,7 +66,13 @@ namespace UE::Chaos::ClothAsset
 		}
 
 	private:
-		//~ Begin IClothingSimulation Interface
+		friend UChaosClothAssetSKMClothingSimulationInteractor;  // For the SetNumIterations/SetNumSubsteps methods
+
+		void SetNumIterations(int32 NumIterations);
+		void SetMaxNumIterations(int32 MaxNumIterations);
+		void SetNumSubsteps(int32 NumSubsteps);
+
+		//~ Begin IClothingSimulationInterface Interface
 		virtual void Initialize() override;
 		virtual void Shutdown() override;
 
@@ -84,8 +91,8 @@ namespace UE::Chaos::ClothAsset
 		virtual void DestroyActors() override;
 
 		virtual bool ShouldSimulateLOD(int32 OwnerLODIndex) const override;
-		virtual void Simulate(const IClothingSimulationContext* Context) override;
-		virtual void ForceClothNextUpdateTeleportAndReset() override;
+		virtual void Simulate_AnyThread(const IClothingSimulationContext* Context) override;
+		virtual void ForceClothNextUpdateTeleportAndReset_AnyThread() override;
 		virtual void HardResetSimulation(const IClothingSimulationContext* InContext) override;
 		virtual void AppendSimulationData(TMap<int32, FClothSimulData>& OutData, const USkeletalMeshComponent* OwnerComponent, const USkinnedMeshComponent* OverrideComponent) const override;
 
@@ -94,9 +101,6 @@ namespace UE::Chaos::ClothAsset
 		virtual void AddExternalCollisions(const FClothCollisionData& Data) override;
 		virtual void ClearExternalCollisions() override;
 		virtual void GetCollisions(FClothCollisionData& OutCollisions, bool bIncludeExternal = true) const override;
-		virtual void SetNumIterations(int32 NumIterations) override;
-		virtual void SetMaxNumIterations(int32 MaxNumIterations) override;
-		virtual void SetNumSubsteps(int32 NumSubsteps) override;
 		virtual int32 GetNumCloths() const override { return ClothSimulationProxy.IsValid() ? ClothSimulationProxy->GetNumCloths() : 0; }
 		virtual int32 GetNumKinematicParticles() const override { return ClothSimulationProxy.IsValid() ? ClothSimulationProxy->GetNumKinematicParticles() : 0; }
 		virtual int32 GetNumDynamicParticles() const override { return ClothSimulationProxy.IsValid() ? ClothSimulationProxy->GetNumDynamicParticles() : 0; }
@@ -104,7 +108,7 @@ namespace UE::Chaos::ClothAsset
 		virtual int32 GetNumSubsteps() const override { return ClothSimulationProxy.IsValid() ? ClothSimulationProxy->GetNumSubsteps() : 0; }
 		virtual float GetSimulationTime() const override { return ClothSimulationProxy.IsValid() ? ClothSimulationProxy->GetSimulationTime() : 0.f; }
 		virtual bool IsTeleported() const override { return ClothSimulationProxy.IsValid() ? ClothSimulationProxy->IsTeleported() : false; }
-		//~ End IClothingSimulation Interface
+		//~ End IClothingSimulationInterface Interface
 
 		//~ Begin IClothComponentAdapter Interface
 		virtual const USkinnedMeshComponent& GetOwnerComponent() const override;
@@ -169,12 +173,12 @@ public:
 
 private:
 	//~ Begin UClothingSimulationFactory Interface
-	virtual IClothingSimulation* CreateSimulation() override
+	virtual IClothingSimulationInterface* CreateSimulation() const override
 	{
-		return static_cast<IClothingSimulation*>(new UE::Chaos::ClothAsset::FSKMClothingSimulation());
+		return static_cast<IClothingSimulationInterface*>(new UE::Chaos::ClothAsset::FSKMClothingSimulation());
 	}
 
-	virtual void DestroySimulation(IClothingSimulation* Simulation) override
+	virtual void DestroySimulation(IClothingSimulationInterface* Simulation) const override
 	{
 		delete Simulation;
 	}
diff --git a/Engine/Plugins/ChaosClothAssetEditor/Source/ChaosClothAssetTools/Private/ChaosClothAsset/ClothAssetSKMSimulationEditorExtender.cpp b/Engine/Plugins/ChaosClothAssetEditor/Source/ChaosClothAssetTools/Private/ChaosClothAsset/ClothAssetSKMSimulationEditorExtender.cpp
--- a/Engine/Plugins/ChaosClothAssetEditor/Source/ChaosClothAssetTools/Private/ChaosClothAsset/ClothAssetSKMSimulationEditorExtender.cpp
+++ b/Engine/Plugins/ChaosClothAssetEditor/Source/ChaosClothAssetTools/Private/ChaosClothAsset/ClothAssetSKMSimulationEditorExtender.cpp
@@ -29,7 +29,7 @@ namespace UE::Chaos::ClothAsset
 		return FText(LOCTEXT("ClothAssetVisualizationSection", "Cloth Asset Visualization"));
 	}
 
-	const ::Chaos::FClothVisualizationNoGC* FSKMSimulationEditorExtender::GetClothVisualization(const IClothingSimulation* Simulation) const
+	const ::Chaos::FClothVisualizationNoGC* FSKMSimulationEditorExtender::GetClothVisualization(const IClothingSimulationInterface* Simulation) const
 	{
 		return Simulation ? static_cast<const FSKMClothingSimulation*>(Simulation)->GetClothVisualization() : nullptr;
 	}
diff --git a/Engine/Plugins/ChaosClothAssetEditor/Source/ChaosClothAssetTools/Private/ChaosClothAsset/ClothAssetSKMSimulationEditorExtender.h b/Engine/Plugins/ChaosClothAssetEditor/Source/ChaosClothAssetTools/Private/ChaosClothAsset/ClothAssetSKMSimulationEditorExtender.h
--- a/Engine/Plugins/ChaosClothAssetEditor/Source/ChaosClothAssetTools/Private/ChaosClothAsset/ClothAssetSKMSimulationEditorExtender.h
+++ b/Engine/Plugins/ChaosClothAssetEditor/Source/ChaosClothAssetTools/Private/ChaosClothAsset/ClothAssetSKMSimulationEditorExtender.h
@@ -26,7 +26,7 @@ namespace UE::Chaos::ClothAsset
 		//~ Begin ::Chaos::FSimulationEditorExtender Interface
 		virtual FName GetSectionName() const override;
 		virtual FText GetSectionHeadingText() const override;
-		virtual const ::Chaos::FClothVisualizationNoGC* GetClothVisualization(const IClothingSimulation* Simulation) const override;
+		virtual const ::Chaos::FClothVisualizationNoGC* GetClothVisualization(const IClothingSimulationInterface* Simulation) const override;
 		//~ End ::Chaos::FSimulationEditorExtender Interface
 	};
 }
diff --git a/Engine/Plugins/MovieScene/MovieRenderPipeline/Source/MovieRenderPipelineCore/Private/MoviePipelineUtils.cpp b/Engine/Plugins/MovieScene/MovieRenderPipeline/Source/MovieRenderPipelineCore/Private/MoviePipelineUtils.cpp
--- a/Engine/Plugins/MovieScene/MovieRenderPipeline/Source/MovieRenderPipelineCore/Private/MoviePipelineUtils.cpp
+++ b/Engine/Plugins/MovieScene/MovieRenderPipeline/Source/MovieRenderPipelineCore/Private/MoviePipelineUtils.cpp
@@ -1280,7 +1280,7 @@ namespace UE
 					{
 						for (const FClothingSimulationInstance& ClothingSimulationInstance : Component->GetClothingSimulationInstances())
 						{
-							const IClothingSimulation* const ClothingSimulation = ClothingSimulationInstance.GetClothingSimulation();
+							const IClothingSimulationInterface* const ClothingSimulation = ClothingSimulationInstance.GetClothingSimulation();
 							UClothingSimulationInteractor* const ClothingSimulationInteractor = ClothingSimulationInstance.GetClothingSimulationInteractor();
 							if (ClothingSimulation && ClothingSimulationInteractor)
 							{
diff --git a/Engine/Source/Editor/ClothingSystemEditorInterface/ClothingSystemEditorInterface.Build.cs b/Engine/Source/Editor/ClothingSystemEditorInterface/ClothingSystemEditorInterface.Build.cs
--- a/Engine/Source/Editor/ClothingSystemEditorInterface/ClothingSystemEditorInterface.Build.cs
+++ b/Engine/Source/Editor/ClothingSystemEditorInterface/ClothingSystemEditorInterface.Build.cs
@@ -10,7 +10,13 @@ public class ClothingSystemEditorInterface : ModuleRules
         PublicDependencyModuleNames.AddRange(
             new string[] {
                 "Core",
-                "CoreUObject",
+                "CoreUObject"
+            }
+        );
+
+        PrivateDependencyModuleNames.AddRange(
+            new string[] {
+                "ClothingSystemRuntimeInterface"
             }
         );
     }
diff --git a/Engine/Source/Editor/ClothingSystemEditorInterface/Private/SimulationEditorExtender.cpp b/Engine/Source/Editor/ClothingSystemEditorInterface/Private/SimulationEditorExtender.cpp
new file mode 100644
--- /dev/null
+++ b/Engine/Source/Editor/ClothingSystemEditorInterface/Private/SimulationEditorExtender.cpp
@@ -0,0 +1,26 @@
+// Copyright Epic Games, Inc. All Rights Reserved.
+
+#pragma once
+
+#include "SimulationEditorExtender.h"
+#include "ClothingSimulationInterface.h"
+
+void ISimulationEditorExtender::DebugDrawSimulation(const IClothingSimulationInterface* InSimulation, USkeletalMeshComponent* InOwnerComponent, FPrimitiveDrawInterface* PDI)
+{
+PRAGMA_DISABLE_DEPRECATION_WARNINGS
+	if (ensureMsgf(InSimulation && InSimulation->DynamicCastToIClothingSimulation(), TEXT("DebugDrawSimulation(const IClothingSimulationInterface*, ...) must be implemented from 5.7, as the function will become pure virtual.")))
+	{
+		DebugDrawSimulation(InSimulation->DynamicCastToIClothingSimulation(), InOwnerComponent, PDI);
+	}
+PRAGMA_ENABLE_DEPRECATION_WARNINGS
+}
+
+void ISimulationEditorExtender::DebugDrawSimulationTexts(const IClothingSimulationInterface* InSimulation, USkeletalMeshComponent* InOwnerComponent, FCanvas* Canvas, const FSceneView* SceneView)
+{
+PRAGMA_DISABLE_DEPRECATION_WARNINGS
+	if (ensureMsgf(InSimulation && InSimulation->DynamicCastToIClothingSimulation(), TEXT("DebugDrawSimulationTexts(const IClothingSimulationInterface*, ...) must be implemented from 5.7, as the function will become pure virtual.")))
+	{
+		DebugDrawSimulationTexts(InSimulation->DynamicCastToIClothingSimulation(), InOwnerComponent, Canvas, SceneView);
+	}
+PRAGMA_ENABLE_DEPRECATION_WARNINGS
+}
diff --git a/Engine/Source/Editor/ClothingSystemEditorInterface/Public/SimulationEditorExtender.h b/Engine/Source/Editor/ClothingSystemEditorInterface/Public/SimulationEditorExtender.h
--- a/Engine/Source/Editor/ClothingSystemEditorInterface/Public/SimulationEditorExtender.h
+++ b/Engine/Source/Editor/ClothingSystemEditorInterface/Public/SimulationEditorExtender.h
@@ -5,9 +5,15 @@
 #include "Features/IModularFeature.h"
 #include "Templates/SharedPointer.h"
 
+#define UE_API CLOTHINGSYSTEMEDITORINTERFACE_API
+
+#if UE_ENABLE_INCLUDE_ORDER_DEPRECATED_IN_5_7
+class UE_DEPRECATED(5.7, "Use IClothingSimulationInterface instead.") IClothingSimulation;
+#endif
+
 class UClass;
 class FMenuBuilder;
-class IClothingSimulation;
+class IClothingSimulationInterface;
 class USkeletalMeshComponent;
 class IPersonaPreviewScene;
 class FPrimitiveDrawInterface;
@@ -47,7 +53,10 @@ public:
 	 * @param InOwnerComponent - The component that owns the running clothing simulation
 	 * @param PDI - The drawing interface to use
 	 */
-	virtual void DebugDrawSimulation(const IClothingSimulation* InSimulation, USkeletalMeshComponent* InOwnerComponent, FPrimitiveDrawInterface* PDI) = 0;
+	UE_API virtual void DebugDrawSimulation(const IClothingSimulationInterface* InSimulation, USkeletalMeshComponent* InOwnerComponent, FPrimitiveDrawInterface* PDI);
+
+	UE_DEPRECATED(5.7, "Use IClothingSimulationInterface instead.")
+	virtual void DebugDrawSimulation(const class IClothingSimulation* InSimulation, USkeletalMeshComponent* InOwnerComponent, FPrimitiveDrawInterface* PDI) {}
 
 	/**
 	 * Called from the editor when clothing is active to process any active debug drawing of any text strings
@@ -57,5 +66,10 @@ public:
 	 * @param Canvas - The canvas to use for drawing the text
 	 * @param SceneView - The view on which to project the text
 	 */
-	virtual void DebugDrawSimulationTexts(const IClothingSimulation* InSimulation, USkeletalMeshComponent* InOwnerComponent, FCanvas* Canvas, const FSceneView* SceneView) = 0;
+	UE_API virtual void DebugDrawSimulationTexts(const IClothingSimulationInterface* InSimulation, USkeletalMeshComponent* InOwnerComponent, FCanvas* Canvas, const FSceneView* SceneView);
+
+	UE_DEPRECATED(5.7, "Use IClothingSimulationInterface instead.")
+	virtual void DebugDrawSimulationTexts(const class IClothingSimulation* InSimulation, USkeletalMeshComponent* InOwnerComponent, FCanvas* Canvas, const FSceneView* SceneView) {}
 };
+
+#undef UE_API
diff --git a/Engine/Source/Editor/Persona/Private/AnimationEditorViewportClient.cpp b/Engine/Source/Editor/Persona/Private/AnimationEditorViewportClient.cpp
--- a/Engine/Source/Editor/Persona/Private/AnimationEditorViewportClient.cpp
+++ b/Engine/Source/Editor/Persona/Private/AnimationEditorViewportClient.cpp
@@ -1393,7 +1393,7 @@ FText FAnimationViewportClient::GetDisplayInfo(bool bDisplayAllInfo) const
 	{
 		const FClothingSimulationInstance& ClothingSimulationInstance = ClothingSimulationInstances[Index];
 
-		if (const IClothingSimulation* const ClothingSimulation = ClothingSimulationInstance.GetClothingSimulation())
+		if (const IClothingSimulationInterface* const ClothingSimulation = ClothingSimulationInstance.GetClothingSimulation())
 		{
 			// Cloth stats
 			if (const int32 NumActiveCloths = ClothingSimulation->GetNumCloths())
diff --git a/Engine/Source/Editor/UnrealEd/Classes/Animation/DebugSkelMeshComponent.h b/Engine/Source/Editor/UnrealEd/Classes/Animation/DebugSkelMeshComponent.h
--- a/Engine/Source/Editor/UnrealEd/Classes/Animation/DebugSkelMeshComponent.h
+++ b/Engine/Source/Editor/UnrealEd/Classes/Animation/DebugSkelMeshComponent.h
@@ -594,8 +594,10 @@ public:
 	 * To normal game/runtime code we don't want to expose a non-const pointer to the simulation, so we can only get
 	 * one from this editor-only component. Intended for debug options/visualisations/editor-only code to poke the sim
 	 */
+PRAGMA_DISABLE_DEPRECATION_WARNINGS
 	UE_DEPRECATED(5.7, "Use GetClothingSimulationInstances instead.")
 	UNREALED_API IClothingSimulation* GetMutableClothingSimulation();
+PRAGMA_ENABLE_DEPRECATION_WARNINGS
 
 	/** to avoid clothing reset while modifying properties in Persona */
 	UNREALED_API virtual void CheckClothTeleport() override;
diff --git a/Engine/Source/Editor/UnrealEd/Private/Animation/DebugSkelMeshComponent.cpp b/Engine/Source/Editor/UnrealEd/Private/Animation/DebugSkelMeshComponent.cpp
--- a/Engine/Source/Editor/UnrealEd/Private/Animation/DebugSkelMeshComponent.cpp
+++ b/Engine/Source/Editor/UnrealEd/Private/Animation/DebugSkelMeshComponent.cpp
@@ -1284,13 +1284,23 @@ void UDebugSkelMeshComponent::GetUsedMaterials(TArray<UMaterialInterface *>& Out
 	}
 }
 
+
 // UE_DEPRECATED(5.7, "Use GetClothingSimulationInstances instead.")
+PRAGMA_DISABLE_DEPRECATION_WARNINGS
 IClothingSimulation* UDebugSkelMeshComponent::GetMutableClothingSimulation()
 {
-	// This function seems unused, and therefore a const_cast is probably a better idea than compromising the const version of GetClothingSimulation()
-	return GetClothingSimulationInstances().Num() ?
-		const_cast<IClothingSimulation*>(GetClothingSimulationInstances()[0].GetClothingSimulation()) : nullptr;
+	if (GetClothingSimulationInstances().Num())
+	{
+		// This function seems unused, and therefore a const_cast is probably a better idea than compromising the const version of GetClothingSimulation()
+		if (IClothingSimulationInterface* const ClothingSimulation =
+			const_cast<IClothingSimulationInterface*>(GetClothingSimulationInstances()[0].GetClothingSimulation()))
+		{
+			return ClothingSimulation->DynamicCastToIClothingSimulation();
+		}
+	}
+	return nullptr;
 }
+PRAGMA_ENABLE_DEPRECATION_WARNINGS
 
 void UDebugSkelMeshComponent::RebuildCachedClothBounds()
 {
diff --git a/Engine/Source/Runtime/ClothingSystemRuntimeCommon/Private/ClothingSimulation.cpp b/Engine/Source/Runtime/ClothingSystemRuntimeCommon/Private/ClothingSimulation.cpp
--- a/Engine/Source/Runtime/ClothingSystemRuntimeCommon/Private/ClothingSimulation.cpp
+++ b/Engine/Source/Runtime/ClothingSystemRuntimeCommon/Private/ClothingSimulation.cpp
@@ -195,6 +195,7 @@ void FClothingSimulationContextCommon::FillSolverGeometryScale(const USkeletalMe
 // FClothingSimulationCommon
 //==============================================================================
 
+PRAGMA_DISABLE_DEPRECATION_WARNINGS  // UE_DEPRECATED(5.7, "Use IClothingSimulationInterface instead.")
 FClothingSimulationCommon::FClothingSimulationCommon()
 {
 	MaxPhysicsDelta = UPhysicsSettings::Get()->MaxPhysicsDeltaTime;
@@ -202,6 +203,7 @@ FClothingSimulationCommon::FClothingSimulationCommon()
 
 FClothingSimulationCommon::~FClothingSimulationCommon()
 {}
+PRAGMA_ENABLE_DEPRECATION_WARNINGS
 
 void FClothingSimulationCommon::FillContextAndPrepareTick(const USkeletalMeshComponent* InComponent, float InDeltaTime, IClothingSimulationContext* InOutContext, bool bIsInitialization, bool bForceTeleportResetOnly)
 {
diff --git a/Engine/Source/Runtime/ClothingSystemRuntimeCommon/Public/ClothingSimulation.h b/Engine/Source/Runtime/ClothingSystemRuntimeCommon/Public/ClothingSimulation.h
--- a/Engine/Source/Runtime/ClothingSystemRuntimeCommon/Public/ClothingSimulation.h
+++ b/Engine/Source/Runtime/ClothingSystemRuntimeCommon/Public/ClothingSimulation.h
@@ -86,7 +86,10 @@ public:
 };
 
 // Base simulation to fill in common data for the base context
-class FClothingSimulationCommon : public IClothingSimulation
+class FClothingSimulationCommon
+PRAGMA_DISABLE_DEPRECATION_WARNINGS
+	: public IClothingSimulation  // UE_DEPRECATED(5.7, "Use IClothingSimulationInterface instead.")
+PRAGMA_ENABLE_DEPRECATION_WARNINGS
 {
 public:
 	UE_API FClothingSimulationCommon();
diff --git a/Engine/Source/Runtime/ClothingSystemRuntimeInterface/Private/ClothingSimulationInstance.cpp b/Engine/Source/Runtime/ClothingSystemRuntimeInterface/Private/ClothingSimulationInstance.cpp
--- a/Engine/Source/Runtime/ClothingSystemRuntimeInterface/Private/ClothingSimulationInstance.cpp
+++ b/Engine/Source/Runtime/ClothingSystemRuntimeInterface/Private/ClothingSimulationInstance.cpp
@@ -9,7 +9,7 @@
 #include UE_INLINE_GENERATED_CPP_BY_NAME(ClothingSimulationInstance)
 
 FClothingSimulationInstance::FClothingSimulationInstance(USkeletalMeshComponent* SkeletalMeshComponent, const FFactoryAssetGroup& FactoryAssetGroup)
-	: ClothingSimulation(FactoryAssetGroup.ClothingSimulationFactory->CreateSimulation())
+	: ClothingSimulation(static_cast<const UClothingSimulationFactory*>(FactoryAssetGroup.ClothingSimulationFactory)->CreateSimulation())  // UE_DEPRECATED(5.7, "Use IClothingSimulationInterface instead.")
 	, ClothingSimulationContext(ClothingSimulation->CreateContext())
 	, ClothingSimulationInteractor(static_cast<const UClothingSimulationFactory*>(FactoryAssetGroup.ClothingSimulationFactory)->SupportsRuntimeInteraction() ? FactoryAssetGroup.ClothingSimulationFactory->CreateInteractor() : nullptr)
 	, ClothingSimulationFactory(FactoryAssetGroup.ClothingSimulationFactory)
diff --git a/Engine/Source/Runtime/ClothingSystemRuntimeInterface/Private/ClothingSimulationInteractor.cpp b/Engine/Source/Runtime/ClothingSystemRuntimeInterface/Private/ClothingSimulationInteractor.cpp
--- a/Engine/Source/Runtime/ClothingSystemRuntimeInterface/Private/ClothingSimulationInteractor.cpp
+++ b/Engine/Source/Runtime/ClothingSystemRuntimeInterface/Private/ClothingSimulationInteractor.cpp
@@ -35,21 +35,65 @@ UClothingInteractor* UClothingSimulationInteractor::GetClothingInteractor(const
 	return UClothingInteractor::StaticClass()->GetDefaultObject<UClothingInteractor>();
 }
 
-void UClothingSimulationInteractor::Sync(IClothingSimulation* Simulation, IClothingSimulationContext* Context)
+void UClothingSimulationInteractor::Sync(IClothingSimulationInterface* Simulation, const IClothingSimulationContext* Context)
 {
-	LastNumCloths = Simulation->GetNumCloths();
-	LastNumKinematicParticles = Simulation->GetNumKinematicParticles();
-	LastNumDynamicParticles = Simulation->GetNumDynamicParticles();
-	LastNumIterations = Simulation->GetNumIterations();
-	LastNumSubsteps = Simulation->GetNumSubsteps();
-	LastSimulationTime = Simulation->GetSimulationTime();
-
-	for (const auto& ClothingInteractor : UClothingSimulationInteractor::ClothingInteractors)
+	if (Simulation)
 	{
-		if (ClothingInteractor.Value)
+PRAGMA_DISABLE_DEPRECATION_WARNINGS  // UE_DEPRECATED(5.7, "Use IClothingSimulationInterface instead.")
+		if (Simulation->DynamicCastToIClothingSimulation())
 		{
-			ClothingInteractor.Value->Sync(Simulation);
+			Sync(Simulation->DynamicCastToIClothingSimulation(), Context);
+			return;
+		}
+PRAGMA_ENABLE_DEPRECATION_WARNINGS
+
+		LastNumCloths = Simulation->GetNumCloths();
+		LastNumKinematicParticles = Simulation->GetNumKinematicParticles();
+		LastNumDynamicParticles = Simulation->GetNumDynamicParticles();
+		LastNumIterations = Simulation->GetNumIterations();
+		LastNumSubsteps = Simulation->GetNumSubsteps();
+		LastSimulationTime = Simulation->GetSimulationTime();
+
+		for (const TPair<FName, TObjectPtr<UClothingInteractor>>& ClothingInteractor : UClothingSimulationInteractor::ClothingInteractors)
+		{
+			if (ClothingInteractor.Value)
+			{
+				ClothingInteractor.Value->Sync(Simulation);
+			}
+		}
+	}
+	else
+	{
+		LastNumCloths = 0;
+		LastNumKinematicParticles = 0;
+		LastNumDynamicParticles = 0;
+		LastNumIterations = 0;
+		LastNumSubsteps = 0;
+		LastSimulationTime = 0.f;
+	}
+}
+
+
+PRAGMA_DISABLE_DEPRECATION_WARNINGS  // UE_DEPRECATED(5.7, "Use IClothingSimulationInterface instead.")
+void UClothingSimulationInteractor::Sync(IClothingSimulation* Simulation, IClothingSimulationContext* /*Context*/)  // TODO 5.9: Remove pragmas and replace IClothingSimulation with IClothingSimulationInterface
+{
+	if (Simulation)
+	{
+		LastNumCloths = Simulation->GetNumCloths();
+		LastNumKinematicParticles = Simulation->GetNumKinematicParticles();
+		LastNumDynamicParticles = Simulation->GetNumDynamicParticles();
+		LastNumIterations = Simulation->GetNumIterations();
+		LastNumSubsteps = Simulation->GetNumSubsteps();
+		LastSimulationTime = Simulation->GetSimulationTime();
+
+		for (const auto& ClothingInteractor : UClothingSimulationInteractor::ClothingInteractors)
+		{
+			if (ClothingInteractor.Value)
+			{
+				ClothingInteractor.Value->Sync(Simulation);
+			}
 		}
 	}
 }
+PRAGMA_ENABLE_DEPRECATION_WARNINGS
 
diff --git a/Engine/Source/Runtime/ClothingSystemRuntimeInterface/Private/ClothingSimulationInterface.cpp b/Engine/Source/Runtime/ClothingSystemRuntimeInterface/Private/ClothingSimulationInterface.cpp
--- a/Engine/Source/Runtime/ClothingSystemRuntimeInterface/Private/ClothingSimulationInterface.cpp
+++ b/Engine/Source/Runtime/ClothingSystemRuntimeInterface/Private/ClothingSimulationInterface.cpp
@@ -6,18 +6,24 @@
 // IClotingSimulationContext
 //==============================================================================
 
-IClothingSimulationContext::IClothingSimulationContext()
-{}
+IClothingSimulationContext::IClothingSimulationContext() = default;
+IClothingSimulationContext::~IClothingSimulationContext() = default;
 
-IClothingSimulationContext::~IClothingSimulationContext()
-{}
+//==============================================================================
+// IClothingSimulationInterface
+//==============================================================================
+
+IClothingSimulationInterface::IClothingSimulationInterface() = default;
+IClothingSimulationInterface::~IClothingSimulationInterface() = default;
 
 //==============================================================================
 // IClothingSimulation
 //==============================================================================
 
 IClothingSimulation::IClothingSimulation()
-{}
-
-IClothingSimulation::~IClothingSimulation()
-{}
+{
+PRAGMA_DISABLE_DEPRECATION_WARNINGS
+	bIsLegacyInterface = true;
+PRAGMA_ENABLE_DEPRECATION_WARNINGS
+}
+IClothingSimulation::~IClothingSimulation() = default;
diff --git a/Engine/Source/Runtime/ClothingSystemRuntimeInterface/Public/ClothingSimulationFactory.h b/Engine/Source/Runtime/ClothingSystemRuntimeInterface/Public/ClothingSimulationFactory.h
--- a/Engine/Source/Runtime/ClothingSystemRuntimeInterface/Public/ClothingSimulationFactory.h
+++ b/Engine/Source/Runtime/ClothingSystemRuntimeInterface/Public/ClothingSimulationFactory.h
@@ -17,7 +17,6 @@
 #define UE_API CLOTHINGSYSTEMRUNTIMEINTERFACE_API
 
 class FName;
-class IClothingSimulation;
 class UClass;
 class UClothConfigBase;
 class UClothingAssetBase;
@@ -56,34 +55,64 @@ public:
 	// If multiple simulation support the specified asset type then this returns the first one registered through IClothingSimulationFactoryClassProvider.
 	static UE_API UClothingSimulationFactory* GetClothingSimulationFactory(const UClothingAssetBase* InAsset);
 
-	// Create a simulation object for a skeletal mesh to use (see IClothingSimulation)
+	// Create a simulation object for a skeletal mesh to use (see IClothingSimulationInterface)
+	virtual IClothingSimulationInterface* CreateSimulation() const  // TODO 5.9: Made this method into a PURE_VIRTUAL
+	{
+PRAGMA_DISABLE_DEPRECATION_WARNINGS
+		return const_cast<UClothingSimulationFactory*>(this)->CreateSimulation();
+PRAGMA_ENABLE_DEPRECATION_WARNINGS
+	}
+
+PRAGMA_DISABLE_DEPRECATION_WARNINGS
+	UE_DEPRECATED(5.7, "Use IClothingSimulationInterface instead.")
 	virtual IClothingSimulation* CreateSimulation()
-	PURE_VIRTUAL(UClothingSimulationFactory::CreateSimulation, return nullptr;);
+	{
+		return nullptr;
+	}
+PRAGMA_ENABLE_DEPRECATION_WARNINGS
 
 	// Destroy a simulation object, guaranteed to be a pointer returned from CreateSimulation for this factory
-	virtual void DestroySimulation(IClothingSimulation* InSimulation)
-	PURE_VIRTUAL(UClothingSimulationFactory::DestroySimulation, );
+	virtual void DestroySimulation(IClothingSimulationInterface* InSimulation) const  // TODO 5.9: Make this method into a PURE_VIRTUAL
+	{
+PRAGMA_DISABLE_DEPRECATION_WARNINGS
+		if (ensureMsgf(InSimulation && InSimulation->DynamicCastToIClothingSimulation(), TEXT("DestroySimulation(IClothingSimulationInterface*) must be implemented from 5.7, as the function will become pure virtual.")))
+		{
+			const_cast<UClothingSimulationFactory*>(this)->DestroySimulation(InSimulation->DynamicCastToIClothingSimulation());
+		}
+PRAGMA_ENABLE_DEPRECATION_WARNINGS
+	}
+
+PRAGMA_DISABLE_DEPRECATION_WARNINGS
+	UE_DEPRECATED(5.7, "Use IClothingSimulationInterface instead.")
+	virtual void DestroySimulation(IClothingSimulation* InSimulation) {}
+PRAGMA_ENABLE_DEPRECATION_WARNINGS
 
 	// Given an asset, decide whether this factory can create a simulation to use the data inside
 	// (return false if data is invalid or missing in the case of custom data)
-	virtual bool SupportsAsset(const UClothingAssetBase* InAsset) const
-	PURE_VIRTUAL(UClothingSimulationFactory::SupportsAsset, return false;);
+	virtual bool SupportsAsset(const UClothingAssetBase* InAsset) const  // TODO 5.9: Make this method into a PURE_VIRTUAL
+	{
+		return const_cast<UClothingSimulationFactory*>(this)->SupportsAsset(InAsset);
+	}
 
 	UE_DEPRECATED(5.7, "Use const version instead.")
 	virtual bool SupportsAsset(UClothingAssetBase* InAsset)
 	{
-		return static_cast<const UClothingSimulationFactory*>(this)->SupportsAsset(InAsset);
+		return false;
 	}
 
 	// Whether or not we provide an interactor object to manipulate the simulation at runtime.
 	// If true is returned then CreateInteractor *must* create a valid object to handle this
-	virtual bool SupportsRuntimeInteraction() const
-	PURE_VIRTUAL(UClothingSimulationFactory::SupportsRuntimeInteraction, return false;);
+	virtual bool SupportsRuntimeInteraction() const  // TODO 5.9: Make this method into a PURE_VIRTUAL
+	{
+PRAGMA_DISABLE_DEPRECATION_WARNINGS
+		return const_cast<UClothingSimulationFactory*>(this)->SupportsRuntimeInteraction();
+PRAGMA_ENABLE_DEPRECATION_WARNINGS
+	}
 
 	UE_DEPRECATED(5.7, "Use const version instead.")
 	virtual bool SupportsRuntimeInteraction()
 	{
-		return static_cast<const UClothingSimulationFactory*>(this)->SupportsRuntimeInteraction();
+		return false;
 	}
 
 	// Creates the runtime interactor object for a clothing simulation. This object will
diff --git a/Engine/Source/Runtime/ClothingSystemRuntimeInterface/Public/ClothingSimulationInstance.h b/Engine/Source/Runtime/ClothingSystemRuntimeInterface/Public/ClothingSimulationInstance.h
--- a/Engine/Source/Runtime/ClothingSystemRuntimeInterface/Public/ClothingSimulationInstance.h
+++ b/Engine/Source/Runtime/ClothingSystemRuntimeInterface/Public/ClothingSimulationInstance.h
@@ -77,13 +77,13 @@ public:
 	}
 
 	/** Return the clothing simulation pointer. */
-	const IClothingSimulation* GetClothingSimulation() const
+	const IClothingSimulationInterface* GetClothingSimulation() const
 	{
 		return ClothingSimulation;
 	}
 
 	/** Return the clothing simulation pointer. */
-	IClothingSimulation* GetClothingSimulation()
+	IClothingSimulationInterface* GetClothingSimulation()
 	{
 		return ClothingSimulation;
 	}
@@ -139,12 +139,6 @@ public:
 		ClothingSimulation->ClearExternalCollisions();
 	}
 
-	/** Update external forces. */
-	void UpdateWorldForces(const USkeletalMeshComponent* OwnerComponent)
-	{
-		ClothingSimulation->UpdateWorldForces(OwnerComponent);
-	}
-
 	/** Fill the simulation context and do any other related game-thread initialization. */
 	UE_API void FillContextAndPrepareTick(const USkeletalMeshComponent* OwnerComponent, float DeltaTime, bool bIsInitialization);
 
@@ -153,11 +147,11 @@ public:
 	{
 		if (ClothingSimulation->ShouldSimulateLOD(CurrentOwnerLODIndex))
 		{
-			ClothingSimulation->Simulate(static_cast<const IClothingSimulationContext*>(ClothingSimulationContext));
+			ClothingSimulation->Simulate_AnyThread(static_cast<const IClothingSimulationContext*>(ClothingSimulationContext));
 		}
 		else
 		{
-			ClothingSimulation->ForceClothNextUpdateTeleportAndReset();
+			ClothingSimulation->ForceClothNextUpdateTeleportAndReset_AnyThread();
 		}
 	}
 
@@ -183,7 +177,7 @@ public:
 
 private:
 	/** ClothingSimulations are responsible for maintaining and simulating clothing actors. */
-	IClothingSimulation* ClothingSimulation = nullptr;
+	IClothingSimulationInterface* ClothingSimulation = nullptr;
 
 	/** ClothingSimulationContexts are a datastore for simulation data sent to the clothing thread. */
 	IClothingSimulationContext* ClothingSimulationContext = nullptr;
diff --git a/Engine/Source/Runtime/ClothingSystemRuntimeInterface/Public/ClothingSimulationInteractor.h b/Engine/Source/Runtime/ClothingSystemRuntimeInterface/Public/ClothingSimulationInteractor.h
--- a/Engine/Source/Runtime/ClothingSystemRuntimeInterface/Public/ClothingSimulationInteractor.h
+++ b/Engine/Source/Runtime/ClothingSystemRuntimeInterface/Public/ClothingSimulationInteractor.h
@@ -10,11 +10,13 @@
 #include "UObject/Object.h"
 #include "UObject/ObjectMacros.h"
 #include "UObject/UObjectGlobals.h"
+#include "ClothingSimulationInterface.h"
 
 #include "ClothingSimulationInteractor.generated.h"
 
+#define UE_API CLOTHINGSYSTEMRUNTIMEINTERFACE_API
+
 class FName;
-class IClothingSimulation;
 class IClothingSimulationContext;
 class UClothingAssetBase;
 struct FFrame;
@@ -32,7 +34,20 @@ class UClothingInteractor : public UObject
 protected:
 	friend class UClothingSimulationInteractor;
 
+	virtual void Sync(IClothingSimulationInterface* Simulation)
+	{
+PRAGMA_DISABLE_DEPRECATION_WARNINGS
+		if (Simulation && Simulation->DynamicCastToIClothingSimulation())
+		{
+			Sync(Simulation->DynamicCastToIClothingSimulation());
+		}
+PRAGMA_ENABLE_DEPRECATION_WARNINGS
+	}
+
+PRAGMA_DISABLE_DEPRECATION_WARNINGS
+	UE_DEPRECATED(5.7, "Use IClothingSimulationInterface instead.")
 	virtual void Sync(IClothingSimulation* Simulation) {}
+PRAGMA_ENABLE_DEPRECATION_WARNINGS
 
 protected:
 	int32 ClothingId = -1;  // Clothing Asset slot as passed to the CreateActor function
@@ -52,21 +67,25 @@ class UClothingSimulationInteractor : public UObject
 
 public:
 	/** Create a cloth specific interactor. Override either this method or the simpler CreateClothingInteractor method. */
-	CLOTHINGSYSTEMRUNTIMEINTERFACE_API virtual void CreateClothingInteractor(const UClothingAssetBase* ClothingAsset, int32 ClothingId);
+	UE_API virtual void CreateClothingInteractor(const UClothingAssetBase* ClothingAsset, int32 ClothingId);
 
 	/** Destroy all interactors. */
-	CLOTHINGSYSTEMRUNTIMEINTERFACE_API void DestroyClothingInteractors();
+	UE_API void DestroyClothingInteractors();
 
 	// Basic interface that clothing simulations are required to support
 
 	/**
-	 * Sync the interactor to the provided context for the clothing simulation to use
-	 * on its next update. This will be called at a time that is guaranteed to
-	 * never overlap with the running simulation. The simulation will have to be
-	 * written in a way to pick up these changes on the next update.
+	 * Sync the interactor with the clothing simulation to use any changes made on the next simulate.
+	 * This will be called at a time that is guaranteed to never overlap with the running simulation.
+	 * The simulation will have to be written in a way to pick up these changes on the next update.
 	 * Any inherited class must call this function to also Sync the ClothingInteractors.
 	 */
-	CLOTHINGSYSTEMRUNTIMEINTERFACE_API virtual void Sync(IClothingSimulation* Simulation, IClothingSimulationContext* Context);
+	UE_API virtual void Sync(IClothingSimulationInterface* Simulation, const IClothingSimulationContext* Context);
+
+PRAGMA_DISABLE_DEPRECATION_WARNINGS
+	UE_DEPRECATED(5.7, "Use IClothingSimulationInterface instead.")
+	UE_API virtual void Sync(IClothingSimulation* Simulation, IClothingSimulationContext* Context);
+PRAGMA_ENABLE_DEPRECATION_WARNINGS
 
 	/** Called to update collision status without restarting the simulation. */
 	UFUNCTION(BlueprintCallable, Category=ClothingSimulation)
@@ -142,7 +161,7 @@ public:
 
 	/**Return a cloth interactor for this simulation. */
 	UFUNCTION(BlueprintCallable, Category = ClothingSimulation)
-	CLOTHINGSYSTEMRUNTIMEINTERFACE_API UClothingInteractor* GetClothingInteractor(const FName ClothingAssetName) const;
+	UE_API UClothingInteractor* GetClothingInteractor(const FName ClothingAssetName) const;
 
 	UE_DEPRECATED(5.7, "Use FName version instead.")
 	UClothingInteractor* GetClothingInteractor(const FString& ClothingAssetName) const
@@ -170,3 +189,5 @@ private:
 	int32 LastNumSubsteps = 0;
 	float LastSimulationTime = 0.f;
 };
+
+#undef UE_API
diff --git a/Engine/Source/Runtime/ClothingSystemRuntimeInterface/Public/ClothingSimulationInterface.h b/Engine/Source/Runtime/ClothingSystemRuntimeInterface/Public/ClothingSimulationInterface.h
--- a/Engine/Source/Runtime/ClothingSystemRuntimeInterface/Public/ClothingSimulationInterface.h
+++ b/Engine/Source/Runtime/ClothingSystemRuntimeInterface/Public/ClothingSimulationInterface.h
@@ -6,9 +6,12 @@
 #include "CoreTypes.h"
 #include "Math/MathFwd.h"
 
+#define UE_API CLOTHINGSYSTEMRUNTIMEINTERFACE_API
+
 class UClothingAssetBase;
 class USkeletalMeshComponent;
 class USkinnedMeshComponent;
+class UE_DEPRECATED(5.7, "Use IClothingSimulationInterface instead.") IClothingSimulation;
 struct FClothCollisionData;
 struct FClothSimulData;
 
@@ -16,16 +19,175 @@
 class IClothingSimulationContext
 {
 public:
-	CLOTHINGSYSTEMRUNTIMEINTERFACE_API IClothingSimulationContext();
-	CLOTHINGSYSTEMRUNTIMEINTERFACE_API virtual ~IClothingSimulationContext();
+	UE_API IClothingSimulationContext();
+	UE_API virtual ~IClothingSimulationContext();
 };
 
-/** Base class for clothing simulators. */
-class IClothingSimulation
+/**
+ * Base class for clothing simulators.
+ * 
+ * The majority of the API for this class is protected.
+ * For the most part the simulation is not designed to be used outside of the FClothingSimulationInstance implementation which is
+ * used in the Skeletal Mesh Component, as its parallel simulation is tied to the Skeletal Mesh Component tick and dependents.
+ * Any method that is available in the public section below should consider that it may be called while the simulation is running.
+ */
+class IClothingSimulationInterface
 {
 public:
-	CLOTHINGSYSTEMRUNTIMEINTERFACE_API IClothingSimulation();
-	CLOTHINGSYSTEMRUNTIMEINTERFACE_API virtual ~IClothingSimulation();
+	UE_API IClothingSimulationInterface();
+	UE_API virtual ~IClothingSimulationInterface();
 
+protected:
+
+	friend struct FClothingSimulationInstance;
+
+	/**
+	 * Create an actor for this simulation from the data in InAsset
+	 * Simulation data for this actor should be written back to SimDataIndex in GetSimulationData
+	 * @param InOwnerComponent - The component requesting this actor
+	 * @param InAsset - The asset to create an actor from
+	 * @param SimDataIndex - the sim data index to use when doing the writeback in AppendSimulationData
+	 */
+	virtual void CreateActor(USkeletalMeshComponent* InOwnerComponent, const UClothingAssetBase* InAsset, int32 SimDataIndex) = 0;
+
+	/** Called once all CreateActor have been called. */
+	virtual void EndCreateActor() = 0;
+
+	/** Create a new context, will not be filled, call FillContextAndPrepareTick before simulating with this context */
+	virtual IClothingSimulationContext* CreateContext() = 0;
+
+	/**
+	 * Fills an existing context for a single simulation step and do any other work that needs to be called by the engine on the game thread prior to simulation.
+	 * @param InComponent - The component to fill the context for
+	 * @param InOutContext - The context to fill
+	 * @param bIsInitialization - Whether this fill is occurring as part of the actor creation stage
+	 * @param bForceTeleportResetOnly - Whether ForceClothNextUpdateTeleportAndReset instead of Simulate and AppendSimulationData will be called.
+	 */
+	virtual void FillContextAndPrepareTick(const USkeletalMeshComponent* InComponent, float InDeltaTime, IClothingSimulationContext* InOutContext, bool bIsInitialization, bool bForceTeleportResetOnly) = 0;
+
+	/** Initialize the simulation, will be called before any Simulate calls */
+	virtual void Initialize() = 0;
+
+	/**
+	 * Shutdown the simulation, this should clear ALL resources as we no longer expect to
+	 * call any other function on this simulation without first calling Initialize again.
+	 */
+	virtual void Shutdown() = 0;
+
+	/**
+	 * Called by the engine to determine if this simulation can run this tick at this LOD.
+	 * ForceClothNextUpdateTeleportAndReset rather than Simulate and AppendSimulationData will be called this tick.
+	 */
+	virtual bool ShouldSimulateLOD(int32 OwnerLODIndex) const = 0;
+
+	/**
+	 * Run a single tick of the simulation.
+	 * The pointer InContext is guaranteed (when called by the engine) to be the context allocated in
+	 * CreateContext and can be assumed to be safely castable to any derived type allocated there.
+	 * New callers should take care to make sure only the correct context is ever passed through.
+	 * @note Can be called asynchronously outside of the game thread.
+	 * @param InContext - The context to use during simulation, will have been filled in FillContextAndPrepareTick
+	 */
+	virtual void Simulate_AnyThread(const IClothingSimulationContext* InContext) = 0;
+
+	/**
+	 * Called instead of Simulate when not ticking.
+	 * @note Can be called asynchronously outside of the game thread.
+	 */
+	virtual void ForceClothNextUpdateTeleportAndReset_AnyThread() = 0;
+
+	/** Hard reset the simulation without necessarily recreating cloth actors. */
+	virtual void HardResetSimulation(const IClothingSimulationContext* InContext) = 0;
+
+	/** Simulation should remove all of it's actors when next possible and free them */
+	virtual void DestroyActors() = 0;
+
+	/**
+	 * Destroy a context object, engine will always pass a context created using CreateContext
+	 * @param InContext - The context to destroy
+	 */
+	virtual void DestroyContext(IClothingSimulationContext* InContext) = 0;
+
+	/**
+	 * Fill FClothSimulData map for the clothing simulation. Should fill a map pair per-actor.
+	 * Do not remove InOutData elements which you are not responsible for: they may have been written by a different instance.
+	 * @param InOutData - The simulation data to write to
+	 * @param OwnerComponent - the component that owns the simulation
+	 * @param OverrideComponent - An override component if bound to a leader pose component
+	 */
+	virtual void AppendSimulationData(TMap<int32, FClothSimulData>& InOutData, const USkeletalMeshComponent* OwnerComponent, const USkinnedMeshComponent* OverrideComponent) const = 0;
+
+	/** Get the bounds of the simulation mesh in local simulation space */
+	virtual FBoxSphereBounds GetBounds(const USkeletalMeshComponent* InOwnerComponent) const = 0;
+
+	/**
+	 * Called by the engine when an external object wants to inject collision data into this simulation above
+	 * and beyond what is specified in the asset for the internal actors
+	 * Examples: Scene collision, collision for parents we are attached to
+	 * @param InData - Collisions to add to this simulation
+	 */
+	virtual void AddExternalCollisions(const FClothCollisionData& InData) = 0;
+
+	/**
+	 * Called by the engine when external collisions are no longer necessary or when they need to be updated
+	 * with some of the previous collisions removed. It is recommended in derived simulations to avoid freeing
+	 * any allocations regarding external collisions as it is likely more will be added soon after this call
+	 */
+	virtual void ClearExternalCollisions() = 0;
+
+	/**
+	 * Called by the engine to request data on all active collisions in a simulation. if bIncludeExternal is
+	 * true, derived implementations should add the asset collisions and any collisions added at runtime, if
+	 * false only the collisions from the asset should be considered
+	 * @param OutCollisions - Array to write collisions to
+	 * @param bIncludeExternal - Whether or not external collisions should be retrieved, or just asset collisions
+	 */
+	virtual void GetCollisions(FClothCollisionData& OutCollisions, bool bIncludeExternal = true) const = 0;
+
+public:
+PRAGMA_DISABLE_DEPRECATION_WARNINGS
+	UE_DEPRECATED(5.7, "This method is used to ease transition between the new and old interface and will be removed.")
+	inline const IClothingSimulation* DynamicCastToIClothingSimulation() const;
+
+	UE_DEPRECATED(5.7, "This method is used to ease transition between the new and old interface and will be removed.")
+	inline IClothingSimulation* DynamicCastToIClothingSimulation();
+PRAGMA_ENABLE_DEPRECATION_WARNINGS
+
+	/** Return the number of simulated cloths. Implementation must be thread safe. */
+	virtual int32 GetNumCloths() const { return 0; }
+	/** Return the number of kinematic (fixed) particles. Implementation must be thread safe. */
+	virtual int32 GetNumKinematicParticles() const { return 0; }
+	/** Return the number of dynamic (simulated) particles. Implementation must be thread safe. */
+	virtual int32 GetNumDynamicParticles() const { return 0; }
+	/**
+	 * Return the number of iterations used by the solver.
+	 * This is the maximum used as an indicative value only, as this could vary between cloths.
+	 * Implementation must be thread safe.
+	 */
+	virtual int32 GetNumIterations() const { return 0; }
+	/**
+	 * Return the number of substeps used by the solver.
+	 * This is the maximum used as an indicative value only, as this could vary between cloths.
+	 * Implementation must be thread safe.
+	 */
+	virtual int32 GetNumSubsteps() const { return 0; }
+	/** Return the simulation time in ms. Implementation must be thread safe. */
+	virtual float GetSimulationTime() const { return 0.f; }
+	/** Return whether the simulation is teleported. Implementation must be thread safe. */
+	virtual bool IsTeleported() const { return false; }
+
+private:
+PRAGMA_DISABLE_DEPRECATION_WARNINGS
+	friend class IClothingSimulation;  // For bIsLegacyInterface
+PRAGMA_ENABLE_DEPRECATION_WARNINGS
+	bool bIsLegacyInterface = false;  // UE_DEPRECATED(5.7, "Use IClothingSimulationInterface instead.")
+};
+
+class UE_DEPRECATED(5.7, "Use IClothingSimulationInterface instead.") IClothingSimulation : public IClothingSimulationInterface
+{
+public:
+	UE_API IClothingSimulation();
+	UE_API virtual ~IClothingSimulation();
+
 	// The majority of the API for this class is protected. The required objects (skel meshes and the parallel task)
 	// are friends so they can use the functionality. For the most part the simulation is not designed to be used
@@ -50,7 +212,7 @@
 	 * @param InAsset - The asset to create an actor from
 	 * @param SimDataIndex - the sim data index to use when doing the writeback in GetSimulationData
 	 */
-	virtual void CreateActor(USkeletalMeshComponent* InOwnerComponent, const UClothingAssetBase* InAsset, int32 SimDataIndex)
+	virtual void CreateActor(USkeletalMeshComponent* InOwnerComponent, const UClothingAssetBase* InAsset, int32 SimDataIndex) override
 	{
 PRAGMA_DISABLE_DEPRECATION_WARNINGS
 		return CreateActor(InOwnerComponent, const_cast<UClothingAssetBase*>(InAsset), SimDataIndex);
@@ -60,12 +222,6 @@
 	UE_DEPRECATED(5.7, "Use const arguments version instead,")
 	virtual void CreateActor(USkeletalMeshComponent* InOwnerComponent, UClothingAssetBase* InAsset, int32 SimDataIndex) {}
 
-	/** Called once all CreateActor have been called. */
-	virtual void EndCreateActor() {}
-
-	/** Create a new context, will not be filled, call FillContext before simulating with this context */
-	virtual IClothingSimulationContext* CreateContext() = 0;
-
 	/**
 	 * Fills an existing context for a single simulation step and do any other work that needs to be called by the engine on the game thread prior to simulation.
 	 * @param InComponent - The component to fill the context for
@@ -73,8 +229,7 @@
 	 * @param bIsInitialization - Whether this fill is occurring as part of the actor creation stage
 	 * @param bForceTeleportResetOnly - Whether ForceClothNextUpdateTeleportAndReset instead of Simulate and AppendSimulationData will be called.
 	 */
-	// TODO: make this pure virtual in the new interface.
-	virtual void FillContextAndPrepareTick(const USkeletalMeshComponent* InComponent, float InDeltaTime, IClothingSimulationContext* InOutContext, bool bIsInitialization, bool bForceTeleportResetOnly)
+	virtual void FillContextAndPrepareTick(const USkeletalMeshComponent* InComponent, float InDeltaTime, IClothingSimulationContext* InOutContext, bool bIsInitialization, bool bForceTeleportResetOnly) override
 	{
 PRAGMA_DISABLE_DEPRECATION_WARNINGS
 		return FillContext(const_cast<USkeletalMeshComponent*>(InComponent), InDeltaTime, InOutContext, bIsInitialization);
@@ -91,15 +246,6 @@
 	UE_DEPRECATED(5.7, "Use FillContextAndPrepareTick instead.")
 	virtual void FillContext(USkeletalMeshComponent* InComponent, float InDeltaTime, IClothingSimulationContext* InOutContext) {}
 
-	/** Initialize the simulation, will be called before any Simulate calls */
-	virtual void Initialize() = 0;
-
 	/**
-	 * Shutdown the simulation, this should clear ALL resources as we no longer expect to
-	 * call any other function on this simulation without first calling Initialize again.
-	 */
-	virtual void Shutdown() = 0;
-
-	/**
 	 * Called by the engine to detect whether or not the simulation should run, essentially
 	 * are there any actors that need to simulate in this simulation
@@ -113,8 +259,7 @@
 	/** Called by the engine to determine if this simulation can run this tick at this LOD.
 	 *  ForceClothNextUpdateTeleportAndReset rather than Simulate and AppendSimulationData will be called this tick.
 	 */
-	 // TODO: make this pure virtual in the new interface.
-	virtual bool ShouldSimulateLOD(int32 OwnerLODIndex) const
+	virtual bool ShouldSimulateLOD(int32 OwnerLODIndex) const override
 	{
 PRAGMA_DISABLE_DEPRECATION_WARNINGS
 		return ShouldSimulate();
@@ -128,33 +273,21 @@
 	 * New callers should make care to make sure only the correct context is ever passed through
 	 * @param InContext - The context to use during simulation, will have been filled in FillContextAndPrepareTick
 	 */
-	 // TODO: make this pure virtual in the new interface. Indicate that this can be called asynchronously
-	virtual void Simulate(const IClothingSimulationContext* InContext)
+	virtual void Simulate_AnyThread(const IClothingSimulationContext* InContext) override
 	{
 PRAGMA_DISABLE_DEPRECATION_WARNINGS
 		return Simulate(const_cast<IClothingSimulationContext*>(InContext));
 PRAGMA_ENABLE_DEPRECATION_WARNINGS
 	}
 	
-	UE_DEPRECATED(5.7, "Use const argument version instead,")
+	UE_DEPRECATED(5.7, "Use Simulate_AnyThread instead,")
 	virtual void Simulate(IClothingSimulationContext* InContext) {}
 
 	/** Called instead of Simulate when not ticking.*/
-	 // TODO: make this pure virtual in the new interface. Indicate that this can be called asynchronously
-	virtual void ForceClothNextUpdateTeleportAndReset() {};
+	virtual void ForceClothNextUpdateTeleportAndReset_AnyThread() override {}
 
 	/** Hard reset the simulation without necessarily recreating cloth actors. */
-	// TODO: make this pure virtual in the new interface.
-	virtual void HardResetSimulation(const IClothingSimulationContext* InContext) {};
-
-	/** Simulation should remove all of it's actors when next possible and free them */
-	virtual void DestroyActors() = 0;
+	virtual void HardResetSimulation(const IClothingSimulationContext* InContext) override {}
 
 	/** 
-	 * Destroy a context object, engine will always pass a context created using CreateContext 
-	 * @param InContext - The context to destroy
-	 */
-	virtual void DestroyContext(IClothingSimulationContext* InContext) = 0;
-
-	/** 
 	 * Fill FClothSimulData map for the clothing simulation. Should fill a map pair per-actor. 
@@ -163,8 +296,7 @@
 	 * @param OwnerComponent - the component that owns the simulation
 	 * @param OverrideComponent - An override component if bound to a leader pose component
 	 */
-	// TODO: make this pure virtual in the new interface.
-	virtual void AppendSimulationData(TMap<int32, FClothSimulData>& InOutData, const USkeletalMeshComponent* OwnerComponent, const USkinnedMeshComponent* OverrideComponent) const 
+	virtual void AppendSimulationData(TMap<int32, FClothSimulData>& InOutData, const USkeletalMeshComponent* OwnerComponent, const USkinnedMeshComponent* OverrideComponent) const override
 	{
 PRAGMA_DISABLE_DEPRECATION_WARNINGS
 		GetSimulationData(InOutData, const_cast<USkeletalMeshComponent*>(OwnerComponent), const_cast<USkinnedMeshComponent*>(OverrideComponent));
@@ -174,76 +306,31 @@
 	UE_DEPRECATED(5.7, "Use AppendSimulationData instead.")
 	virtual void GetSimulationData(TMap<int32, FClothSimulData>& OutData, USkeletalMeshComponent* InOwnerComponent, USkinnedMeshComponent* InOverrideComponent) const {}
 
-	/** Get the bounds of the simulation mesh in local simulation space */
-	virtual FBoxSphereBounds GetBounds(const USkeletalMeshComponent* InOwnerComponent) const = 0;
-
-	/**
-	 * Called by the engine when an external object wants to inject collision data into this simulation above
-	 * and beyond what is specified in the asset for the internal actors
-	 * Examples: Scene collision, collision for parents we are attached to
-	 * @param InData - Collisions to add to this simulation
-	 */
-	virtual void AddExternalCollisions(const FClothCollisionData& InData) = 0;
-
-	/**
-	 * Called by the engine when external collisions are no longer necessary or when they need to be updated
-	 * with some of the previous collisions removed. It is recommended in derived simulations to avoid freeing
-	 * any allocations regarding external collisions as it is likely more will be added soon after this call
-	 */
-	virtual void ClearExternalCollisions() = 0;
-
-	/**
-	 * Called by the engine to request data on all active collisions in a simulation. if bIncludeExternal is
-	 * true, derived implementations should add the asset collisions and any collisions added at runtime, if
-	 * false only the collisions from the asset should be considered
-	 * @param OutCollisions - Array to write collisions to
-	 * @param bIncludeExternal - Whether or not external collisions should be retrieved, or just asset collisions
-	 */
-	virtual void GetCollisions(FClothCollisionData& OutCollisions, bool bIncludeExternal = true) const = 0;
-
 public:
-	/**
-	 * Set the number of iterations used by the solver.
-	 * Implementation is not considered thread safe, and therefore this function must be called at a synchronization point.
-	 */
-	virtual void SetNumIterations(int32 /*NumIterations*/) {}
-	/**
-	 * Set the maximum number of iterations used by the solver.
-	 * Implementation is not considered thread safe, and therefore this function must be called at a synchronization point.
-	 */
-	virtual void SetMaxNumIterations(int32 /*MaxNumIterations*/) {}
-	/**
-	 * Set the number of substeps used by the solver.
-	 * Implementation is not considered thread safe, and therefore this function must be called at a synchronization point.
-	 */
-	virtual void SetNumSubsteps(int32 /*NumSubsteps*/) {}
-
 	UE_DEPRECATED(5.7, "GatherStats is unused and will be removed from future interfaces. Do not override.")
 	virtual void GatherStats() const {}
 
-	/** Return the number of simulated cloths. Implementation must be thread safe. */
-	virtual int32 GetNumCloths() const { return 0;  }
-	/** Return the number of kinematic (fixed) particles. Implementation must be thread safe. */
-	virtual int32 GetNumKinematicParticles() const { return 0; }
-	/** Return the number of dynamic (simulated) particles. Implementation must be thread safe. */
-	virtual int32 GetNumDynamicParticles() const { return 0; }
-	/**
-	 * Return the number of iterations used by the solver.
-	 * This is the maximum used as an indicative value only, as this could vary between cloths.
-	 * Implementation must be thread safe.
-	 */
-	virtual int32 GetNumIterations() const { return 0; }
-	/**
-	 * Return the number of substeps used by the solver.
-	 * This is the maximum used as an indicative value only, as this could vary between cloths.
-	 * Implementation must be thread safe.
-	 */
-	virtual int32 GetNumSubsteps() const { return 0; }
-	/** Return the simulation time in ms. Implementation must be thread safe. */
-	virtual float GetSimulationTime() const { return 0.f; }
-	/** Return whether the simulation is teleported. Implementation must be thread safe. */
-	virtual bool IsTeleported() const { return false; }
+	UE_DEPRECATED(5.7, "Use a clothing simulation interactor to set this simulation property instead.")
+	virtual void SetNumIterations(int32 NumIterations) {}
+	UE_DEPRECATED(5.7, "Use a clothing simulation interactor to set this simulation property instead.")
+	virtual void SetMaxNumIterations(int32 MaxNumIterations) {}
+	UE_DEPRECATED(5.7, "Use a clothing simulation interactor to set this simulation property instead.")
+	virtual void SetNumSubsteps(int32 NumSubsteps) {}
 
-	/** Update the forces coming from the world. */
-	virtual void UpdateWorldForces(const USkeletalMeshComponent* OwnerComponent) {};
+	UE_DEPRECATED(5.7, "Use FillContextAndPrepareTick instead.")
+	virtual void UpdateWorldForces(const USkeletalMeshComponent* OwnerComponent) {}
 };
+
+PRAGMA_DISABLE_DEPRECATION_WARNINGS
+inline const class IClothingSimulation* IClothingSimulationInterface::DynamicCastToIClothingSimulation() const
+{
+	return bIsLegacyInterface ? static_cast<const class IClothingSimulation*>(this) : nullptr;
+}
+
+inline IClothingSimulation* IClothingSimulationInterface::DynamicCastToIClothingSimulation()
+{
+	return bIsLegacyInterface ? static_cast<class IClothingSimulation*>(this) : nullptr;
+}
+PRAGMA_ENABLE_DEPRECATION_WARNINGS
+
+#undef UE_API
diff --git a/Engine/Source/Runtime/ClothingSystemRuntimeNv/Private/ClothingSimulationFactoryNv.cpp b/Engine/Source/Runtime/ClothingSystemRuntimeNv/Private/ClothingSimulationFactoryNv.cpp
--- a/Engine/Source/Runtime/ClothingSystemRuntimeNv/Private/ClothingSimulationFactoryNv.cpp
+++ b/Engine/Source/Runtime/ClothingSystemRuntimeNv/Private/ClothingSimulationFactoryNv.cpp
@@ -7,12 +7,12 @@
 // Legacy support allowing the loading of older classes that may still reference old NvCloth simulation
 // classes. These are no longer expected to work but should still load so the references can be safely changed.
 PRAGMA_DISABLE_DEPRECATION_WARNINGS
-IClothingSimulation* UClothingSimulationFactoryNv::CreateSimulation()
+IClothingSimulationInterface* UClothingSimulationFactoryNv::CreateSimulation() const
 {
 	return nullptr;
 }
 
-void UClothingSimulationFactoryNv::DestroySimulation(IClothingSimulation* InSimulation)
+void UClothingSimulationFactoryNv::DestroySimulation(IClothingSimulationInterface* InSimulation) const
 {
 }
 
@@ -23,7 +23,7 @@ bool UClothingSimulationFactoryNv::SupportsAsset(const UClothingAssetBase* InAss
 
 bool UClothingSimulationFactoryNv::SupportsRuntimeInteraction() const
 {
-	return true;
+	return false;
 }
 
 UClothingSimulationInteractor* UClothingSimulationFactoryNv::CreateInteractor()
diff --git a/Engine/Source/Runtime/ClothingSystemRuntimeNv/Private/ClothingSimulationInteractorNv.cpp b/Engine/Source/Runtime/ClothingSystemRuntimeNv/Private/ClothingSimulationInteractorNv.cpp
--- a/Engine/Source/Runtime/ClothingSystemRuntimeNv/Private/ClothingSimulationInteractorNv.cpp
+++ b/Engine/Source/Runtime/ClothingSystemRuntimeNv/Private/ClothingSimulationInteractorNv.cpp
@@ -16,7 +16,7 @@ void UClothingSimulationInteractorNv::ClothConfigUpdated()
 {
 }
 
-void UClothingSimulationInteractorNv::Sync(IClothingSimulation* InSimulation, IClothingSimulationContext* InContext)
+void UClothingSimulationInteractorNv::Sync(IClothingSimulationInterface* InSimulation, const IClothingSimulationContext* InContext)
 {
 }
 
diff --git a/Engine/Source/Runtime/ClothingSystemRuntimeNv/Public/ClothingSimulationFactoryNv.h b/Engine/Source/Runtime/ClothingSystemRuntimeNv/Public/ClothingSimulationFactoryNv.h
--- a/Engine/Source/Runtime/ClothingSystemRuntimeNv/Public/ClothingSimulationFactoryNv.h
+++ b/Engine/Source/Runtime/ClothingSystemRuntimeNv/Public/ClothingSimulationFactoryNv.h
@@ -10,7 +10,6 @@
 
 #include "ClothingSimulationFactoryNv.generated.h"
 
-class IClothingSimulation;
 class UClothConfigBase;
 class UClothingAssetBase;
 class UClothingSimulationInteractor;
@@ -26,8 +25,8 @@ class UClothingSimulationFactoryNv final : public UClothingSimulationFactory
 	GENERATED_BODY()
 public:
 
-	CLOTHINGSYSTEMRUNTIMENV_API virtual IClothingSimulation* CreateSimulation() override;
-	CLOTHINGSYSTEMRUNTIMENV_API virtual void DestroySimulation(IClothingSimulation* InSimulation) override;
+	CLOTHINGSYSTEMRUNTIMENV_API virtual IClothingSimulationInterface* CreateSimulation() const override;
+	CLOTHINGSYSTEMRUNTIMENV_API virtual void DestroySimulation(IClothingSimulationInterface* InSimulation) const override;
 	CLOTHINGSYSTEMRUNTIMENV_API virtual bool SupportsAsset(const UClothingAssetBase* InAsset) const override;
 
 	CLOTHINGSYSTEMRUNTIMENV_API virtual bool SupportsRuntimeInteraction() const override;
diff --git a/Engine/Source/Runtime/ClothingSystemRuntimeNv/Public/ClothingSimulationInteractorNv.h b/Engine/Source/Runtime/ClothingSystemRuntimeNv/Public/ClothingSimulationInteractorNv.h
--- a/Engine/Source/Runtime/ClothingSystemRuntimeNv/Public/ClothingSimulationInteractorNv.h
+++ b/Engine/Source/Runtime/ClothingSystemRuntimeNv/Public/ClothingSimulationInteractorNv.h
@@ -14,7 +14,6 @@
 
 class FClothingSimulationContextNv;
 class FClothingSimulationNv;
-class IClothingSimulation;
 class IClothingSimulationContext;
 class UObject;
 struct FFrame;
@@ -31,7 +30,7 @@ class UClothingSimulationInteractorNv : public UClothingSimulationInteractor
 public:
 
 	// UClothingSimulationInteractor Interface
-	CLOTHINGSYSTEMRUNTIMENV_API virtual void Sync(IClothingSimulation* InSimulation, IClothingSimulationContext* InContext) override;
+	CLOTHINGSYSTEMRUNTIMENV_API virtual void Sync(IClothingSimulationInterface* InSimulation, const IClothingSimulationContext* InContext) override;
 
 	CLOTHINGSYSTEMRUNTIMENV_API virtual void PhysicsAssetUpdated() override;
 	CLOTHINGSYSTEMRUNTIMENV_API virtual void ClothConfigUpdated() override;
diff --git a/Engine/Source/Runtime/Engine/Private/Components/SkeletalMeshComponent.cpp b/Engine/Source/Runtime/Engine/Private/Components/SkeletalMeshComponent.cpp
--- a/Engine/Source/Runtime/Engine/Private/Components/SkeletalMeshComponent.cpp
+++ b/Engine/Source/Runtime/Engine/Private/Components/SkeletalMeshComponent.cpp
@@ -2428,16 +2428,28 @@ void USkeletalMeshComponent::EvaluatePostProcessMeshInstance(TArray<FTransform>&
 }
 
 // UE_DEPRECATED(5.7, "Use GetClothingSimulationInstances instead.")
+PRAGMA_DISABLE_DEPRECATION_WARNINGS
 const IClothingSimulation* USkeletalMeshComponent::GetClothingSimulation() const
 {
-	return GetClothingSimulationInstances().Num() ? ClothingSimulationInstances[0].GetClothingSimulation() : nullptr;
+	if (const IClothingSimulationInterface* const ClothingSimulation = GetClothingSimulationInstances().Num() ? ClothingSimulationInstances[0].GetClothingSimulation() : nullptr)
+	{
+		return ClothingSimulation->DynamicCastToIClothingSimulation();
+	}
+	return nullptr;
 }
+PRAGMA_ENABLE_DEPRECATION_WARNINGS
 
 // UE_DEPRECATED(5.7, "Use GetClothingSimulationInstances instead.")
+PRAGMA_DISABLE_DEPRECATION_WARNINGS
 IClothingSimulation* USkeletalMeshComponent::GetClothingSimulation()
 {
-	return GetClothingSimulationInstances().Num() ? ClothingSimulationInstances[0].GetClothingSimulation() : nullptr;
+	if (IClothingSimulationInterface* const ClothingSimulation = GetClothingSimulationInstances().Num() ? ClothingSimulationInstances[0].GetClothingSimulation() : nullptr)
+	{
+		return ClothingSimulation->DynamicCastToIClothingSimulation();
+	}
+	return nullptr;
 }
+PRAGMA_ENABLE_DEPRECATION_WARNINGS
 
 // UE_DEPRECATED(5.7, "Use GetClothingSimulationInstances instead.")
 const IClothingSimulationContext* USkeletalMeshComponent::GetClothingSimulationContext() const
diff --git a/Engine/Source/Runtime/Engine/Private/SkeletalMeshComponentPhysics.cpp b/Engine/Source/Runtime/Engine/Private/SkeletalMeshComponentPhysics.cpp
--- a/Engine/Source/Runtime/Engine/Private/SkeletalMeshComponentPhysics.cpp
+++ b/Engine/Source/Runtime/Engine/Private/SkeletalMeshComponentPhysics.cpp
@@ -3565,12 +3565,6 @@ void USkeletalMeshComponent::UpdateClothStateAndSimulate(float DeltaTime, FTickF
 	}
 #endif // WITH_CLOTH_COLLISION_DETECTION
 
-
-	for (FClothingSimulationInstance& ClothingSimulationInstance : ClothingSimulationInstances)
-	{
-		ClothingSimulationInstance.UpdateWorldForces(this);
-	}
-
 	UpdateClothSimulationContext(DeltaTime);
 
 	ParallelClothTask = TGraphTask<FParallelClothTask>::CreateTask(nullptr, ENamedThreads::GameThread).ConstructAndDispatchWhenReady(*this, DeltaTime);
