Add Chaos Cloth Velocity Scale debug draw options

diff --git a/Engine/Plugins/ChaosCloth/Source/ChaosCloth/Private/ChaosCloth/ChaosClothVisualization.cpp b/Engine/Plugins/ChaosCloth/Source/ChaosCloth/Private/ChaosCloth/ChaosClothVisualization.cpp
--- a/Engine/Plugins/ChaosCloth/Source/ChaosCloth/Private/ChaosCloth/ChaosClothVisualization.cpp
+++ b/Engine/Plugins/ChaosCloth/Source/ChaosCloth/Private/ChaosCloth/ChaosClothVisualization.cpp
@@ -50,6 +50,8 @@
 #include "UObject/ICookInfo.h"
 #endif  // #if WITH_EDITOR
 
+
+
 namespace Chaos
 {
 	namespace Private
@@ -2725,6 +2727,76 @@ namespace Chaos
 		}
 	}
 
+	void FClothVisualizationNoGC::DrawVelocityScale(FPrimitiveDrawInterface* PDI, const FReal LengthScale) const
+	{
+		if (!Solver)
+		{
+			return;
+		}
+		for (const FClothingSimulationCloth* const Cloth : Solver->GetCloths())
+		{
+			const int32 ParticleRangeId = Cloth->GetParticleRangeId(Solver);
+			if (ParticleRangeId == INDEX_NONE)
+			{
+				continue;
+			}
+			const uint32 GroupId = Cloth->GetGroupId();
+
+			const FRigidTransform3& ReferenceSpaceTransform = Cloth->GetReferenceSpaceTransform();
+			const EChaosSoftsSimulationSpace VelocityScaleSpace = (EChaosSoftsSimulationSpace)FClothingSimulationCloth::GetVelocityScaleSpace(Cloth->GetConfig()->GetProperties(Cloth->GetLODIndex(Solver)), (int32)EChaosSoftsSimulationSpace::ReferenceBoneSpace);
+
+			const FVector OriginalVelocity = FVector(Solver->GetOriginalReferenceSpaceVelocity(GroupId));
+			const FVector OriginalAngularVelocity = FVector(Solver->GetOriginalReferenceSpaceAngularVelocity(GroupId));
+			const FVector AppliedVelocity = FVector(Solver->GetAppliedReferenceSpaceVelocity(GroupId));
+			const FVector AppliedAngularVelocity = FVector(Solver->GetAppliedReferenceSpaceAngularVelocity(GroupId));
+
+			const FQuat& VelocityScaleSpaceRot = VelocityScaleSpace == EChaosSoftsSimulationSpace::WorldSpace ? FQuat::Identity : (VelocityScaleSpace == EChaosSoftsSimulationSpace::ComponentSpace ? FQuat(Solver->GetLocalSpaceRotation()) : ReferenceSpaceTransform.GetRotation());
+			
+			const FVector XDir = VelocityScaleSpaceRot.RotateVector(FVector::ForwardVector);
+			const FVector YDir = VelocityScaleSpaceRot.RotateVector(FVector::RightVector);
+			const FVector ZDir = VelocityScaleSpaceRot.RotateVector(FVector::UpVector);
+
+			const FReal XOrig = FVector::DotProduct(OriginalVelocity, XDir);
+			const FReal YOrig = FVector::DotProduct(OriginalVelocity, YDir);
+			const FReal ZOrig = FVector::DotProduct(OriginalVelocity, ZDir);
+
+			const FReal XApplied = FVector::DotProduct(AppliedVelocity, XDir);
+			const FReal YApplied = FVector::DotProduct(AppliedVelocity, YDir);
+			const FReal ZApplied = FVector::DotProduct(AppliedVelocity, ZDir);
+
+			DrawArrow(PDI, ReferenceSpaceTransform.GetLocation(), ReferenceSpaceTransform.GetLocation() + LengthScale * XApplied * XDir, FLinearColor::Red);
+			DrawArrow(PDI, ReferenceSpaceTransform.GetLocation(), ReferenceSpaceTransform.GetLocation() + LengthScale * YApplied * YDir, FLinearColor::Green);
+			DrawArrow(PDI, ReferenceSpaceTransform.GetLocation(), ReferenceSpaceTransform.GetLocation() + LengthScale * ZApplied * ZDir, FLinearColor::Blue);
+
+			// Draw OrigVelocity lighter
+			constexpr float OrigColorScale = 0.5f;
+			if (XOrig > XApplied)
+			{
+				DrawArrow(PDI, ReferenceSpaceTransform.GetLocation() + LengthScale * XApplied * XDir, ReferenceSpaceTransform.GetLocation() + LengthScale * XOrig * XDir, FLinearColor::Red * OrigColorScale);
+			}
+			if (YOrig > YApplied)
+			{
+				DrawArrow(PDI, ReferenceSpaceTransform.GetLocation() + LengthScale * YApplied * YDir, ReferenceSpaceTransform.GetLocation() + LengthScale * YOrig * YDir, FLinearColor::Green * OrigColorScale);
+			}
+			if (ZOrig > ZApplied)
+			{
+				DrawArrow(PDI, ReferenceSpaceTransform.GetLocation() + LengthScale * ZApplied * ZDir, ReferenceSpaceTransform.GetLocation() + LengthScale * ZOrig * ZDir, FLinearColor::Blue * OrigColorScale);
+			}
+
+
+			DrawArrow(PDI, ReferenceSpaceTransform.GetLocation(), ReferenceSpaceTransform.GetLocation() + OriginalAngularVelocity, FLinearColor::Yellow * OrigColorScale);
+			DrawArrow(PDI, ReferenceSpaceTransform.GetLocation(), ReferenceSpaceTransform.GetLocation() + AppliedAngularVelocity, FLinearColor::Yellow);
+
+			// Text
+			const FText Label = FText::Format(FText::FromString("VelocityScale\nX: {0} ({1})\nY: {2} ({3})\nZ: {4} ({5})\nA: {6} ({7})"),
+				FText::AsNumber(XApplied), FText::AsNumber(XOrig),
+				FText::AsNumber(YApplied), FText::AsNumber(YOrig),
+				FText::AsNumber(ZApplied), FText::AsNumber(ZOrig),
+				FText::AsNumber(AppliedAngularVelocity.Length()), FText::AsNumber(OriginalAngularVelocity.Length()));
+			DrawText(nullptr, nullptr, ReferenceSpaceTransform.GetLocation() + FVec3(5.), Label, FLinearColor::White, 1.5f);
+		}
+	}
+
 	void FClothVisualizationNoGC::DrawSelfCollision(FPrimitiveDrawInterface* PDI) const
 	{
 		if (!Solver)
diff --git a/Engine/Plugins/ChaosCloth/Source/ChaosCloth/Private/ChaosCloth/ChaosClothingSimulation.cpp b/Engine/Plugins/ChaosCloth/Source/ChaosCloth/Private/ChaosCloth/ChaosClothingSimulation.cpp
--- a/Engine/Plugins/ChaosCloth/Source/ChaosCloth/Private/ChaosCloth/ChaosClothingSimulation.cpp
+++ b/Engine/Plugins/ChaosCloth/Source/ChaosCloth/Private/ChaosCloth/ChaosClothingSimulation.cpp
@@ -101,6 +101,7 @@ namespace ClothingSimulationCVar
 	TAutoConsoleVariable<bool> DebugDrawClothClothConstraints    (TEXT("p.ChaosCloth.DebugDrawClothClothConstraints"    ), false, TEXT("Whether to debug draw the Chaos Cloth-Cloth constraints"), ECVF_Cheat);
 	TAutoConsoleVariable<bool> DebugDrawTeleportReset(TEXT("p.ChaosCloth.DebugDrawTeleportReset"), false, TEXT("Whether to debug draw the Chaos Cloth teleport/reset status"), ECVF_Cheat);
 	TAutoConsoleVariable<bool> DebugDrawExtremlyDeformedEdges(TEXT("p.ChaosCloth.DebugDrawExtremlyDeformedEdges"), false, TEXT("Whether to debug draw edges that are deformed beyond EdgeDeformationThreshold"), ECVF_Cheat);
+	TAutoConsoleVariable<bool> DebugDrawVelocityScale(TEXT("p.ChaosCloth.DebugDrawVelocityScale"), false, TEXT("Whether to debug draw the Chaos Cloth velocity scale"), ECVF_Cheat);
 }
 #endif  // #if CHAOS_DEBUG_DRAW
 
@@ -669,6 +670,7 @@ void FClothingSimulation::Simulate_AnyThread(const IClothingSimulationContext* I
 	if (ClothingSimulationCVar::DebugDrawSelfIntersection    .GetValueOnAnyThread()) { DebugDrawSelfIntersection    (); }
 	if (ClothingSimulationCVar::DebugDrawTeleportReset       .GetValueOnAnyThread()) { DebugDrawTeleportReset       (); }
 	if (ClothingSimulationCVar::DebugDrawExtremlyDeformedEdges.GetValueOnAnyThread()) { DebugDrawExtremlyDeformedEdges(); }
+	if (ClothingSimulationCVar::DebugDrawVelocityScale.GetValueOnAnyThread()) { DebugDrawVelocityScale(); }
 #endif  // #if CHAOS_DEBUG_DRAW
 }
 
diff --git a/Engine/Plugins/ChaosCloth/Source/ChaosCloth/Private/ChaosCloth/ChaosClothingSimulationSolver.cpp b/Engine/Plugins/ChaosCloth/Source/ChaosCloth/Private/ChaosCloth/ChaosClothingSimulationSolver.cpp
--- a/Engine/Plugins/ChaosCloth/Source/ChaosCloth/Private/ChaosCloth/ChaosClothingSimulationSolver.cpp
+++ b/Engine/Plugins/ChaosCloth/Source/ChaosCloth/Private/ChaosCloth/ChaosClothingSimulationSolver.cpp
@@ -207,6 +207,10 @@ FClothingSimulationSolver::FClothingSimulationSolver(FClothingSimulationConfig*
 		Evolution->AddGroupArray(&ReferenceSpaceLocations);
 		Evolution->AddGroupArray(&AppliedReferenceSpaceVelocities);
 		Evolution->AddGroupArray(&AppliedReferenceSpaceAngularVelocities);
+#if CHAOS_DEBUG_DRAW
+		Evolution->AddGroupArray(&OriginalReferenceSpaceVelocities);
+		Evolution->AddGroupArray(&OriginalReferenceSpaceAngularVelocities);
+#endif
 
 		Evolution->AddParticleArray(&Normals);
 		Evolution->AddParticleArray(&OldAnimationPositions);
@@ -300,6 +304,10 @@ FClothingSimulationSolver::FClothingSimulationSolver(FClothingSimulationConfig*
 		PBDEvolution->AddArray(&ReferenceSpaceLocations);
 		PBDEvolution->AddArray(&AppliedReferenceSpaceVelocities);
 		PBDEvolution->AddArray(&AppliedReferenceSpaceAngularVelocities);
+#if CHAOS_DEBUG_DRAW
+		PBDEvolution->AddArray(&OriginalReferenceSpaceVelocities);
+		PBDEvolution->AddArray(&OriginalReferenceSpaceAngularVelocities);
+#endif
 
 		PBDEvolution->Particles().AddArray(&Normals);
 		PBDEvolution->Particles().AddArray(&OldAnimationPositions);
@@ -1315,6 +1323,10 @@ void FClothingSimulationSolver::SetReferenceVelocityScale(uint32 GroupId,
 	ReferenceSpaceLocations[GroupId] = LocalSpaceScaleInv * (ReferenceSpaceTransform.GetLocation() - LocalSpaceLocation);
 	AppliedReferenceSpaceVelocities[GroupId] = DeltaTime > (Softs::FSolverReal)0.f ? FSolverVec3(LocalSpaceScaleInv * InOutAppliedReferenceVelocity) : Softs::FSolverVec3(0.f);
 	AppliedReferenceSpaceAngularVelocities[GroupId] = DeltaTime > (Softs::FSolverReal)0.f ? FSolverVec3(InOutAppliedReferenceAngularVelocity) : Softs::FSolverVec3(0.f);
+#if CHAOS_DEBUG_DRAW
+	OriginalReferenceSpaceVelocities[GroupId] = DeltaTime > (Softs::FSolverReal)0.f ? FSolverVec3(LocalSpaceScaleInv * WorldDeltaTranslation / DeltaTime) : Softs::FSolverVec3(0.f);
+	OriginalReferenceSpaceAngularVelocities[GroupId] = DeltaTime > (Softs::FSolverReal)0.f ? FSolverVec3(FullDeltaRotation / DeltaTime) : Softs::FSolverVec3(0.f);
+#endif
 }
 
 Softs::FSolverReal FClothingSimulationSolver::SetParticleMassPerArea(Softs::FSolverParticlesRange& Particles, const FTriangleMesh& Mesh)
diff --git a/Engine/Plugins/ChaosCloth/Source/ChaosCloth/Public/ChaosCloth/ChaosClothVisualization.h b/Engine/Plugins/ChaosCloth/Source/ChaosCloth/Public/ChaosCloth/ChaosClothVisualization.h
--- a/Engine/Plugins/ChaosCloth/Source/ChaosCloth/Public/ChaosCloth/ChaosClothVisualization.h
+++ b/Engine/Plugins/ChaosCloth/Source/ChaosCloth/Public/ChaosCloth/ChaosClothVisualization.h
@@ -52,6 +52,7 @@ namespace Chaos
 		CHAOSCLOTH_API void DrawWindAndPressureForces(FPrimitiveDrawInterface* PDI, const FReal ForceLengthScale) const;
 		CHAOSCLOTH_API void DrawWindVelocity(FPrimitiveDrawInterface* PDI, const FReal LengthScale) const;
 		CHAOSCLOTH_API void DrawLocalSpace(FPrimitiveDrawInterface* PDI = nullptr) const;
+		CHAOSCLOTH_API void DrawVelocityScale(FPrimitiveDrawInterface* PDI, const FReal LengthScale) const;
 		CHAOSCLOTH_API void DrawSelfCollision(FPrimitiveDrawInterface* PDI = nullptr) const;
 		CHAOSCLOTH_API void DrawSelfIntersection(FPrimitiveDrawInterface* PDI = nullptr) const;
 		CHAOSCLOTH_API void DrawSelfCollisionThickness(FPrimitiveDrawInterface* PDI = nullptr) const;
@@ -89,6 +90,7 @@ namespace Chaos
 		void DrawWindAndPressureForces(FPrimitiveDrawInterface* /*PDI*/, const FReal /*ForceLengthScale*/) const {}
 		void DrawWindVelocity(FPrimitiveDrawInterface* /*PDI*/, const FReal /*LengthScale*/) const {}
 		void DrawLocalSpace(FPrimitiveDrawInterface* /*PDI*/ = nullptr) const {}
+		void DrawVelocityScale(FPrimitiveDrawInterface* /*PDI*/, const FReal /*LengthScale*/) const {}
 		void DrawSelfCollision(FPrimitiveDrawInterface* /*PDI*/ = nullptr) const {}
 		void DrawSelfIntersection(FPrimitiveDrawInterface* /*PDI*/ = nullptr) const {}
 		void DrawSelfCollisionThickness(FPrimitiveDrawInterface* /*PDI*/ = nullptr) const {}
@@ -128,6 +130,11 @@ namespace Chaos
 			constexpr FReal DefaultLengthScale = .1;
 			DrawWindVelocity(PDI, DefaultLengthScale);
 		}
+		void DrawVelocityScale(FPrimitiveDrawInterface* PDI = nullptr) const
+		{
+			constexpr FReal DefaultLengthScale = .1;
+			DrawVelocityScale(PDI, DefaultLengthScale);
+		}
 
 #if WITH_EDITOR && CHAOS_DEBUG_DRAW
 		// Editor only functions
diff --git a/Engine/Plugins/ChaosCloth/Source/ChaosCloth/Public/ChaosCloth/ChaosClothingSimulation.h b/Engine/Plugins/ChaosCloth/Source/ChaosCloth/Public/ChaosCloth/ChaosClothingSimulation.h
--- a/Engine/Plugins/ChaosCloth/Source/ChaosCloth/Public/ChaosCloth/ChaosClothingSimulation.h
+++ b/Engine/Plugins/ChaosCloth/Source/ChaosCloth/Public/ChaosCloth/ChaosClothingSimulation.h
@@ -123,6 +123,7 @@ PRAGMA_ENABLE_DEPRECATION_WARNINGS
 		void DebugDrawGravity(FPrimitiveDrawInterface* PDI = nullptr) const { Visualization.DrawGravity(PDI); }
 		void DebugDrawTeleportReset(FPrimitiveDrawInterface* PDI = nullptr) const { Visualization.DrawTeleportReset(PDI); }
 		void DebugDrawExtremlyDeformedEdges(FPrimitiveDrawInterface* PDI = nullptr) const { Visualization.DrawExtremlyDeformedEdges(PDI); }
+		void DebugDrawVelocityScale(FPrimitiveDrawInterface* PDI = nullptr) const { Visualization.DrawVelocityScale(PDI); }
 #endif  // #if CHAOS_DEBUG_DRAW
 
 		/** Return the visualization object for this simulation. */
diff --git a/Engine/Plugins/ChaosCloth/Source/ChaosCloth/Public/ChaosCloth/ChaosClothingSimulationSolver.h b/Engine/Plugins/ChaosCloth/Source/ChaosCloth/Public/ChaosCloth/ChaosClothingSimulationSolver.h
--- a/Engine/Plugins/ChaosCloth/Source/ChaosCloth/Public/ChaosCloth/ChaosClothingSimulationSolver.h
+++ b/Engine/Plugins/ChaosCloth/Source/ChaosCloth/Public/ChaosCloth/ChaosClothingSimulationSolver.h
@@ -347,6 +347,13 @@ namespace Chaos
 		CHAOSCLOTH_API uint32 GetNumParticles() const;
 		CHAOSCLOTH_API int32 GetNumActiveParticles() const;
 		CHAOSCLOTH_API int32 GetGlobalParticleOffset(int32 ParticleRangeId) const;
+
+		const Softs::FSolverVec3& GetAppliedReferenceSpaceVelocity(uint32 GroupId) const { return AppliedReferenceSpaceVelocities[GroupId]; }
+		const Softs::FSolverVec3& GetAppliedReferenceSpaceAngularVelocity(uint32 GroupId) const { return AppliedReferenceSpaceAngularVelocities[GroupId]; }
+#if CHAOS_DEBUG_DRAW
+		const Softs::FSolverVec3& GetOriginalReferenceSpaceVelocity(uint32 GroupId) const { return OriginalReferenceSpaceVelocities[GroupId]; }
+		const Softs::FSolverVec3& GetOriginalReferenceSpaceAngularVelocity(uint32 GroupId) const { return OriginalReferenceSpaceAngularVelocities[GroupId]; }
+#endif
 		// ---- End of the Cloth interface ----
 
 		// ---- Collider interface ----
@@ -483,6 +490,11 @@ namespace Chaos
 		TArrayCollectionArray<Softs::FSolverVec3> ReferenceSpaceLocations;  // Center of rotations for fictitious forces and local velocity damping in local coordinate to the simulation space location
 		TArrayCollectionArray<Softs::FSolverVec3> AppliedReferenceSpaceVelocities; // Used for local velocity damping.
 		TArrayCollectionArray<Softs::FSolverVec3> AppliedReferenceSpaceAngularVelocities; // Used for local velocity damping.
+#if CHAOS_DEBUG_DRAW
+		TArrayCollectionArray<Softs::FSolverVec3> OriginalReferenceSpaceVelocities; // Reference space velocity before any scaling or clamping. Used for debug draw.
+		TArrayCollectionArray<Softs::FSolverVec3> OriginalReferenceSpaceAngularVelocities; // Reference space velocity before any scaling or clamping. Used for debug draw.
+#endif
+
 
 		// Particle attributes
 		TArrayCollectionArray<Softs::FSolverVec3> Normals;
diff --git a/Engine/Plugins/ChaosCloth/Source/ChaosClothEditor/Private/ChaosClothEditor/ChaosSimulationEditorExtender.cpp b/Engine/Plugins/ChaosCloth/Source/ChaosClothEditor/Private/ChaosClothEditor/ChaosSimulationEditorExtender.cpp
--- a/Engine/Plugins/ChaosCloth/Source/ChaosClothEditor/Private/ChaosClothEditor/ChaosSimulationEditorExtender.cpp
+++ b/Engine/Plugins/ChaosCloth/Source/ChaosClothEditor/Private/ChaosClothEditor/ChaosSimulationEditorExtender.cpp
@@ -85,6 +85,7 @@ const FVisualizationOption FVisualizationOption::OptionData[] =
 	FVisualizationOption(&FClothVisualizationNoGC::DrawWindAndPressureForces, TEXT("p.ChaosClothEditor.DebugDrawWindAndPressureForces"), LOCTEXT("ChaosVisName_WindAndPressureForces", "Wind Aerodynamic And Pressure Forces"), LOCTEXT("ChaosVisName_WindAndPressure_Tooltip", "Draws the Wind drag and lift and pressure forces")),
 	FVisualizationOption(&FClothVisualizationNoGC::DrawWindVelocity, TEXT("p.ChaosClothEditor.DebugDrawWindVelocity"), LOCTEXT("ChaosVisName_WindVelocity", "Wind Velocity"), LOCTEXT("ChaosVisName_WindVelocity_Tooltip", "Draws the Wind Velocity vector")),
 	FVisualizationOption(&FClothVisualizationNoGC::DrawLocalSpace          , TEXT("p.ChaosClothEditor.DebugDrawLocalSpace"          ), LOCTEXT("ChaosVisName_LocalSpace"          , "Local Space Reference Bone" ), LOCTEXT("ChaosVisName_LocalSpace_Tooltip"          , "Draws the local space reference bone")),
+	FVisualizationOption(&FClothVisualizationNoGC::DrawVelocityScale       , TEXT("p.ChaosClothEditor.DebugDrawVelocityScale"       ), LOCTEXT("ChaosVisName_VelocityScale"       , "Velocity Scale"             ), LOCTEXT("ChaosVisName_VelocityScale_Tooltip"       , "Draws the velocity scale")),
 	FVisualizationOption(&FClothVisualizationNoGC::DrawSelfCollision       , TEXT("p.ChaosClothEditor.DebugDrawSelfCollision"       ), LOCTEXT("ChaosVisName_SelfCollision"       , "Self Collision"             ), LOCTEXT("ChaosVisName_SelfCollision_Tooltip"       , "Draws the self collision thickness/debugging information")),
 	FVisualizationOption(&FClothVisualizationNoGC::DrawSelfIntersection    , TEXT("p.ChaosClothEditor.DebugDrawSelfIntersection"    ), LOCTEXT("ChaosVisName_SelfIntersection"    , "Self Intersection"          ), LOCTEXT("ChaosVisName_SelfIntersection_Tooltip"    , "Draws the self intersection contour/region information")),
 };
diff --git a/Engine/Plugins/ChaosClothAsset/Source/ChaosClothAssetEngine/Private/ChaosClothAsset/ClothSimulationProxy.cpp b/Engine/Plugins/ChaosClothAsset/Source/ChaosClothAssetEngine/Private/ChaosClothAsset/ClothSimulationProxy.cpp
--- a/Engine/Plugins/ChaosClothAsset/Source/ChaosClothAssetEngine/Private/ChaosClothAsset/ClothSimulationProxy.cpp
+++ b/Engine/Plugins/ChaosClothAsset/Source/ChaosClothAssetEngine/Private/ChaosClothAsset/ClothSimulationProxy.cpp
@@ -675,6 +675,7 @@ namespace UE::Chaos::ClothAsset
 		static const TConsoleVariableData<bool>* const DebugDrawClothClothConstraintsCVar = IConsoleManager::Get().FindTConsoleVariableDataBool(TEXT("p.ChaosCloth.DebugDrawClothClothConstraints"));
 		static const TConsoleVariableData<bool>* const DebugDrawTeleportResetCVar = IConsoleManager::Get().FindTConsoleVariableDataBool(TEXT("p.ChaosCloth.DebugDrawTeleportReset"));
 		static const TConsoleVariableData<bool>* const DebugDrawExtremlyDeformedEdgesCVar = IConsoleManager::Get().FindTConsoleVariableDataBool(TEXT("p.ChaosCloth.DebugDrawExtremlyDeformedEdges"));
+		static const TConsoleVariableData<bool>* const DebugDrawVelocityScaleCVar = IConsoleManager::Get().FindTConsoleVariableDataBool(TEXT("p.ChaosCloth.DebugDrawVelocityScale"));
 
 		if (DebugDrawLocalSpaceCVar && DebugDrawLocalSpaceCVar->GetValueOnAnyThread()) { Visualization->DrawLocalSpace(); }
 		if (DebugDrawBoundsCVar && DebugDrawBoundsCVar->GetValueOnAnyThread()) { Visualization->DrawBounds(); }
@@ -711,6 +712,10 @@ namespace UE::Chaos::ClothAsset
 		{
 			Visualization->DrawExtremlyDeformedEdges();
 		}
+		if (DebugDrawVelocityScaleCVar && DebugDrawVelocityScaleCVar->GetValueOnAnyThread())
+		{
+			Visualization->DrawVelocityScale();
+		}
 #endif  // #if CHAOS_DEBUG_DRAW
 	}
 
diff --git a/Engine/Plugins/ChaosClothAssetEditor/Source/ChaosClothAssetEditor/Private/ChaosClothAsset/ClothEditorSimulationVisualization.cpp b/Engine/Plugins/ChaosClothAssetEditor/Source/ChaosClothAssetEditor/Private/ChaosClothAsset/ClothEditorSimulationVisualization.cpp
--- a/Engine/Plugins/ChaosClothAssetEditor/Source/ChaosClothAssetEditor/Private/ChaosClothAsset/ClothEditorSimulationVisualization.cpp
+++ b/Engine/Plugins/ChaosClothAssetEditor/Source/ChaosClothAssetEditor/Private/ChaosClothAsset/ClothEditorSimulationVisualization.cpp
@@ -346,6 +346,18 @@ const FVisualizationOption FVisualizationOption::OptionData[] =
 		}),
 		LOCTEXT("ChaosVisName_LocalSpace", "Local Space Reference Bone"), 
 		LOCTEXT("ChaosVisName_LocalSpace_Tooltip", "Draws the local space reference bone")),
+	FVisualizationOption(
+		FClothVisualizationDebugDraw::CreateLambda([](const FClothEditorSimulationVisualization& EditorVisualization, const ::Chaos::FClothVisualizationNoGC& Visualization, FPrimitiveDrawInterface* PDI)
+		{
+			Visualization.DrawVelocityScale(PDI, EditorVisualization.GetVelocityScaleLength());
+		}),
+		LOCTEXT("ChaosVisName_VelocityScale", "Velocity Scale"), 
+		LOCTEXT("ChaosVisName_VelocityScale_Tooltip", "Draws the velocity scale"),
+		/*bDisablesSimulation =*/false, /*bHidesClothSections=*/false, /*bDefaultFlagValue=*/false,
+		FAdditionalMenus::CreateLambda([](FClothEditorSimulationVisualization& EditorVisualization, FMenuBuilder& MenuBuilder)
+		{
+			EditorVisualization.ExtendViewportShowMenuVelocityScaleLength(MenuBuilder);
+		})),
 	FVisualizationOption(
 		FClothVisualizationDebugDraw::CreateLambda([](const FClothEditorSimulationVisualization&, const ::Chaos::FClothVisualizationNoGC& Visualization, FPrimitiveDrawInterface* PDI)
 		{
diff --git a/Engine/Plugins/ChaosClothAssetEditor/Source/ChaosClothAssetEditor/Private/ChaosClothAsset/ClothEditorSimulationVisualization.h b/Engine/Plugins/ChaosClothAssetEditor/Source/ChaosClothAssetEditor/Private/ChaosClothAsset/ClothEditorSimulationVisualization.h
--- a/Engine/Plugins/ChaosClothAssetEditor/Source/ChaosClothAssetEditor/Private/ChaosClothAsset/ClothEditorSimulationVisualization.h
+++ b/Engine/Plugins/ChaosClothAssetEditor/Source/ChaosClothAssetEditor/Private/ChaosClothAsset/ClothEditorSimulationVisualization.h
@@ -74,11 +74,16 @@ public:
 	{
 		ExtendViewportShowMenuSpinBox(MenuBuilder, WindVelocityLengthScale, 0.f, FLT_MAX, 0.f, 40.f);
 	}
+	void ExtendViewportShowMenuVelocityScaleLength(FMenuBuilder& MenuBuilder)
+	{
+		ExtendViewportShowMenuSpinBox(MenuBuilder, VelocityScaleLength, 0.f, FLT_MAX, 0.f, 40.f);
+	}
 	float GetPointNormalLength() const { return PointNormalLength; }
 	float GetAnimatedNormalLength() const { return AnimatedNormalLength; }
 	float GetAccessoryMeshNormalLength() const { return AccessoryMeshNormalLength; }
 	float GetAerodynamicsLengthScale() const { return AerodynamicsLengthScale; }
 	float GetWindVelocityLengthScale() const { return WindVelocityLengthScale; }
+	float GetVelocityScaleLength() const { return VelocityScaleLength; }
 
 private:
 	struct FNameSelectionData
@@ -109,6 +114,7 @@ private:
 	float AccessoryMeshNormalLength = 20.f;
 	float AerodynamicsLengthScale = 10.f;
 	float WindVelocityLengthScale = .1f;
+	float VelocityScaleLength = .1f;
 	
 };
 } // namespace UE::Chaos::ClothAsset
