Chaos Cloth - Added MaxLinearVelocity/MaxAngularVelocity/MaxLinearAcceleration/MaxAngularAcceleration clamp CVars.

[RN] Chaos Cloth - Added MaxLinearVelocity/MaxAngularVelocity/MaxLinearAcceleration/MaxAngularAcceleration clamp CVars.
#tests Verified that when unset, the changes have no effect.

diff --git a/Engine/Plugins/ChaosCloth/Source/ChaosCloth/Private/ChaosCloth/ChaosClothingSimulationCloth.cpp b/Engine/Plugins/ChaosCloth/Source/ChaosCloth/Private/ChaosCloth/ChaosClothingSimulationCloth.cpp
--- a/Engine/Plugins/ChaosCloth/Source/ChaosCloth/Private/ChaosCloth/ChaosClothingSimulationCloth.cpp
+++ b/Engine/Plugins/ChaosCloth/Source/ChaosCloth/Private/ChaosCloth/ChaosClothingSimulationCloth.cpp
@@ -61,6 +61,38 @@ namespace ClothingSimulationClothConsoleVariables
 		TEXT("p.ChaosCloth.GravityMultiplier"),
 		1.f,
 		TEXT("Scalar multiplier applied at the final stage of the cloth's gravity formulation."));
+
+	static float GlobalMaxLinearVelocity = 0.f;
+	FAutoConsoleVariableRef CVarMaxLinearVelocity(
+		TEXT("p.ChaosCloth.MaxLinearVelocity"),
+		GlobalMaxLinearVelocity,
+		TEXT("The maximum amount of linear velocity sent to the local cloth space from the reference bone. "
+			 "The console variable value will take precedence when set smaller to the equivalent Velocity Scale config properties."
+			 "Use 0 to disable."));
+
+	static float GlobalMaxAngularVelocity = 0.f;
+	FAutoConsoleVariableRef CVarMaxAngularVelocity(
+		TEXT("p.ChaosCloth.MaxAngularVelocity"),
+		GlobalMaxAngularVelocity,
+		TEXT("The maximum amount of angular velocity sent to the local cloth space from the reference bone. "
+			 "The console variable value will take precedence when set smaller to the equivalent Velocity Scale config properties."
+			 "Use 0 to disable."));
+
+	static float GlobalMaxLinearAcceleration = 0.f;
+	FAutoConsoleVariableRef CVarMaxLinearAcceleration(
+		TEXT("p.ChaosCloth.MaxLinearAcceleration"),
+		GlobalMaxLinearAcceleration,
+		TEXT("The maximum amount of linear acceleration sent to the local cloth space from the reference bone. "
+			 "The console variable value will take precedence when set smaller to the equivalent Velocity Scale config properties."
+			 "Use 0 to disable."));
+
+	static float GlobalMaxAngularAcceleration = 0.f;
+	FAutoConsoleVariableRef CVarMaxAngularAcceleration(
+		TEXT("p.ChaosCloth.MaxAngularAcceleration"),
+		GlobalMaxAngularAcceleration,
+		TEXT("The maximum amount of angular acceleration sent to the local cloth space from the reference bone. "
+			 "The console variable value will take precedence when set smaller to the equivalent Velocity Scale config properties."
+			 "Use 0 to disable."));
 }
 
 using namespace Softs;
@@ -1143,6 +1175,13 @@ void FClothingSimulationCloth::Update(FClothingSimulationSolver* Solver)
 			MaxLinearAcceleration = GetMaxLinearAcceleration(*ConfigProperties, MaxLinearAcceleration);
 			MaxAngularVelocity = GetMaxAngularVelocity(*ConfigProperties, MaxAngularVelocity);
 			MaxAngularAcceleration = GetMaxAngularAcceleration(*ConfigProperties, MaxAngularAcceleration);
+
+			// Apply the cvars clamp, if any are active
+			using namespace ClothingSimulationClothConsoleVariables;
+			MaxLinearVelocity = (GlobalMaxLinearVelocity > 0.f) ? MaxLinearVelocity.ComponentwiseMin(FVec3f(GlobalMaxLinearVelocity)) : MaxLinearVelocity;
+			MaxLinearAcceleration = (GlobalMaxLinearAcceleration > 0.f) ? MaxLinearAcceleration.ComponentwiseMin(FVec3f(GlobalMaxLinearAcceleration)) : MaxLinearAcceleration;
+			MaxAngularVelocity = (GlobalMaxAngularVelocity > 0.f) ? FMath::Min(MaxAngularVelocity, (FRealSingle)GlobalMaxAngularVelocity) : MaxAngularVelocity;
+			MaxAngularAcceleration = (GlobalMaxAngularAcceleration > 0.f) ? FMath::Min(MaxAngularAcceleration, (FRealSingle)GlobalMaxAngularAcceleration) : MaxAngularAcceleration;
 		}
 
 		// NOTE: Force-based solver doesn't actually use FictitiousAngularScale here. It gets it from the property collection directly.
