Chaos Cloth - Added fictitious forces on full velocity scale simulations.

[RN] Chaos Cloth - Added fictitious forces on full velocity scale simulations.

diff --git a/Engine/Plugins/ChaosCloth/Source/ChaosCloth/Private/ChaosCloth/ChaosClothConfig.cpp b/Engine/Plugins/ChaosCloth/Source/ChaosCloth/Private/ChaosCloth/ChaosClothConfig.cpp
--- a/Engine/Plugins/ChaosCloth/Source/ChaosCloth/Private/ChaosCloth/ChaosClothConfig.cpp
+++ b/Engine/Plugins/ChaosCloth/Source/ChaosCloth/Private/ChaosCloth/ChaosClothConfig.cpp
@@ -8,6 +8,7 @@
 #include "Components/SkeletalMeshComponent.h"
 #include "UObject/PhysicsObjectVersion.h"
 #include "UObject/FortniteMainBranchObjectVersion.h"
+#include "UObject/FortniteReleaseBranchCustomObjectVersion.h"
 #include "UObject/UE5ReleaseStreamObjectVersion.h"
 #include "UObject/UObjectIterator.h"
 
@@ -132,6 +133,7 @@ void UChaosClothConfig::Serialize(FArchive& Ar)
 	Ar.UsingCustomVersion(FChaosClothConfigCustomVersion::GUID);
 	Ar.UsingCustomVersion(FPhysicsObjectVersion::GUID);
 	Ar.UsingCustomVersion(FFortniteMainBranchObjectVersion::GUID);
+	Ar.UsingCustomVersion(FFortniteReleaseBranchCustomObjectVersion::GUID);
 	Ar.UsingCustomVersion(FUE5ReleaseStreamObjectVersion::GUID);
 }
 
@@ -144,6 +146,7 @@ void UChaosClothConfig::PostLoad()
 	const int32 PhysicsObjectVersion = GetLinkerCustomVersion(FPhysicsObjectVersion::GUID);
 	const int32 FortniteMainBranchObjectVersion = GetLinkerCustomVersion(FFortniteMainBranchObjectVersion::GUID);
 	const int32 UE5ReleaseStreamObjectVersion = GetLinkerCustomVersion(FUE5ReleaseStreamObjectVersion::GUID);
+	const int32 FortniteReleaseBranchCustomObjectVersion = GetLinkerCustomVersion(FFortniteReleaseBranchCustomObjectVersion::GUID);
 
 	if (ChaosClothConfigCustomVersion < FChaosClothConfigCustomVersion::UpdateDragDefault)
 	{
@@ -217,6 +220,11 @@ void UChaosClothConfig::PostLoad()
 			BucklingStiffnessWeighted.Low = BucklingStiffnessWeighted.High = 1.f;
 		}
 	}
+
+	if (FortniteReleaseBranchCustomObjectVersion < FFortniteReleaseBranchCustomObjectVersion::ChaosClothAddNewFictitiousForcesModes)
+	{
+		FictitiousForcesModel = EChaosClothingSimulationSolverFictitiousForcesModel::Legacy;
+	}
 #endif  // #if WITH_EDITORONLY_DATA
 }
 
diff --git a/Engine/Plugins/ChaosCloth/Source/ChaosCloth/Private/ChaosCloth/ChaosClothingSimulationCloth.cpp b/Engine/Plugins/ChaosCloth/Source/ChaosCloth/Private/ChaosCloth/ChaosClothingSimulationCloth.cpp
--- a/Engine/Plugins/ChaosCloth/Source/ChaosCloth/Private/ChaosCloth/ChaosClothingSimulationCloth.cpp
+++ b/Engine/Plugins/ChaosCloth/Source/ChaosCloth/Private/ChaosCloth/ChaosClothingSimulationCloth.cpp
@@ -1138,7 +1138,7 @@ void FClothingSimulationCloth::Update(FClothingSimulationSolver* Solver)
 			OutMaxVelocityScale = 1.f;
 			AppliedReferenceSpaceAngularVelocity = FVec3(0.);
 			AppliedReferenceSpaceVelocity = FVec3(0.);
-			bDisableFictitiousForces = true; // It doesn't actually matter what value we set here since AngularVelocityScale == 1 means fictitious forces will be 0.
+			bDisableFictitiousForces = true;
 
 			// Reset to start pose
 			LODData[LODIndex]->ResetStartPose(Solver);
@@ -1161,7 +1161,7 @@ void FClothingSimulationCloth::Update(FClothingSimulationSolver* Solver)
 			OutMaxVelocityScale = 1.f;
 			AppliedReferenceSpaceAngularVelocity = FVec3(0.);
 			AppliedReferenceSpaceVelocity = FVec3(0.);
-			bDisableFictitiousForces = true; // Disable fictitious forces. Otherwise they will be applied since AngularVelocityScale < 1.
+			bDisableFictitiousForces = true;
 			UE_LOG(LogChaosCloth, Verbose, TEXT("Cloth in group Id %d Needs teleport."), GroupId);
 		}
 		else
@@ -1183,6 +1183,10 @@ void FClothingSimulationCloth::Update(FClothingSimulationSolver* Solver)
 			MaxAngularVelocity = (GlobalMaxAngularVelocity > 0.f) ? FMath::Min(MaxAngularVelocity, (FRealSingle)GlobalMaxAngularVelocity) : MaxAngularVelocity;
 			MaxAngularAcceleration = (GlobalMaxAngularAcceleration > 0.f) ? FMath::Min(MaxAngularAcceleration, (FRealSingle)GlobalMaxAngularAcceleration) : MaxAngularAcceleration;
 		}
+		using EFictitiousForcesModel = FClothingSimulationSolver::EFictitiousForcesModel;
+		const EFictitiousForcesModel FictitiousForcesModel = bDisableFictitiousForces ?
+			EFictitiousForcesModel::None :
+			(EFictitiousForcesModel)GetFictitiousForcesModel(*ConfigProperties, (int32)EFictitiousForcesModel::Legacy);
 
 		// NOTE: Force-based solver doesn't actually use FictitiousAngularScale here. It gets it from the property collection directly.
 		const FRealSingle FictitiousAngularScale = GetFictitiousAngularScale(*ConfigProperties, ClothingSimulationClothDefault::FictitiousAngularScale);
@@ -1201,7 +1205,7 @@ void FClothingSimulationCloth::Update(FClothingSimulationSolver* Solver)
 			MaxAngularAcceleration,
 			FictitiousAngularScale,
 			OutMaxVelocityScale,
-			bDisableFictitiousForces);
+			FictitiousForcesModel);
 
 		if (bNeedsReset || bNeedsTeleport)
 		{
diff --git a/Engine/Plugins/ChaosCloth/Source/ChaosCloth/Private/ChaosCloth/ChaosClothingSimulationConfig.cpp b/Engine/Plugins/ChaosCloth/Source/ChaosCloth/Private/ChaosCloth/ChaosClothingSimulationConfig.cpp
--- a/Engine/Plugins/ChaosCloth/Source/ChaosCloth/Private/ChaosCloth/ChaosClothingSimulationConfig.cpp
+++ b/Engine/Plugins/ChaosCloth/Source/ChaosCloth/Private/ChaosCloth/ChaosClothingSimulationConfig.cpp
@@ -166,6 +166,7 @@ namespace Chaos
 				Property->AddValue(FClothingSimulationCloth::LinearVelocityScaleName, FVector3f(ClothConfig->LinearVelocityScale), AnimatablePropertyFlags);
 				Property->AddValue(FClothingSimulationCloth::AngularVelocityScaleName, ClothConfig->AngularVelocityScale, AnimatablePropertyFlags);
 				Property->AddValue(FClothingSimulationCloth::FictitiousAngularScaleName, ClothConfig->FictitiousAngularScale, AnimatablePropertyFlags);
+				Property->AddValue(FClothingSimulationCloth::FictitiousForcesModelName, (int32)ClothConfig->FictitiousForcesModel, AnimatablePropertyFlags);
 
 				static const FVector3f DefaultLinearClamp = { TNumericLimits<float>::Max(), TNumericLimits<float>::Max(), TNumericLimits<float>::Max() };
 				Property->AddValue(FClothingSimulationCloth::MaxLinearVelocityName, ClothConfig->bEnableLinearVelocityClamping ? ClothConfig->MaxLinearVelocity : DefaultLinearClamp, AnimatablePropertyFlags);
diff --git a/Engine/Plugins/ChaosCloth/Source/ChaosCloth/Private/ChaosCloth/ChaosClothingSimulationSolver.cpp b/Engine/Plugins/ChaosCloth/Source/ChaosCloth/Private/ChaosCloth/ChaosClothingSimulationSolver.cpp
--- a/Engine/Plugins/ChaosCloth/Source/ChaosCloth/Private/ChaosCloth/ChaosClothingSimulationSolver.cpp
+++ b/Engine/Plugins/ChaosCloth/Source/ChaosCloth/Private/ChaosCloth/ChaosClothingSimulationSolver.cpp
@@ -1131,7 +1131,7 @@ void FClothingSimulationSolver::SetReferenceVelocityScale(uint32 GroupId,
 	FRealSingle MaxAngularAcceleration,
 	FRealSingle FictitiousAngularScale,
 	FRealSingle MaxVelocityScale,
-	bool bDisableFictitiousForces
+	EFictitiousForcesModel FictitiousForcesModel
 )
 {
 	FRigidTransform3 OldRootBoneLocalTransform = OldReferenceSpaceTransform;
@@ -1314,11 +1314,14 @@ void FClothingSimulationSolver::SetReferenceVelocityScale(uint32 GroupId,
 		Softs::FSolverRotation3(PreSimulationTransform.GetRotation()));
 
 	// Fictitious angular scale only applied for PBDEvolution. It's applied by ExternalForces for Evolution.
-	const FReal AppliedFictitiousAngularScale = bDisableFictitiousForces ? (FReal)0.f : (PBDEvolution ? FMath::Min((FReal)2., (FReal)FictitiousAngularScale) : (FReal)1.);
+	const FReal AppliedFictitiousAngularScale = (FictitiousForcesModel == EFictitiousForcesModel::None) ? (FReal)0.f : (PBDEvolution ? FMath::Min((FReal)2., (FReal)FictitiousAngularScale) : (FReal)1.);
 
 	// Save the reference bone relative angular velocity for calculating the fictitious forces
-	const FVec3 FictitiousAngularDisplacement = ReferenceSpaceTransform.TransformVector(DeltaRotation)
-		* AppliedFictitiousAngularScale;
+	const FVec3 FictitiousAngularDisplacement = (FictitiousForcesModel == EFictitiousForcesModel::None) ? FVec3(0) :
+		(FictitiousForcesModel == EFictitiousForcesModel::Legacy) ?
+			ReferenceSpaceTransform.TransformVector(DeltaRotation) * AppliedFictitiousAngularScale :  // Legacy
+			ReferenceSpaceTransform.TransformVector(AppliedDeltaRotation)* AppliedFictitiousAngularScale;  // Full scale
+
 	FictitiousAngularVelocities[GroupId] = DeltaTime > (Softs::FSolverReal)0.f ? Softs::FSolverVec3(FictitiousAngularDisplacement) / DeltaTime : Softs::FSolverVec3(0.f);
 	ReferenceSpaceLocations[GroupId] = LocalSpaceScaleInv * (ReferenceSpaceTransform.GetLocation() - LocalSpaceLocation);
 	AppliedReferenceSpaceVelocities[GroupId] = DeltaTime > (Softs::FSolverReal)0.f ? FSolverVec3(LocalSpaceScaleInv * InOutAppliedReferenceVelocity) : Softs::FSolverVec3(0.f);
diff --git a/Engine/Plugins/ChaosCloth/Source/ChaosCloth/Private/ChaosCloth/IspcTestChaosClothingSimulationSolver.cpp b/Engine/Plugins/ChaosCloth/Source/ChaosCloth/Private/ChaosCloth/IspcTestChaosClothingSimulationSolver.cpp
--- a/Engine/Plugins/ChaosCloth/Source/ChaosCloth/Private/ChaosCloth/IspcTestChaosClothingSimulationSolver.cpp
+++ b/Engine/Plugins/ChaosCloth/Source/ChaosCloth/Private/ChaosCloth/IspcTestChaosClothingSimulationSolver.cpp
@@ -88,11 +88,20 @@ bool FIspcTestChaosClothingSimulationSolverPreSimulation::RunTest(const FString&
 	const FQuat4d Rotation(1., 1., 1., 1.);
 	// Debug complains that the rotation isn't normalized.
 	const FQuat4d NormalizedRotation(Rotation.GetNormalized());
+
+	constexpr uint32 GroupId = 0;
 	const FRigidTransform3 OldTransform(Softs::FSolverVec3(1.), NormalizedRotation);
 	const FRigidTransform3 NewTransform(Softs::FSolverVec3(10.), NormalizedRotation);
+	TVec3<FReal> ReferenceVelocity(0.);
+	TVec3<FReal> ReferenceAngularVelocity(0.);
 	const TVec3<FRealSingle> LinearVelocity(10.f);
+	const TVec3<FRealSingle> MaxLinearVelocity(TNumericLimits<FRealSingle>::Max());
+	const TVec3<FRealSingle> MaxLinearAcceleration(TNumericLimits<FRealSingle>::Max());
 	const FRealSingle AngularVelocityScale = 1.f;
+	const FRealSingle MaxAngularVelocity = TNumericLimits<FRealSingle>::Max();
+	const FRealSingle MaxAngularAcceleration = TNumericLimits<FRealSingle>::Max();
 	const FRealSingle FicticiousAngularScale = 1.f;
+	constexpr FRealSingle MaxVelocityScale = 1.f;
 
 	bChaos_PreSimulationTransforms_ISPC_Enabled = true;
 	FClothingSimulationSolver ISPCSimulation;
@@ -123,7 +132,23 @@ bool FIspcTestChaosClothingSimulationSolverPreSimulation::RunTest(const FString&
 		}
 
 		ISPCSimulation.ResetStartPose(0, NumParticles);
-		ISPCSimulation.SetReferenceVelocityScale(0, OldTransform, NewTransform, LinearVelocity, AngularVelocityScale, FicticiousAngularScale);
+
+		ISPCSimulation.SetReferenceVelocityScale(
+			GroupId,
+			OldTransform,
+			NewTransform,
+			ReferenceVelocity,
+			ReferenceAngularVelocity,
+			EChaosSoftsSimulationSpace::ReferenceBoneSpace,
+			LinearVelocity,
+			MaxLinearVelocity,
+			MaxLinearAcceleration,
+			AngularVelocityScale,
+			MaxAngularVelocity,
+			MaxAngularAcceleration,
+			FicticiousAngularScale,
+			MaxVelocityScale,
+			FClothingSimulationSolver::EFictitiousForcesModel::None);
 
 		ISPCSimulation.Update(.5); // Calls ApplyPreSimulationTransform()
 	}
@@ -157,12 +182,28 @@ bool FIspcTestChaosClothingSimulationSolverPreSimulation::RunTest(const FString&
 		}
 
 		CPPSimulation.ResetStartPose(0, NumParticles);
-		CPPSimulation.SetReferenceVelocityScale(0, OldTransform, NewTransform, LinearVelocity, AngularVelocityScale, FicticiousAngularScale);
+
+		CPPSimulation.SetReferenceVelocityScale(
+			GroupId,
+			OldTransform,
+			NewTransform,
+			ReferenceVelocity,
+			ReferenceAngularVelocity,
+			EChaosSoftsSimulationSpace::ReferenceBoneSpace,
+			LinearVelocity,
+			MaxLinearVelocity,
+			MaxLinearAcceleration,
+			AngularVelocityScale,
+			MaxAngularVelocity,
+			MaxAngularAcceleration,
+			FicticiousAngularScale,
+			MaxVelocityScale,
+			FClothingSimulationSolver::EFictitiousForcesModel::None);
 
 		CPPSimulation.Update(.5); // Calls ApplyPreSimulationTransform()
 	}
 
-	bChaos_CalculateBounds_ISPC_Enabled = InitialState;
+	bChaos_PreSimulationTransforms_ISPC_Enabled = InitialState;
 
 	for (int32 i = 0; i < NumParticles; ++i)
 	{
diff --git a/Engine/Plugins/ChaosCloth/Source/ChaosCloth/Public/ChaosCloth/ChaosClothConfig.h b/Engine/Plugins/ChaosCloth/Source/ChaosCloth/Public/ChaosCloth/ChaosClothConfig.h
--- a/Engine/Plugins/ChaosCloth/Source/ChaosCloth/Public/ChaosCloth/ChaosClothConfig.h
+++ b/Engine/Plugins/ChaosCloth/Source/ChaosCloth/Public/ChaosCloth/ChaosClothConfig.h
@@ -1,6 +1,7 @@
 // Copyright Epic Games, Inc. All Rights Reserved.
 #pragma once
 #include "Chaos/SoftsSimulationSpace.h"
+#include "ChaosCloth/ChaosClothingSimulationSolver.h"
 #include "Components/SkeletalMeshComponent.h"
 #include "Engine/SkeletalMesh.h"
 #include "ClothConfig.h"
@@ -470,14 +471,24 @@ public:
 	UPROPERTY(EditAnywhere, Category = "Animation Properties", meta = (ClampMin = "0", EditCondition = "bEnableAngularAccelerationClamping"))
 	float MaxAngularAcceleration = 12000.f;
 
+	/**
+	 * Fictitious forces model.
+	 * Decide how fictitious forces (e.g. centrifugal force) are applied to the simulation.
+	 */
+	UPROPERTY(EditAnywhere, Category = "Animation Properties")
+	EChaosClothingSimulationSolverFictitiousForcesModel FictitiousForcesModel = EChaosClothingSimulationSolverFictitiousForcesModel::FullScale;
+
 	/**
 	 * The portion of the angular velocity that is used to calculate the strength of all fictitious forces (e.g. centrifugal force).
-	 * This parameter is only having an effect on the portion of the reference bone's angular velocity that has been removed from the
-	 * simulation via the Angular Velocity Scale parameter. This means it has no effect when AngularVelocityScale is set to 1 in which
-	 * case the cloth is simulated with full world space angular velocities and subjected to the true physical world inertial forces.
+	 * In Legacy model, this parameter is only having an effect on the portion of the reference bone's angular velocity that has been
+	 * removed from the simulation via the Angular Velocity Scale parameter. This means it has no effect when AngularVelocityScale is
+	 * set to 1 and Angular Velocity and Acceleration clamps are disabled, in which case the cloth is simulated with full world space
+	 * angular velocities and subjected to the true physical world inertial forces.
+	 * In FullScale model, this parameter has also an effect when the cloth is simulated with full world space angular velocities, and
+	 * isn't affected by the Angular Velocity Scale parameter.
 	 * Values range from 0 to 2, with 0 showing no centrifugal effect, 1 full centrifugal effect, and 2 an overdriven centrifugal effect.
 	 */
-	UPROPERTY(EditAnywhere, Category = "Animation Properties", meta = (UIMin = "0", UIMax = "1", ClampMin = "0", ClampMax = "2"))
+	UPROPERTY(EditAnywhere, Category = "Animation Properties", meta = (UIMin = "0", UIMax = "1", ClampMin = "0", ClampMax = "2", EditCondition = "FictitiousForcesModel != EChaosClothingSimulationSolverFictitiousForcesModel::None"))
 	float FictitiousAngularScale = 1.f;
 
 	/** Enable tetrahedral constraints. */
diff --git a/Engine/Plugins/ChaosCloth/Source/ChaosCloth/Public/ChaosCloth/ChaosClothingSimulationCloth.h b/Engine/Plugins/ChaosCloth/Source/ChaosCloth/Public/ChaosCloth/ChaosClothingSimulationCloth.h
--- a/Engine/Plugins/ChaosCloth/Source/ChaosCloth/Public/ChaosCloth/ChaosClothingSimulationCloth.h
+++ b/Engine/Plugins/ChaosCloth/Source/ChaosCloth/Public/ChaosCloth/ChaosClothingSimulationCloth.h
@@ -185,6 +185,7 @@ namespace Chaos
 		UE_CHAOS_DECLARE_INDEXLESS_PROPERTYCOLLECTION_NAME(MaxAngularVelocity, float);
 		UE_CHAOS_DECLARE_INDEXLESS_PROPERTYCOLLECTION_NAME(MaxAngularAcceleration, float);
 		UE_CHAOS_DECLARE_INDEXLESS_PROPERTYCOLLECTION_NAME(FictitiousAngularScale, float);
+		UE_CHAOS_DECLARE_INDEXLESS_PROPERTYCOLLECTION_NAME(FictitiousForcesModel, int32);
 
 		UE_CHAOS_DECLARE_INDEXLESS_PROPERTYCOLLECTION_NAME(UsePointBasedWindModel, bool);
 
diff --git a/Engine/Plugins/ChaosCloth/Source/ChaosCloth/Public/ChaosCloth/ChaosClothingSimulationSolver.h b/Engine/Plugins/ChaosCloth/Source/ChaosCloth/Public/ChaosCloth/ChaosClothingSimulationSolver.h
--- a/Engine/Plugins/ChaosCloth/Source/ChaosCloth/Public/ChaosCloth/ChaosClothingSimulationSolver.h
+++ b/Engine/Plugins/ChaosCloth/Source/ChaosCloth/Public/ChaosCloth/ChaosClothingSimulationSolver.h
@@ -10,11 +10,24 @@
 #include "Chaos/Framework/PhysicsSolverBase.h"
 #include "Chaos/ImplicitObject.h"
 #include "PhysicsProxy/PerSolverFieldSystem.h"
+#include "ChaosClothingSimulationSolver.generated.h"
 
 class USkeletalMeshComponent;
 class UClothingAssetCommon;
 struct FClothingSimulationCacheData;
 
+/**
+ * Fictitious forces model.
+ * Decide how fictitious forces (e.g. centrifugal force) are applied to the simulation.
+ */
+UENUM()
+enum class EChaosClothingSimulationSolverFictitiousForcesModel : uint8
+{
+	None,  // Don't apply fictitious forces.
+	Legacy,  // Use partial scaled angular velocity, fictitious forces are calculated from the portion of the scaled character velocities. Legacy behavio, use the FullScale model instead.
+	FullScale,  // Use full angular velocity, fictitious forces are calculated from the full scale character velocities.
+};
+
 namespace Chaos
 {
 	class FTriangleMesh;
@@ -34,6 +47,8 @@ namespace Chaos
 	class FClothingSimulationSolver final : public FPhysicsSolverEvents
 	{
 	public:
+		using EFictitiousForcesModel = EChaosClothingSimulationSolverFictitiousForcesModel;
+
 		CHAOSCLOTH_API FClothingSimulationSolver(FClothingSimulationConfig* InConfig = nullptr, bool bUseLegacySolver = true);
 		CHAOSCLOTH_API ~FClothingSimulationSolver();
 		
@@ -138,7 +153,26 @@ namespace Chaos
 
 		// Set the amount of velocity allowed to filter from the given change in reference space transform, including local simulation space.
 		// NOTE: Force-based solver does not apply FictitiousAngularScale here. It's applied directly via the PropertyCollection.
-		CHAOSCLOTH_API void SetReferenceVelocityScale(uint32 GroupId,
+		CHAOSCLOTH_API void SetReferenceVelocityScale(
+			uint32 GroupId,
+			const FRigidTransform3& OldReferenceSpaceTransform,
+			const FRigidTransform3& ReferenceSpaceTransform,
+			TVec3<FReal>& InOutAppliedReferenceVelocity, // Old reference velocity is passed in. New reference velocity is returned.
+			TVec3<FReal>& InOutAppliedReferenceAngularVelocity, // Old reference velocity is passed in. New reference velocity is returned.
+			const EChaosSoftsSimulationSpace VelocityScaleSpace, // the space the following linear velocity properties are in.
+			const TVec3<FRealSingle>& LinearVelocityScale,
+			const TVec3<FRealSingle>& MaxLinearVelocity,
+			const TVec3<FRealSingle>& MaxLinearAcceleration,
+			FRealSingle AngularVelocityScale,
+			FRealSingle MaxAngularVelocity,
+			FRealSingle MaxAngularAcceleration,
+			FRealSingle FictitiousAngularScale,
+			FRealSingle MaxVelocityScale,
+			EFictitiousForcesModel FictitiousForcesModel);
+
+		UE_DEPRECATED(5.8, "Use the version that takes a EFictitiousForcesModel argument instead.")
+		void SetReferenceVelocityScale(
+			uint32 GroupId,
 			const FRigidTransform3& OldReferenceSpaceTransform,
 			const FRigidTransform3& ReferenceSpaceTransform,
 			TVec3<FReal>& InOutAppliedReferenceVelocity, // Old reference velocity is passed in. New reference velocity is returned.
@@ -152,9 +186,29 @@ namespace Chaos
 			FRealSingle MaxAngularAcceleration,
 			FRealSingle FictitiousAngularScale,
 			FRealSingle MaxVelocityScale,
-			bool bDisableFictitiousForces);
+			bool bDisableFictitiousForces)
+		{
+			SetReferenceVelocityScale(
+				GroupId,
+				OldReferenceSpaceTransform,
+				ReferenceSpaceTransform,
+				InOutAppliedReferenceVelocity,
+				InOutAppliedReferenceAngularVelocity,
+				VelocityScaleSpace,
+				LinearVelocityScale,
+				MaxLinearVelocity,
+				MaxLinearAcceleration,
+				AngularVelocityScale,
+				MaxAngularVelocity,
+				MaxAngularAcceleration,
+				FictitiousAngularScale,
+				MaxVelocityScale,
+				bDisableFictitiousForces ? EFictitiousForcesModel::None : EFictitiousForcesModel::Legacy);
+		}
 
-		void SetReferenceVelocityScale(uint32 GroupId,
+		UE_DEPRECATED(5.8, "Use the version that takes a EFictitiousForcesModel argument instead.")
+		void SetReferenceVelocityScale(
+			uint32 GroupId,
 			const FRigidTransform3& OldReferenceSpaceTransform,
 			const FRigidTransform3& ReferenceSpaceTransform,
 			const TVec3<FRealSingle>& LinearVelocityScale,
@@ -166,10 +220,22 @@ namespace Chaos
 			// Acceleration clamps are disabled, so it doesn't matter what the old velocities were.
 			TVec3<FReal> ReferenceVelocity(0.);
 			TVec3<FReal> ReferenceAngularVelocity(0.);
-			SetReferenceVelocityScale(GroupId, OldReferenceSpaceTransform, ReferenceSpaceTransform, ReferenceVelocity, ReferenceAngularVelocity,
-				EChaosSoftsSimulationSpace::ReferenceBoneSpace, LinearVelocityScale, TVec3<FRealSingle>(TNumericLimits<FRealSingle>::Max()),
-				TVec3<FRealSingle>(TNumericLimits<FRealSingle>::Max()), AngularVelocityScale, FictitiousAngularScale,
-				MaxVelocityScale, TNumericLimits<FRealSingle>::Max(), TNumericLimits<FRealSingle>::Max(), bDisableFictitiousForces);
+			SetReferenceVelocityScale(
+				GroupId,
+				OldReferenceSpaceTransform,
+				ReferenceSpaceTransform,
+				ReferenceVelocity,
+				ReferenceAngularVelocity,
+				EChaosSoftsSimulationSpace::ReferenceBoneSpace,
+				LinearVelocityScale,
+				TVec3<FRealSingle>(TNumericLimits<FRealSingle>::Max()),
+				TVec3<FRealSingle>(TNumericLimits<FRealSingle>::Max()),
+				AngularVelocityScale,
+				TNumericLimits<FRealSingle>::Max(),
+				TNumericLimits<FRealSingle>::Max(),
+				FictitiousAngularScale,
+				MaxVelocityScale,
+				bDisableFictitiousForces ? EFictitiousForcesModel::None : EFictitiousForcesModel::Legacy);
 		}
 
 		/** PBDSolver version */
diff --git a/Engine/Plugins/ChaosClothAssetEditor/Source/ChaosClothAssetDataflowNodes/Private/ChaosClothAsset/SimulationVelocityScaleConfigNode.cpp b/Engine/Plugins/ChaosClothAssetEditor/Source/ChaosClothAssetDataflowNodes/Private/ChaosClothAsset/SimulationVelocityScaleConfigNode.cpp
--- a/Engine/Plugins/ChaosClothAssetEditor/Source/ChaosClothAssetDataflowNodes/Private/ChaosClothAsset/SimulationVelocityScaleConfigNode.cpp
+++ b/Engine/Plugins/ChaosClothAssetEditor/Source/ChaosClothAssetDataflowNodes/Private/ChaosClothAsset/SimulationVelocityScaleConfigNode.cpp
@@ -2,6 +2,7 @@
 
 #include "ChaosClothAsset/SimulationVelocityScaleConfigNode.h"
 #include "Chaos/CollectionPropertyFacade.h"
+#include "UObject/FortniteReleaseBranchCustomObjectVersion.h"
 
 #include UE_INLINE_GENERATED_CPP_BY_NAME(SimulationVelocityScaleConfigNode)
 
@@ -23,4 +24,15 @@ void FChaosClothAssetSimulationVelocityScaleConfigNode::AddProperties(FPropertyH
 	PropertyHelper.SetProperty(TEXT("MaxAngularAcceleration"), bEnableAngularAccelerationClamping ? MaxAngularAcceleration : TNumericLimits<float>::Max());
 	PropertyHelper.SetProperty(this, &MaxVelocityScale);
 	PropertyHelper.SetProperty(this, &FictitiousAngularScale);
+	PropertyHelper.SetPropertyEnum(this, &FictitiousForcesModel);
+}
+
+void FChaosClothAssetSimulationVelocityScaleConfigNode::Serialize(FArchive& Ar)
+{
+	Super::Serialize(Ar);
+	Ar.UsingCustomVersion(FFortniteReleaseBranchCustomObjectVersion::GUID);
+	if (Ar.CustomVer(FFortniteReleaseBranchCustomObjectVersion::GUID) < FFortniteReleaseBranchCustomObjectVersion::ChaosClothAddNewFictitiousForcesModes)
+	{
+		FictitiousForcesModel = EChaosClothingSimulationSolverFictitiousForcesModel::Legacy;
+	}
 }
diff --git a/Engine/Plugins/ChaosClothAssetEditor/Source/ChaosClothAssetDataflowNodes/Public/ChaosClothAsset/SimulationVelocityScaleConfigNode.h b/Engine/Plugins/ChaosClothAssetEditor/Source/ChaosClothAssetDataflowNodes/Public/ChaosClothAsset/SimulationVelocityScaleConfigNode.h
--- a/Engine/Plugins/ChaosClothAssetEditor/Source/ChaosClothAssetDataflowNodes/Public/ChaosClothAsset/SimulationVelocityScaleConfigNode.h
+++ b/Engine/Plugins/ChaosClothAssetEditor/Source/ChaosClothAssetDataflowNodes/Public/ChaosClothAsset/SimulationVelocityScaleConfigNode.h
@@ -2,9 +2,10 @@
 
 #pragma once 
 
+#include "Chaos/SoftsSimulationSpace.h"
+#include "ChaosCloth/ChaosClothingSimulationSolver.h"
 #include "ChaosClothAsset/SimulationBaseConfigNode.h"
 #include "ChaosClothAsset/WeightedValue.h"
-#include "Chaos/SoftsSimulationSpace.h"
 #include "SimulationVelocityScaleConfigNode.generated.h"
 
 /** Velocity scale properties configuration node. */
@@ -97,19 +98,34 @@ public:
 	UPROPERTY(EditAnywhere, Category = "Velocity Scale Properties", meta = (UIMin = "0", UIMax = "1", ClampMin = "0", ClampMax = "100", InteractorName = "MaxVelocityScale"))
 	float MaxVelocityScale = 1.f;
 
+	/**
+	 * Fictitious forces model.
+	 * Decide how fictitious forces (e.g. centrifugal force) are applied to the simulation.
+	 */
+	UPROPERTY(EditAnywhere, Category = "Velocity Scale Properties", meta = (InteractorName = "FictitiousForcesModel"))
+	EChaosClothingSimulationSolverFictitiousForcesModel FictitiousForcesModel = EChaosClothingSimulationSolverFictitiousForcesModel::FullScale;
+
 	/**
 	 * The portion of the angular velocity that is used to calculate the strength of all fictitious forces (e.g. centrifugal force).
-	 * This parameter is only having an effect on the portion of the reference bone's angular velocity that has been removed from the
-	 * simulation via the Angular Velocity Scale parameter. This means it has no effect when AngularVelocityScale is set to 1 and 
-	 * Angular Velocity and Acceleration clamps are disabled, in which case the cloth is simulated with full world space angular 
-	 * velocities and subjected to the true physical world inertial forces.
+	 * In Legacy model, this parameter is only having an effect on the portion of the reference bone's angular velocity that has been
+	 * removed from the simulation via the Angular Velocity Scale parameter. This means it has no effect when AngularVelocityScale is
+	 * set to 1 and Angular Velocity and Acceleration clamps are disabled, in which case the cloth is simulated with full world space
+	 * angular velocities and subjected to the true physical world inertial forces.
+	 * In FullScale model, this parameter has also an effect when the cloth is simulated with full world space angular velocities, and
+	 * isn't affected by the Angular Velocity Scale parameter. 
 	 * Values range from 0 to 2, with 0 showing no centrifugal effect, 1 full centrifugal effect, and 2 an overdriven centrifugal effect.
 	 */
-	UPROPERTY(EditAnywhere, Category = "Velocity Scale Properties", meta = (UIMin = "0", UIMax = "1", ClampMin = "0", ClampMax = "2", InteractorName = "FictitiousAngularScale"))
+	UPROPERTY(EditAnywhere, Category = "Velocity Scale Properties", meta = (UIMin = "0", UIMax = "1", ClampMin = "0", ClampMax = "2", EditCondition = "FictitiousForcesModel != EChaosClothingSimulationSolverFictitiousForcesModel::None", InteractorName = "FictitiousAngularScale"))
 	float FictitiousAngularScale = 1.f;
 
 	FChaosClothAssetSimulationVelocityScaleConfigNode(const UE::Dataflow::FNodeParameters& InParam, FGuid InGuid = FGuid::NewGuid());
 
 private:
+	//~ Begin FChaosClothAssetSimulationBaseConfigNode interface
 	virtual void AddProperties(FPropertyHelper& PropertyHelper) const override;
+	//~ End FChaosClothAssetSimulationBaseConfigNode interface
+
+	//~ Begin FDataflowNode interface
+	virtual void Serialize(FArchive& Ar) override;
+	//~ End FDataflowNode interface
 };
diff --git a/Engine/Source/Runtime/Core/Public/UObject/FortniteReleaseBranchCustomObjectVersions.inl b/Engine/Source/Runtime/Core/Public/UObject/FortniteReleaseBranchCustomObjectVersions.inl
--- a/Engine/Source/Runtime/Core/Public/UObject/FortniteReleaseBranchCustomObjectVersions.inl
+++ b/Engine/Source/Runtime/Core/Public/UObject/FortniteReleaseBranchCustomObjectVersions.inl
@@ -60,5 +60,8 @@ FORTNITE_RELEASE_VERSION(OptionSanitizeOutputAttributeNamesPCG, 16)
 // Add automatic platform naming fix up for CommonUI input action data tables
 FORTNITE_RELEASE_VERSION(CommonUIPlatformNamingUpgradeOption, 17)
 
+// Add new fictitious forces modes to the cloth simulation
+FORTNITE_RELEASE_VERSION(ChaosClothAddNewFictitiousForcesModes, 18)
+
 // -----<new versions can be added above this line>-------------------------------------------------
 #undef FORTNITE_RELEASE_VERSION
\ No newline at end of file
