Chaos Cloth - Added a new AppliedScale fictitious force mode to use for clamped values, and made thje exisitng FullScale use full velocities.

#rnx
#Synced-CL 48548117

diff --git a/Engine/Plugins/ChaosCloth/Source/ChaosCloth/Private/ChaosCloth/ChaosClothingSimulationSolver.cpp b/Engine/Plugins/ChaosCloth/Source/ChaosCloth/Private/ChaosCloth/ChaosClothingSimulationSolver.cpp
--- a/Engine/Plugins/ChaosCloth/Source/ChaosCloth/Private/ChaosCloth/ChaosClothingSimulationSolver.cpp
+++ b/Engine/Plugins/ChaosCloth/Source/ChaosCloth/Private/ChaosCloth/ChaosClothingSimulationSolver.cpp
@@ -1316,19 +1316,33 @@ void FClothingSimulationSolver::SetReferenceVelocityScale(uint32 GroupId,
 	// Fictitious angular scale only applied for PBDEvolution. It's applied by ExternalForces for Evolution.
 	const FReal AppliedFictitiousAngularScale = (FictitiousForcesModel == EFictitiousForcesModel::None) ? (FReal)0.f : (PBDEvolution ? FMath::Min((FReal)2., (FReal)FictitiousAngularScale) : (FReal)1.);
 
-	// Save the reference bone relative angular velocity for calculating the fictitious forces
-	const FVec3 FictitiousAngularDisplacement = (FictitiousForcesModel == EFictitiousForcesModel::None) ? FVec3(0) :
-		(FictitiousForcesModel == EFictitiousForcesModel::Legacy) ?
-			ReferenceSpaceTransform.TransformVector(DeltaRotation) * AppliedFictitiousAngularScale :  // Legacy
-			ReferenceSpaceTransform.TransformVector(AppliedDeltaRotation)* AppliedFictitiousAngularScale;  // Full scale
+	FVec3 FictitiousAngularDisplacement;
+	switch (FictitiousForcesModel)
+	{
+	default:
+	case EFictitiousForcesModel::None:
+		FictitiousAngularDisplacement = FVec3(0);
+		break;
+	case EFictitiousForcesModel::Legacy:
+		FictitiousAngularDisplacement = ReferenceSpaceTransform.TransformVector(DeltaRotation) * AppliedFictitiousAngularScale;
+		break;
+	case EFictitiousForcesModel::AppliedScale:
+		FictitiousAngularDisplacement = ReferenceSpaceTransform.TransformVector(AppliedDeltaRotation) * AppliedFictitiousAngularScale;
+		break;
+	case EFictitiousForcesModel::FullScale:
+		FictitiousAngularDisplacement = ReferenceSpaceTransform.TransformVector(FullDeltaRotation) * AppliedFictitiousAngularScale;
+		break;
+	}
 
-	FictitiousAngularVelocities[GroupId] = DeltaTime > (Softs::FSolverReal)0.f ? Softs::FSolverVec3(FictitiousAngularDisplacement) / DeltaTime : Softs::FSolverVec3(0.f);
+	// Save the reference bone relative angular velocity for calculating the fictitious forces
+	const bool bHasNonZeroDeltaTime = (DeltaTime > (Softs::FSolverReal)0.f);
+	FictitiousAngularVelocities[GroupId] = bHasNonZeroDeltaTime ? Softs::FSolverVec3(FictitiousAngularDisplacement) / DeltaTime : Softs::FSolverVec3(0.f);
 	ReferenceSpaceLocations[GroupId] = LocalSpaceScaleInv * (ReferenceSpaceTransform.GetLocation() - LocalSpaceLocation);
-	AppliedReferenceSpaceVelocities[GroupId] = DeltaTime > (Softs::FSolverReal)0.f ? FSolverVec3(LocalSpaceScaleInv * InOutAppliedReferenceVelocity) : Softs::FSolverVec3(0.f);
-	AppliedReferenceSpaceAngularVelocities[GroupId] = DeltaTime > (Softs::FSolverReal)0.f ? FSolverVec3(InOutAppliedReferenceAngularVelocity) : Softs::FSolverVec3(0.f);
+	AppliedReferenceSpaceVelocities[GroupId] = bHasNonZeroDeltaTime ? FSolverVec3(LocalSpaceScaleInv * InOutAppliedReferenceVelocity) : Softs::FSolverVec3(0.f);
+	AppliedReferenceSpaceAngularVelocities[GroupId] = bHasNonZeroDeltaTime ? FSolverVec3(InOutAppliedReferenceAngularVelocity) : Softs::FSolverVec3(0.f);
 #if CHAOS_DEBUG_DRAW
-	OriginalReferenceSpaceVelocities[GroupId] = DeltaTime > (Softs::FSolverReal)0.f ? FSolverVec3(LocalSpaceScaleInv * WorldDeltaTranslation / DeltaTime) : Softs::FSolverVec3(0.f);
-	OriginalReferenceSpaceAngularVelocities[GroupId] = DeltaTime > (Softs::FSolverReal)0.f ? FSolverVec3(FullDeltaRotation / DeltaTime) : Softs::FSolverVec3(0.f);
+	OriginalReferenceSpaceVelocities[GroupId] = bHasNonZeroDeltaTime ? FSolverVec3(LocalSpaceScaleInv * WorldDeltaTranslation / DeltaTime) : Softs::FSolverVec3(0.f);
+	OriginalReferenceSpaceAngularVelocities[GroupId] = bHasNonZeroDeltaTime ? FSolverVec3(FullDeltaRotation / DeltaTime) : Softs::FSolverVec3(0.f);
 #endif
 }
 
diff --git a/Engine/Plugins/ChaosCloth/Source/ChaosCloth/Public/ChaosCloth/ChaosClothingSimulationSolver.h b/Engine/Plugins/ChaosCloth/Source/ChaosCloth/Public/ChaosCloth/ChaosClothingSimulationSolver.h
--- a/Engine/Plugins/ChaosCloth/Source/ChaosCloth/Public/ChaosCloth/ChaosClothingSimulationSolver.h
+++ b/Engine/Plugins/ChaosCloth/Source/ChaosCloth/Public/ChaosCloth/ChaosClothingSimulationSolver.h
@@ -23,9 +23,14 @@ struct FClothingSimulationCacheData;
 UENUM()
 enum class EChaosClothingSimulationSolverFictitiousForcesModel : uint8
 {
-	None,  // Don't apply fictitious forces.
-	Legacy,  // Use partial scaled angular velocity, fictitious forces are calculated from the portion of the scaled character velocities. Legacy behavio, use the FullScale model instead.
-	FullScale,  // Use full angular velocity, fictitious forces are calculated from the full scale character velocities.
+	// Don't apply fictitious forces.
+	None,
+	// Use partial scaled angular velocity, fictitious forces are calculated from the portion of the scaled character velocities. Legacy behavior, use the FullScale model instead.
+	Legacy,
+	// Use applied angular velocity (velocities after clamping), fictitious forces are calculated from the applied scale character velocities.
+	AppliedScale,
+	// Use full angular velocity, fictitious forces are calculated from the full scale character velocities.
+	FullScale,
 };
 
 namespace Chaos
