Chaos Cloth Asset - Added bone names and axis lengths to reference bone visualization.

[RN] Chaos Cloth Asset - Added bone names and axis lengths to reference bone visualization.
#Synced-CL 48653274

diff --git a/Engine/Plugins/ChaosCloth/Source/ChaosCloth/Private/ChaosCloth/ChaosClothVisualization.cpp b/Engine/Plugins/ChaosCloth/Source/ChaosCloth/Private/ChaosCloth/ChaosClothVisualization.cpp
--- a/Engine/Plugins/ChaosCloth/Source/ChaosCloth/Private/ChaosCloth/ChaosClothVisualization.cpp
+++ b/Engine/Plugins/ChaosCloth/Source/ChaosCloth/Private/ChaosCloth/ChaosClothVisualization.cpp
@@ -51,7 +51,7 @@
 #include "UObject/ICookInfo.h"
 #endif  // #if WITH_EDITOR
 
-
+#define LOCTEXT_NAMESPACE "ChaosClothVisualization"
 
 namespace Chaos
 {
@@ -542,6 +542,30 @@ namespace Chaos
 		}
 	}
 
+	void FClothVisualizationNoGC::DrawLocalSpaceBoneNames(FCanvas* Canvas, const FSceneView* SceneView) const
+	{
+		if (!Solver)
+		{
+			return;
+		}
+
+		// Draw local space
+		DrawText(Canvas, SceneView, Solver->GetLocalSpaceLocation() + FVec3(2.), LOCTEXT("SimulationSpace", "Simulation Space"), FLinearColor::White, 1.5f);
+
+		// Draw reference spaces
+		for (const FClothingSimulationCloth* const Cloth : Solver->GetCloths())
+		{
+			const int32 ParticleRangeId = Cloth->GetParticleRangeId(Solver);
+			if (ParticleRangeId == INDEX_NONE)
+			{
+				continue;
+			}
+			const FRigidTransform3& ReferenceSpaceTransform = Cloth->GetReferenceSpaceTransform();
+			check(Cloth->GetMesh());
+			DrawText(Canvas, SceneView, ReferenceSpaceTransform.GetLocation() + FVec3(2.), FText::FromName(Cloth->GetMesh()->GetReferenceBoneName()), FLinearColor::White, 1.5f);
+		}
+	}
+
 	static void DrawPoint(FPrimitiveDrawInterface* PDI, const FVector& Pos, const FLinearColor& Color, const UMaterial* ClothMaterialVertex, const float Thickness = 1.f)  // Use color or material
 	{
 		if (!PDI)
@@ -2833,7 +2857,7 @@ namespace Chaos
 		}
 	}
 
-	void FClothVisualizationNoGC::DrawLocalSpace(FPrimitiveDrawInterface* PDI) const
+	void FClothVisualizationNoGC::DrawLocalSpace(FPrimitiveDrawInterface* PDI, const FReal LocalSpaceLengthScale) const
 	{
 		if (!Solver)
 		{
@@ -2841,8 +2865,9 @@ namespace Chaos
 		}
 
 		// Draw local space
-		DrawCoordinateSystem(PDI, FQuat::Identity, Solver->GetLocalSpaceLocation(), 5.f * Solver->GetLocalSpaceScale(), 0.25f);
-		DrawCoordinateSystem(PDI, Solver->GetLocalSpaceRotation(), Solver->GetLocalSpaceLocation(), 7.f * Solver->GetLocalSpaceScale(), 0.5f);
+		DrawCoordinateSystem(PDI, FQuat::Identity, Solver->GetLocalSpaceLocation(), 5.f * Solver->GetLocalSpaceScale() * LocalSpaceLengthScale, 0.25f);
+		DrawCoordinateSystem(PDI, Solver->GetLocalSpaceRotation(), Solver->GetLocalSpaceLocation(), 7.f * Solver->GetLocalSpaceScale() * LocalSpaceLengthScale, 0.5f);
+		DrawText(nullptr, nullptr, Solver->GetLocalSpaceLocation() + FVec3(2.), LOCTEXT("SimulationSpace", "Simulation Space"), FLinearColor::White, 1.5f);
 
 		// Draw reference spaces
 		for (const FClothingSimulationCloth* const Cloth : Solver->GetCloths())
@@ -2853,9 +2878,9 @@ namespace Chaos
 				continue;
 			}
 			const FRigidTransform3& ReferenceSpaceTransform = Cloth->GetReferenceSpaceTransform();
-			DrawCoordinateSystem(PDI, ReferenceSpaceTransform.GetRotation(), ReferenceSpaceTransform.GetLocation(), 10.f * Solver->GetLocalSpaceScale());
+			DrawCoordinateSystem(PDI, ReferenceSpaceTransform.GetRotation(), ReferenceSpaceTransform.GetLocation(), 10.f * Solver->GetLocalSpaceScale() * LocalSpaceLengthScale);
 			check(Cloth->GetMesh());
-			DrawText(nullptr, nullptr, ReferenceSpaceTransform.GetLocation() + FVec3(2.), FText::FromName(Cloth->GetMesh()->GetReferenceBoneName()), FLinearColor::White, 2.f);
+			DrawText(nullptr, nullptr, ReferenceSpaceTransform.GetLocation() + FVec3(2.), FText::FromName(Cloth->GetMesh()->GetReferenceBoneName()), FLinearColor::White, 1.5f);
 		}
 	}
 
@@ -3699,3 +3724,5 @@ namespace Chaos
 	FClothVisualizationNoGC::~FClothVisualizationNoGC() = default;
 #endif  // #else #if CHAOS_DEBUG_DRAW
 }  // End namespace Chaos
+
+#undef LOCTEXT_NAMESPACE
diff --git a/Engine/Plugins/ChaosCloth/Source/ChaosCloth/Public/ChaosCloth/ChaosClothVisualization.h b/Engine/Plugins/ChaosCloth/Source/ChaosCloth/Public/ChaosCloth/ChaosClothVisualization.h
--- a/Engine/Plugins/ChaosCloth/Source/ChaosCloth/Public/ChaosCloth/ChaosClothVisualization.h
+++ b/Engine/Plugins/ChaosCloth/Source/ChaosCloth/Public/ChaosCloth/ChaosClothVisualization.h
@@ -33,6 +33,7 @@ namespace Chaos
 		CHAOSCLOTH_API void DrawParticleIndices(FCanvas* Canvas = nullptr, const FSceneView* SceneView = nullptr) const;
 		CHAOSCLOTH_API void DrawElementIndices(FCanvas* Canvas = nullptr, const FSceneView* SceneView = nullptr) const;
 		CHAOSCLOTH_API void DrawMaxDistanceValues(FCanvas* Canvas = nullptr, const FSceneView* SceneView = nullptr) const;
+		CHAOSCLOTH_API void DrawLocalSpaceBoneNames(FCanvas* Canvas = nullptr, const FSceneView* SceneView = nullptr) const;
 
 		CHAOSCLOTH_API void DrawPhysMeshWired(FPrimitiveDrawInterface* PDI = nullptr) const;
 		CHAOSCLOTH_API void DrawAnimMeshWired(FPrimitiveDrawInterface* PDI = nullptr) const;
@@ -51,7 +52,7 @@ namespace Chaos
 		CHAOSCLOTH_API void DrawLongRangeConstraint(FPrimitiveDrawInterface* PDI = nullptr) const;
 		CHAOSCLOTH_API void DrawWindAndPressureForces(FPrimitiveDrawInterface* PDI, const FReal ForceLengthScale) const;
 		CHAOSCLOTH_API void DrawWindVelocity(FPrimitiveDrawInterface* PDI, const FReal LengthScale) const;
-		CHAOSCLOTH_API void DrawLocalSpace(FPrimitiveDrawInterface* PDI = nullptr) const;
+		CHAOSCLOTH_API void DrawLocalSpace(FPrimitiveDrawInterface* PDI, FReal LocalSpaceLengthScale) const;
 		CHAOSCLOTH_API void DrawVelocityScale(FPrimitiveDrawInterface* PDI, const FReal LengthScale) const;
 		CHAOSCLOTH_API void DrawSelfCollision(FPrimitiveDrawInterface* PDI = nullptr) const;
 		CHAOSCLOTH_API void DrawSelfIntersection(FPrimitiveDrawInterface* PDI = nullptr) const;
@@ -73,6 +74,7 @@ namespace Chaos
 		void DrawParticleIndices(FCanvas* /*Canvas*/, const FSceneView* /*SceneView*/) const {}
 		void DrawElementIndices(FCanvas* /*Canvas*/, const FSceneView* /*SceneView*/) const {}
 		void DrawMaxDistanceValues(FCanvas* /*Canvas*/, const FSceneView* /*SceneView*/) const {}
+		void DrawLocalSpaceBoneNames(FCanvas* /*Canvas*/, const FSceneView* /*SceneView*/) const {}
 		void DrawPhysMeshWired(FPrimitiveDrawInterface* /*PDI*/ = nullptr) const {}
 		void DrawAnimMeshWired(FPrimitiveDrawInterface* /*PDI*/ = nullptr) const {}
 		void DrawAnimNormals(FPrimitiveDrawInterface* /*PDI*/, const FReal /*Length*/) const {}
@@ -89,7 +91,7 @@ namespace Chaos
 		void DrawLongRangeConstraint(FPrimitiveDrawInterface* /*PDI*/ = nullptr) const {}
 		void DrawWindAndPressureForces(FPrimitiveDrawInterface* /*PDI*/, const FReal /*ForceLengthScale*/) const {}
 		void DrawWindVelocity(FPrimitiveDrawInterface* /*PDI*/, const FReal /*LengthScale*/) const {}
-		void DrawLocalSpace(FPrimitiveDrawInterface* /*PDI*/ = nullptr) const {}
+		void DrawLocalSpace(FPrimitiveDrawInterface* /*PDI*/, FReal /*LocalSpaceLengthScale*/) const {}
 		void DrawVelocityScale(FPrimitiveDrawInterface* /*PDI*/, const FReal /*LengthScale*/) const {}
 		void DrawSelfCollision(FPrimitiveDrawInterface* /*PDI*/ = nullptr) const {}
 		void DrawSelfIntersection(FPrimitiveDrawInterface* /*PDI*/ = nullptr) const {}
@@ -135,6 +137,11 @@ namespace Chaos
 			constexpr FReal DefaultLengthScale = .1;
 			DrawVelocityScale(PDI, DefaultLengthScale);
 		}
+		void DrawLocalSpace(FPrimitiveDrawInterface* PDI = nullptr) const
+		{
+			constexpr FReal DefaultLengthScale = 1.;
+			DrawLocalSpace(PDI, DefaultLengthScale);
+		}
 
 #if WITH_EDITOR && CHAOS_DEBUG_DRAW
 		// Editor only functions
diff --git a/Engine/Plugins/ChaosCloth/Source/ChaosCloth/Public/ChaosCloth/ChaosClothingSimulation.h b/Engine/Plugins/ChaosCloth/Source/ChaosCloth/Public/ChaosCloth/ChaosClothingSimulation.h
--- a/Engine/Plugins/ChaosCloth/Source/ChaosCloth/Public/ChaosCloth/ChaosClothingSimulation.h
+++ b/Engine/Plugins/ChaosCloth/Source/ChaosCloth/Public/ChaosCloth/ChaosClothingSimulation.h
@@ -100,6 +100,7 @@ PRAGMA_ENABLE_DEPRECATION_WARNINGS
 		void DebugDrawParticleIndices(FCanvas* Canvas = nullptr, const FSceneView* SceneView = nullptr) const { Visualization.DrawParticleIndices(Canvas, SceneView); }
 		void DebugDrawElementIndices(FCanvas* Canvas = nullptr, const FSceneView* SceneView = nullptr) const { Visualization.DrawElementIndices(Canvas, SceneView); }
 		void DebugDrawMaxDistanceValues(FCanvas* Canvas = nullptr, const FSceneView* SceneView = nullptr) const { Visualization.DrawMaxDistanceValues(Canvas, SceneView); }
+		void DebugDrawLocalSpaceBoneNames(FCanvas* Canvas = nullptr, const FSceneView* SceneView = nullptr) const { Visualization.DrawLocalSpaceBoneNames(Canvas, SceneView); }
 		void DebugDrawPhysMeshWired(FPrimitiveDrawInterface* PDI = nullptr) const { Visualization.DrawPhysMeshWired(PDI); }
 		void DebugDrawAnimMeshWired(FPrimitiveDrawInterface* PDI = nullptr) const { Visualization.DrawAnimMeshWired(PDI); }
 		void DebugDrawAnimNormals(FPrimitiveDrawInterface* PDI = nullptr) const { Visualization.DrawAnimNormals(PDI); }
diff --git a/Engine/Plugins/ChaosCloth/Source/ChaosClothEditor/Private/ChaosClothEditor/ChaosSimulationEditorExtender.cpp b/Engine/Plugins/ChaosCloth/Source/ChaosClothEditor/Private/ChaosClothEditor/ChaosSimulationEditorExtender.cpp
--- a/Engine/Plugins/ChaosCloth/Source/ChaosClothEditor/Private/ChaosClothEditor/ChaosSimulationEditorExtender.cpp
+++ b/Engine/Plugins/ChaosCloth/Source/ChaosClothEditor/Private/ChaosClothEditor/ChaosSimulationEditorExtender.cpp
@@ -65,29 +65,130 @@ struct FVisualizationOption
 
 const FVisualizationOption FVisualizationOption::OptionData[] = 
 {
-	FVisualizationOption(&FClothVisualizationNoGC::DrawPhysMeshShaded      , TEXT("p.ChaosClothEditor.DebugDrawPhysMeshShaded"      ), LOCTEXT("ChaosVisName_PhysMesh"            , "Physical Mesh (Flat Shaded)"), LOCTEXT("ChaosVisName_PhysMeshShaded_ToolTip"      , "Draws the current physical result as a doubled sided flat shaded mesh"), /*bDisablesSimulation =*/false, /*bHidesClothSections=*/true),
-	FVisualizationOption(&FClothVisualizationNoGC::DrawPhysMeshWired       , TEXT("p.ChaosClothEditor.DebugDrawPhysMeshWired"       ), LOCTEXT("ChaosVisName_PhysMeshWire"        , "Physical Mesh (Wireframe)"  ), LOCTEXT("ChaosVisName_PhysMeshWired_ToolTip"       , "Draws the current physical mesh result in wireframe")),
-	FVisualizationOption(&FClothVisualizationNoGC::DrawAnimMeshWired       , TEXT("p.ChaosClothEditor.DebugDrawAnimMeshWired"       ), LOCTEXT("ChaosVisName_AnimMeshWire"        , "Animated Mesh (Wireframe)"  ), LOCTEXT("ChaosVisName_AnimMeshWired_ToolTip"       , "Draws the current animated mesh input in wireframe")),
-	FVisualizationOption(&FClothVisualizationNoGC::DrawParticleIndices     , TEXT("p.ChaosClothEditor.DebugDrawParticleIndices"     ), LOCTEXT("ChaosVisName_ParticleIndices"     , "Particle Indices"           ), LOCTEXT("ChaosVisName_ParticleIndices_ToolTip"     , "Draws the particle indices as instantiated by the solver")),
-	FVisualizationOption(&FClothVisualizationNoGC::DrawElementIndices      , TEXT("p.ChaosClothEditor.DebugDrawElementIndices"      ), LOCTEXT("ChaosVisName_ElementIndices"      , "Element Indices"            ), LOCTEXT("ChaosVisName_ElementIndices_ToolTip"      , "Draws the element's (triangle or other) indices as instantiated by the solver")),
-	FVisualizationOption(&FClothVisualizationNoGC::DrawPointNormals        , TEXT("p.ChaosClothEditor.DebugDrawPointNormals"        ), LOCTEXT("ChaosVisName_PointNormals"        , "Physical Mesh Normals"      ), LOCTEXT("ChaosVisName_PointNormals_ToolTip"        , "Draws the current point normals for the simulation mesh")),
-	FVisualizationOption(&FClothVisualizationNoGC::DrawPointVelocities     , TEXT("p.ChaosClothEditor.DebugDrawPointVelocities"     ), LOCTEXT("ChaosVisName_PointVelocities"     , "Point Velocities"           ), LOCTEXT("ChaosVisName_PointVelocities_ToolTip"     , "Draws the current point velocities for the simulation mesh")),
-	FVisualizationOption(&FClothVisualizationNoGC::DrawAnimNormals         , TEXT("p.ChaosClothEditor.DebugDrawAnimNormals"         ), LOCTEXT("ChaosVisName_AnimNormals"         , "Animated Mesh Normals"      ), LOCTEXT("ChaosVisName_AnimNormals_ToolTip"         , "Draws the current point normals for the animated mesh")),
-	FVisualizationOption(&FClothVisualizationNoGC::DrawCollision           , TEXT("p.ChaosClothEditor.DebugDrawCollision"           ), LOCTEXT("ChaosVisName_Collision"           , "Collisions"                 ), LOCTEXT("ChaosVisName_Collision_ToolTip"           , "Draws the collision bodies the simulation is currently using")),
-	FVisualizationOption(&FClothVisualizationNoGC::DrawBackstops           , TEXT("p.ChaosClothEditor.DebugDrawBackstops"           ), LOCTEXT("ChaosVisName_Backstop"            , "Backstops"                  ), LOCTEXT("ChaosVisName_Backstop_ToolTip"            , "Draws the backstop radius and position for each simulation particle")),
-	FVisualizationOption(&FClothVisualizationNoGC::DrawBackstopDistances   , TEXT("p.ChaosClothEditor.DebugDrawBackstopDistances"   ), LOCTEXT("ChaosVisName_BackstopDistance"    , "Backstop Distances"         ), LOCTEXT("ChaosVisName_BackstopDistance_ToolTip"    , "Draws the backstop distance offset for each simulation particle"), /*bDisablesSimulation =*/true),
-	FVisualizationOption(&FClothVisualizationNoGC::DrawMaxDistances        , TEXT("p.ChaosClothEditor.DebugDrawMaxDistances"        ), LOCTEXT("ChaosVisName_MaxDistance"         , "Max Distances"              ), LOCTEXT("ChaosVisName_MaxDistance_ToolTip"         , "Draws the current max distances for the sim particles as a line along its normal"), /*bDisablesSimulation =*/true),
-	FVisualizationOption(&FClothVisualizationNoGC::DrawMaxDistanceValues   , TEXT("p.ChaosClothEditor.DebugDrawMaxDistanceValues"   ), LOCTEXT("ChaosVisName_MaxDistanceValue"    , "Max Distances As Numbers"   ), LOCTEXT("ChaosVisName_MaxDistanceValue_ToolTip"    , "Draws the current max distances as numbers")),
-	FVisualizationOption(&FClothVisualizationNoGC::DrawAnimDrive           , TEXT("p.ChaosClothEditor.DebugDrawAnimDrive"           ), LOCTEXT("ChaosVisName_AnimDrive"           , "Anim Drive"                 ), LOCTEXT("ChaosVisName_AnimDrive_Tooltip"           , "Draws the current skinned reference mesh for the simulation which anim drive will attempt to reach if enabled")),
-	FVisualizationOption(&FClothVisualizationNoGC::DrawEdgeConstraint      , TEXT("p.ChaosClothEditor.DebugDrawEdgeConstraint"      ), LOCTEXT("ChaosVisName_EdgeConstraint"      , "Edge Constraint"            ), LOCTEXT("ChaosVisName_EdgeConstraint_Tooltip"      , "Draws the edge spring constraints")),
-	FVisualizationOption(&FClothVisualizationNoGC::DrawBendingConstraint   , TEXT("p.ChaosClothEditor.DebugDrawBendingConstraint"   ), LOCTEXT("ChaosVisName_BendingConstraint"   , "Bending Constraint"         ), LOCTEXT("ChaosVisName_BendingConstraint_Tooltip"   , "Draws the bending spring constraints")),
-	FVisualizationOption(&FClothVisualizationNoGC::DrawLongRangeConstraint , TEXT("p.ChaosClothEditor.DebugDrawLongRangeConstraint" ), LOCTEXT("ChaosVisName_LongRangeConstraint" , "Long Range Constraint"      ), LOCTEXT("ChaosVisName_LongRangeConstraint_Tooltip" , "Draws the long range attachment constraint distances")),
-	FVisualizationOption(&FClothVisualizationNoGC::DrawWindAndPressureForces, TEXT("p.ChaosClothEditor.DebugDrawWindAndPressureForces"), LOCTEXT("ChaosVisName_WindAndPressureForces", "Wind Aerodynamic And Pressure Forces"), LOCTEXT("ChaosVisName_WindAndPressure_Tooltip", "Draws the Wind drag and lift and pressure forces")),
-	FVisualizationOption(&FClothVisualizationNoGC::DrawWindVelocity, TEXT("p.ChaosClothEditor.DebugDrawWindVelocity"), LOCTEXT("ChaosVisName_WindVelocity", "Wind Velocity"), LOCTEXT("ChaosVisName_WindVelocity_Tooltip", "Draws the Wind Velocity vector")),
-	FVisualizationOption(&FClothVisualizationNoGC::DrawLocalSpace          , TEXT("p.ChaosClothEditor.DebugDrawLocalSpace"          ), LOCTEXT("ChaosVisName_LocalSpace"          , "Local Space Reference Bone" ), LOCTEXT("ChaosVisName_LocalSpace_Tooltip"          , "Draws the local space reference bone")),
-	FVisualizationOption(&FClothVisualizationNoGC::DrawVelocityScale       , TEXT("p.ChaosClothEditor.DebugDrawVelocityScale"       ), LOCTEXT("ChaosVisName_VelocityScale"       , "Velocity Scale"             ), LOCTEXT("ChaosVisName_VelocityScale_Tooltip"       , "Draws the velocity scale")),
-	FVisualizationOption(&FClothVisualizationNoGC::DrawSelfCollision       , TEXT("p.ChaosClothEditor.DebugDrawSelfCollision"       ), LOCTEXT("ChaosVisName_SelfCollision"       , "Self Collision"             ), LOCTEXT("ChaosVisName_SelfCollision_Tooltip"       , "Draws the self collision thickness/debugging information")),
-	FVisualizationOption(&FClothVisualizationNoGC::DrawSelfIntersection    , TEXT("p.ChaosClothEditor.DebugDrawSelfIntersection"    ), LOCTEXT("ChaosVisName_SelfIntersection"    , "Self Intersection"          ), LOCTEXT("ChaosVisName_SelfIntersection_Tooltip"    , "Draws the self intersection contour/region information")),
+	FVisualizationOption(
+		&FClothVisualizationNoGC::DrawPhysMeshShaded,
+		TEXT("p.ChaosClothEditor.DebugDrawPhysMeshShaded"),
+		LOCTEXT("ChaosVisName_PhysMesh", "Physical Mesh (Flat Shaded)"),
+		LOCTEXT("ChaosVisName_PhysMeshShaded_ToolTip", "Draws the current physical result as a doubled sided flat shaded mesh"),
+		/*bDisablesSimulation =*/false,
+		/*bHidesClothSections=*/true),
+	FVisualizationOption(
+		&FClothVisualizationNoGC::DrawPhysMeshWired,
+		TEXT("p.ChaosClothEditor.DebugDrawPhysMeshWired"),
+		LOCTEXT("ChaosVisName_PhysMeshWire", "Physical Mesh (Wireframe)"),
+		LOCTEXT("ChaosVisName_PhysMeshWired_ToolTip", "Draws the current physical mesh result in wireframe")),
+	FVisualizationOption(
+		&FClothVisualizationNoGC::DrawAnimMeshWired,
+		TEXT("p.ChaosClothEditor.DebugDrawAnimMeshWired"),
+		LOCTEXT("ChaosVisName_AnimMeshWire", "Animated Mesh (Wireframe)"),
+		LOCTEXT("ChaosVisName_AnimMeshWired_ToolTip", "Draws the current animated mesh input in wireframe")),
+	FVisualizationOption(
+		&FClothVisualizationNoGC::DrawParticleIndices,
+		TEXT("p.ChaosClothEditor.DebugDrawParticleIndices"),
+		LOCTEXT("ChaosVisName_ParticleIndices", "Particle Indices"),
+		LOCTEXT("ChaosVisName_ParticleIndices_ToolTip", "Draws the particle indices as instantiated by the solver")),
+	FVisualizationOption(
+		&FClothVisualizationNoGC::DrawElementIndices,
+		TEXT("p.ChaosClothEditor.DebugDrawElementIndices"),
+		LOCTEXT("ChaosVisName_ElementIndices", "Element Indices"),
+		LOCTEXT("ChaosVisName_ElementIndices_ToolTip", "Draws the element's (triangle or other) indices as instantiated by the solver")),
+	FVisualizationOption(
+		&FClothVisualizationNoGC::DrawPointNormals,
+		TEXT("p.ChaosClothEditor.DebugDrawPointNormals"),
+		LOCTEXT("ChaosVisName_PointNormals", "Physical Mesh Normals"),
+		LOCTEXT("ChaosVisName_PointNormals_ToolTip", "Draws the current point normals for the simulation mesh")),
+	FVisualizationOption(
+		&FClothVisualizationNoGC::DrawPointVelocities,
+		TEXT("p.ChaosClothEditor.DebugDrawPointVelocities"),
+		LOCTEXT("ChaosVisName_PointVelocities", "Point Velocities"),
+		LOCTEXT("ChaosVisName_PointVelocities_ToolTip", "Draws the current point velocities for the simulation mesh")),
+	FVisualizationOption(
+		&FClothVisualizationNoGC::DrawAnimNormals,
+		TEXT("p.ChaosClothEditor.DebugDrawAnimNormals"),
+		LOCTEXT("ChaosVisName_AnimNormals", "Animated Mesh Normals"),
+		LOCTEXT("ChaosVisName_AnimNormals_ToolTip", "Draws the current point normals for the animated mesh")),
+	FVisualizationOption(
+		&FClothVisualizationNoGC::DrawCollision,
+		TEXT("p.ChaosClothEditor.DebugDrawCollision"),
+		LOCTEXT("ChaosVisName_Collision", "Collisions"),
+		LOCTEXT("ChaosVisName_Collision_ToolTip", "Draws the collision bodies the simulation is currently using")),
+	FVisualizationOption(
+		&FClothVisualizationNoGC::DrawBackstops,
+		TEXT("p.ChaosClothEditor.DebugDrawBackstops"),
+		LOCTEXT("ChaosVisName_Backstop", "Backstops"),
+		LOCTEXT("ChaosVisName_Backstop_ToolTip", "Draws the backstop radius and position for each simulation particle")),
+	FVisualizationOption(
+		&FClothVisualizationNoGC::DrawBackstopDistances,
+		TEXT("p.ChaosClothEditor.DebugDrawBackstopDistances"),
+		LOCTEXT("ChaosVisName_BackstopDistance", "Backstop Distances"),
+		LOCTEXT("ChaosVisName_BackstopDistance_ToolTip", "Draws the backstop distance offset for each simulation particle"),
+		/*bDisablesSimulation =*/true),
+	FVisualizationOption(
+		&FClothVisualizationNoGC::DrawMaxDistances,
+		TEXT("p.ChaosClothEditor.DebugDrawMaxDistances"),
+		LOCTEXT("ChaosVisName_MaxDistance", "Max Distances"),
+		LOCTEXT("ChaosVisName_MaxDistance_ToolTip", "Draws the current max distances for the sim particles as a line along its normal"),
+		/*bDisablesSimulation =*/true),
+	FVisualizationOption(
+		&FClothVisualizationNoGC::DrawMaxDistanceValues,
+		TEXT("p.ChaosClothEditor.DebugDrawMaxDistanceValues"),
+		LOCTEXT("ChaosVisName_MaxDistanceValue", "Max Distances As Numbers"),
+		LOCTEXT("ChaosVisName_MaxDistanceValue_ToolTip", "Draws the current max distances as numbers")),
+	FVisualizationOption(
+		&FClothVisualizationNoGC::DrawAnimDrive,
+		TEXT("p.ChaosClothEditor.DebugDrawAnimDrive"),
+		LOCTEXT("ChaosVisName_AnimDrive", "Anim Drive"),
+		LOCTEXT("ChaosVisName_AnimDrive_Tooltip", "Draws the current skinned reference mesh for the simulation which anim drive will attempt to reach if enabled")),
+	FVisualizationOption(
+		&FClothVisualizationNoGC::DrawEdgeConstraint,
+		TEXT("p.ChaosClothEditor.DebugDrawEdgeConstraint"),
+		LOCTEXT("ChaosVisName_EdgeConstraint", "Edge Constraint"),
+		LOCTEXT("ChaosVisName_EdgeConstraint_Tooltip", "Draws the edge spring constraints")),
+	FVisualizationOption(
+		&FClothVisualizationNoGC::DrawBendingConstraint,
+		TEXT("p.ChaosClothEditor.DebugDrawBendingConstraint"),
+		LOCTEXT("ChaosVisName_BendingConstraint", "Bending Constraint"),
+		LOCTEXT("ChaosVisName_BendingConstraint_Tooltip", "Draws the bending spring constraints")),
+	FVisualizationOption(
+		&FClothVisualizationNoGC::DrawLongRangeConstraint,
+		TEXT("p.ChaosClothEditor.DebugDrawLongRangeConstraint"),
+		LOCTEXT("ChaosVisName_LongRangeConstraint", "Long Range Constraint"),
+		LOCTEXT("ChaosVisName_LongRangeConstraint_Tooltip", "Draws the long range attachment constraint distances")),
+	FVisualizationOption(
+		&FClothVisualizationNoGC::DrawWindAndPressureForces,
+		TEXT("p.ChaosClothEditor.DebugDrawWindAndPressureForces"),
+		LOCTEXT("ChaosVisName_WindAndPressureForces", "Wind Aerodynamic And Pressure Forces"),
+		LOCTEXT("ChaosVisName_WindAndPressure_Tooltip", "Draws the Wind drag and lift and pressure forces")),
+	FVisualizationOption(
+		&FClothVisualizationNoGC::DrawWindVelocity,
+		TEXT("p.ChaosClothEditor.DebugDrawWindVelocity"),
+		LOCTEXT("ChaosVisName_WindVelocity", "Wind Velocity"),
+		LOCTEXT("ChaosVisName_WindVelocity_Tooltip", "Draws the Wind Velocity vector")),
+	FVisualizationOption(
+		&FClothVisualizationNoGC::DrawLocalSpace,
+		TEXT("p.ChaosClothEditor.DebugDrawLocalSpace"),
+		LOCTEXT("ChaosVisName_LocalSpace", "Local Space Reference Bone"),
+		LOCTEXT("ChaosVisName_LocalSpace_Tooltip", "Draws the local space reference bone")),
+	FVisualizationOption(
+		&FClothVisualizationNoGC::DrawLocalSpaceBoneNames,
+		TEXT("p.ChaosClothEditor.DebugDrawLocalSpaceBoneNames"),
+		LOCTEXT("ChaosVisName_LocalSpaceBoneNames", "Local Space Reference Bone Names"),
+		LOCTEXT("ChaosVisName_LocalSpaceBoneNames_Tooltip", "Draws the local space reference bone names")),
+	FVisualizationOption(
+		&FClothVisualizationNoGC::DrawVelocityScale,
+		TEXT("p.ChaosClothEditor.DebugDrawVelocityScale"),
+		LOCTEXT("ChaosVisName_VelocityScale", "Velocity Scale"),
+		LOCTEXT("ChaosVisName_VelocityScale_Tooltip", "Draws the velocity scale")),
+	FVisualizationOption(
+		&FClothVisualizationNoGC::DrawSelfCollision,
+		TEXT("p.ChaosClothEditor.DebugDrawSelfCollision"),
+		LOCTEXT("ChaosVisName_SelfCollision", "Self Collision"),
+		LOCTEXT("ChaosVisName_SelfCollision_Tooltip", "Draws the self collision thickness/debugging information")),
+	FVisualizationOption(
+		&FClothVisualizationNoGC::DrawSelfIntersection,
+		TEXT("p.ChaosClothEditor.DebugDrawSelfIntersection"),
+		LOCTEXT("ChaosVisName_SelfIntersection", "Self Intersection"),
+		LOCTEXT("ChaosVisName_SelfIntersection_Tooltip", "Draws the self intersection contour/region information"))
 };
 const uint32 FVisualizationOption::Count = sizeof(OptionData) / sizeof(FVisualizationOption);
 
diff --git a/Engine/Plugins/ChaosClothAssetEditor/Source/ChaosClothAssetEditor/Private/ChaosClothAsset/ClothEditorSimulationVisualization.cpp b/Engine/Plugins/ChaosClothAssetEditor/Source/ChaosClothAssetEditor/Private/ChaosClothAsset/ClothEditorSimulationVisualization.cpp
--- a/Engine/Plugins/ChaosClothAssetEditor/Source/ChaosClothAssetEditor/Private/ChaosClothAsset/ClothEditorSimulationVisualization.cpp
+++ b/Engine/Plugins/ChaosClothAssetEditor/Source/ChaosClothAssetEditor/Private/ChaosClothAsset/ClothEditorSimulationVisualization.cpp
@@ -341,12 +341,24 @@ const FVisualizationOption FVisualizationOption::OptionData[] =
 			EditorVisualization.ExtendViewportShowMenuWindVelocityLength(MenuBuilder);
 		})),
 	FVisualizationOption(
-		FClothVisualizationDebugDraw::CreateLambda([](const FClothEditorSimulationVisualization&, const ::Chaos::FClothVisualizationNoGC& Visualization, FPrimitiveDrawInterface* PDI)
+		FClothVisualizationDebugDraw::CreateLambda([](const FClothEditorSimulationVisualization& EditorVisualization, const ::Chaos::FClothVisualizationNoGC& Visualization, FPrimitiveDrawInterface* PDI)
 		{
-			Visualization.DrawLocalSpace(PDI);
+			Visualization.DrawLocalSpace(PDI, EditorVisualization.GetLocalSpaceLengthScale());
 		}),
 		LOCTEXT("ChaosVisName_LocalSpace", "Local Space Reference Bone"), 
-		LOCTEXT("ChaosVisName_LocalSpace_Tooltip", "Draws the local space reference bone")),
+		LOCTEXT("ChaosVisName_LocalSpace_Tooltip", "Draws the local space reference bone"),
+		/*bDisablesSimulation =*/false, /*bHidesClothSections=*/false, /*bDefaultFlagValue=*/false,
+		FAdditionalMenus::CreateLambda([](FClothEditorSimulationVisualization& EditorVisualization, FMenuBuilder& MenuBuilder)
+		{
+			EditorVisualization.ExtendViewportShowMenuLocalSpaceLengthScale(MenuBuilder);
+		})),
+	FVisualizationOption(
+		FClothVisualizationDebugDrawTexts::CreateLambda([](const FClothEditorSimulationVisualization&, const ::Chaos::FClothVisualizationNoGC& Visualization, FCanvas* Canvas, const FSceneView* SceneView)
+		{
+			Visualization.DrawLocalSpaceBoneNames(Canvas, SceneView);
+		}),
+		LOCTEXT("ChaosVisName_LocalSpaceBoneNames", "Local Space Reference Bone Names"), 
+		LOCTEXT("ChaosVisName_LocalSpaceBoneNames_ToolTip", "Draws the local space reference bone names")),
 	FVisualizationOption(
 		FClothVisualizationDebugDraw::CreateLambda([](const FClothEditorSimulationVisualization& EditorVisualization, const ::Chaos::FClothVisualizationNoGC& Visualization, FPrimitiveDrawInterface* PDI)
 		{
@@ -496,7 +508,7 @@ const FVisualizationOption FVisualizationOption::OptionData[] =
 		}),
 		LOCTEXT("ChaosVisName_DrawInpaintWeightsMatched", "Transfer Skin Weights Node: Matched Vertices"), 
 		LOCTEXT("ChaosVisName_DrawInpaintWeightsMatched_ToolTip", "When transferring weights using the InpaintWeights method, will highlight the vertices for which we copied the weights directly from the source mesh. For all other vertices, the weights were computed automatically."), 
-		/*bDisablesSimulation =*/false, /*bHidesClothSections=*/true),
+		/*bDisablesSimulation =*/false, /*bHidesClothSections=*/true)
 };
 const uint32 FVisualizationOption::Count = sizeof(OptionData) / sizeof(FVisualizationOption);
 
diff --git a/Engine/Plugins/ChaosClothAssetEditor/Source/ChaosClothAssetEditor/Private/ChaosClothAsset/ClothEditorSimulationVisualization.h b/Engine/Plugins/ChaosClothAssetEditor/Source/ChaosClothAssetEditor/Private/ChaosClothAsset/ClothEditorSimulationVisualization.h
--- a/Engine/Plugins/ChaosClothAssetEditor/Source/ChaosClothAssetEditor/Private/ChaosClothAsset/ClothEditorSimulationVisualization.h
+++ b/Engine/Plugins/ChaosClothAssetEditor/Source/ChaosClothAssetEditor/Private/ChaosClothAsset/ClothEditorSimulationVisualization.h
@@ -75,6 +75,10 @@ public:
 	{
 		ExtendViewportShowMenuSpinBox(MenuBuilder, WindVelocityLengthScale, 0.f, FLT_MAX, 0.f, 40.f);
 	}
+	void ExtendViewportShowMenuLocalSpaceLengthScale(FMenuBuilder& MenuBuilder)
+	{
+		ExtendViewportShowMenuSpinBox(MenuBuilder, LocalSpaceLengthScale , 0.f, FLT_MAX, 0.f, 40.f);
+	}
 	void ExtendViewportShowMenuVelocityScaleLength(FMenuBuilder& MenuBuilder)
 	{
 		ExtendViewportShowMenuSpinBox(MenuBuilder, VelocityScaleLength, 0.f, FLT_MAX, 0.f, 40.f);
@@ -85,6 +89,7 @@ public:
 	float GetAerodynamicsLengthScale() const { return AerodynamicsLengthScale; }
 	float GetWindVelocityLengthScale() const { return WindVelocityLengthScale; }
 	float GetVelocityScaleLength() const { return VelocityScaleLength; }
+	float GetLocalSpaceLengthScale() const { return LocalSpaceLengthScale; }
 
 private:
 	struct FNameSelectionData
@@ -116,6 +121,6 @@ private:
 	float AerodynamicsLengthScale = 10.f;
 	float WindVelocityLengthScale = .1f;
 	float VelocityScaleLength = .1f;
-	
+	float LocalSpaceLengthScale = 1.f;
 };
 } // namespace UE::Chaos::ClothAsset
