Cloth visualization : make sure reference bone name debug text does not overlap the simulation space if they are at the same position in space

#Synced-CL 49572843

diff --git a/Engine/Plugins/ChaosCloth/Source/ChaosCloth/Private/ChaosCloth/ChaosClothVisualization.cpp b/Engine/Plugins/ChaosCloth/Source/ChaosCloth/Private/ChaosCloth/ChaosClothVisualization.cpp
--- a/Engine/Plugins/ChaosCloth/Source/ChaosCloth/Private/ChaosCloth/ChaosClothVisualization.cpp
+++ b/Engine/Plugins/ChaosCloth/Source/ChaosCloth/Private/ChaosCloth/ChaosClothVisualization.cpp
@@ -562,7 +562,7 @@ namespace Chaos
 			}
 			const FRigidTransform3& ReferenceSpaceTransform = Cloth->GetReferenceSpaceTransform();
 			check(Cloth->GetMesh());
-			DrawText(Canvas, SceneView, ReferenceSpaceTransform.GetLocation() + FVec3(2.), FText::FromName(Cloth->GetMesh()->GetReferenceBoneName()), FLinearColor::White, 1.5f);
+			DrawText(Canvas, SceneView, ReferenceSpaceTransform.GetLocation() + FVec3(-2.), FText::FromName(Cloth->GetMesh()->GetReferenceBoneName()), FLinearColor::White, 1.5f);
 		}
 	}
 
@@ -2880,7 +2880,7 @@ namespace Chaos
 			const FRigidTransform3& ReferenceSpaceTransform = Cloth->GetReferenceSpaceTransform();
 			DrawCoordinateSystem(PDI, ReferenceSpaceTransform.GetRotation(), ReferenceSpaceTransform.GetLocation(), 10.f * Solver->GetLocalSpaceScale() * LocalSpaceLengthScale);
 			check(Cloth->GetMesh());
-			DrawText(nullptr, nullptr, ReferenceSpaceTransform.GetLocation() + FVec3(2.), FText::FromName(Cloth->GetMesh()->GetReferenceBoneName()), FLinearColor::White, 1.5f);
+			DrawText(nullptr, nullptr, ReferenceSpaceTransform.GetLocation() + FVec3(-2.), FText::FromName(Cloth->GetMesh()->GetReferenceBoneName()), FLinearColor::White, 1.5f);
 		}
 	}
 
