Chaos Cloth - Added a cvar ro reverse the order of the Backstop and AnimDrive constraints.

#rnx

#Synced-CL 50258746

diff --git a/Engine/Plugins/ChaosCloth/Source/ChaosCloth/Private/ChaosCloth/ChaosClothConstraints.cpp b/Engine/Plugins/ChaosCloth/Source/ChaosCloth/Private/ChaosCloth/ChaosClothConstraints.cpp
--- a/Engine/Plugins/ChaosCloth/Source/ChaosCloth/Private/ChaosCloth/ChaosClothConstraints.cpp
+++ b/Engine/Plugins/ChaosCloth/Source/ChaosCloth/Private/ChaosCloth/ChaosClothConstraints.cpp
@@ -56,11 +56,10 @@ bool bDoPlanarCollisionConstraints = false;
 Softs::FSolverReal PlanarCollisionTol = 0.f;
 Softs::FSolverReal PlanarCollisionStiffness = 1e10;
 
-// #todo(dmp): once we decide what we want to do for enabling buoyancy by default or controlling it some way, we can 
-// enable it for shipping
-namespace ChaosClothConstraintsBuoyancyConsoleVariables
+namespace ClothConstraints::Private
 {
-	bool bEnableBuoyancy = false;
+	static bool bEnableBuoyancy = false;  // #todo(dmp): once we decide what we want to do for enabling buoyancy by default or controlling it some way, we can enable it for shipping
+	static bool bSwapBackstopAnimDriveApply = false;  // TODO: Move this to a property
 }
 
 #if !UE_BUILD_SHIPPING
@@ -124,7 +123,9 @@ namespace ChaosClothConstraintsBuoyancyConsoleVariables
 
 	static FAutoConsoleVariableRef CVarClothPlanarCollisionStiffness(TEXT("p.Chaos.Cloth.PlanarCollisionStiffness"), PlanarCollisionStiffness, TEXT("Collision stiffness for XPBD planar constraint [def: 1e10.]"));
 	
-	static FAutoConsoleVariableRef CVarClothEnableBuoyancy(TEXT("p.Chaos.Cloth.EnableBuoyancy"), ChaosClothConstraintsBuoyancyConsoleVariables::bEnableBuoyancy, TEXT("Query water bodies for drag and buoyancy effects with cloth [def: false]"));
+	static FAutoConsoleVariableRef CVarClothEnableBuoyancy(TEXT("p.Chaos.Cloth.EnableBuoyancy"), ClothConstraints::Private::bEnableBuoyancy, TEXT("Query water bodies for drag and buoyancy effects with cloth [def: false]"));
+
+	static FAutoConsoleVariableRef CVarClothSwapBackstopAnimDriveApply(TEXT("p.Chaos.Cloth.SwapBackstopAnimDriveApply"), ClothConstraints::Private::bSwapBackstopAnimDriveApply, TEXT("Swap the order of the Backstop and AnimDrive Constraints so that the Backstop is processed last [def: false]"));
 #endif
 
 class FClothConstraints::FRuleCreator
@@ -835,7 +836,7 @@ void FClothConstraints::AddRules(
 		// External Forces
 		CreateExternalForces(ConfigProperties, WeightMaps);
 
-		if (!ChaosClothConstraintsBuoyancyConsoleVariables::bEnableBuoyancy)
+		if (!ClothConstraints::Private::bEnableBuoyancy)
 		{
 			// Velocity Field
 			CreateVelocityAndPressureField(ConfigProperties, WeightMaps, TriangleMesh);
@@ -1950,7 +1951,7 @@ void FClothConstraints::CreateForceBasedRules(const TSharedPtr<Softs::FMultiResC
 		constexpr bool bPostCollisions = false;
 		RuleCreator.AddPerIterationPBDConstraintRule_Apply<bPostCollisions>(MaximumDistanceConstraints.Get());
 	}
-	if (BackstopConstraints)
+	if (BackstopConstraints && !ClothConstraints::Private::bSwapBackstopAnimDriveApply)
 	{
 		constexpr bool bPostCollisions = false;
 		RuleCreator.AddPerIterationPBDConstraintRule_Apply<bPostCollisions>(BackstopConstraints.Get());
@@ -1961,6 +1962,11 @@ void FClothConstraints::CreateForceBasedRules(const TSharedPtr<Softs::FMultiResC
 		constexpr bool bPostCollisions = false;
 		RuleCreator.AddPerIterationPBDConstraintRule_Apply<bPostCollisions>(AnimDriveConstraints.Get());
 	}
+	if (BackstopConstraints && ClothConstraints::Private::bSwapBackstopAnimDriveApply)
+	{
+		constexpr bool bPostCollisions = false;
+		RuleCreator.AddPerIterationPBDConstraintRule_Apply<bPostCollisions>(BackstopConstraints.Get());
+	}
 	if (ClothVertexSpringConstraints)
 	{
 		constexpr bool bInitParticles = false;
