ChaosCloth: fixed animation position update on the accessory mesh

#Synced-CL 50337989

diff --git a/Engine/Plugins/ChaosCloth/Source/ChaosCloth/Private/ChaosCloth/ChaosClothingAccessoryMeshData.cpp b/Engine/Plugins/ChaosCloth/Source/ChaosCloth/Private/ChaosCloth/ChaosClothingAccessoryMeshData.cpp
--- a/Engine/Plugins/ChaosCloth/Source/ChaosCloth/Private/ChaosCloth/ChaosClothingAccessoryMeshData.cpp
+++ b/Engine/Plugins/ChaosCloth/Source/ChaosCloth/Private/ChaosCloth/ChaosClothingAccessoryMeshData.cpp
@@ -52,6 +52,8 @@ namespace Chaos
 	{
 		if (AnimationPositions.Num() == NumParticles)
 		{
+			Swap(OldAnimationPositions, AnimationPositions);
+			Swap(OldAnimationNormals, AnimationNormals);
 			AccessoryMesh.Update(Solver, TArrayView<Softs::FSolverVec3>(AnimationPositions), TArrayView<Softs::FSolverVec3>(AnimationNormals));
 		}
 	}
