Chaos Cloth: Add options to filter debug draw to a single vertex or named set of vertices.
Updated the Dataflow-based editor visualization UI with these options.
Updated the SKM-based editor and PIE visualization to use cvars to enable these options. Legacy cloth does not have a notion of sets, so only single vertex visualization is possible there.

#Synced-CL 51472478

diff --git a/Engine/Plugins/ChaosCloth/Source/ChaosCloth/Private/ChaosCloth/ChaosClothVisualization.cpp b/Engine/Plugins/ChaosCloth/Source/ChaosCloth/Private/ChaosCloth/ChaosClothVisualization.cpp
--- a/Engine/Plugins/ChaosCloth/Source/ChaosCloth/Private/ChaosCloth/ChaosClothVisualization.cpp
+++ b/Engine/Plugins/ChaosCloth/Source/ChaosCloth/Private/ChaosCloth/ChaosClothVisualization.cpp
@@ -104,9 +104,6 @@ namespace Chaos
 		static bool bDrawInForeground = true;
 		static FAutoConsoleVariableRef CVarClothVizDrawInForeground(TEXT("p.ChaosClothVisualization.DrawInForeground"), bDrawInForeground, TEXT("Draw in foreground when outside the cloth/SKM editor"));
 
-		static int32 BackstopVertexIndex = -1;
-		static FAutoConsoleVariableRef CVarClothVizDrawSingleBackstopVertex(TEXT("p.ChaosClothVisualization.DrawSingleBackstopVertex"), BackstopVertexIndex, TEXT("Draw backstop for a single vertex with a full sphere. Set to -1 to display all vertices with partial spheres."));
-
 		// copied from ClothEditorMode
 		FLinearColor PseudoRandomColor(int32 NumColorRotations)
 		{
@@ -126,6 +123,75 @@ namespace Chaos
 			// SDPG_Foreground does not draw when PIE is paused (its buffer is flushed).
 			return bDrawInForeground ? SDPG_MAX : SDPG_World;
 		}
+
+		static bool ShouldDrawParticle(const FClothVisualizationNoGC::FDrawFilter& Filter, const TSet<int32>* const DebugFilterSet, int32 ParticleIndex, int32 GlobalParticleOffset)
+		{
+			return (Filter.SingleVertex == INDEX_NONE ||
+				Filter.SingleVertex == ParticleIndex + GlobalParticleOffset)
+
+				&&
+
+				(!DebugFilterSet ||
+					DebugFilterSet->Contains(ParticleIndex));
+		}
+
+		static bool ShouldDrawEdge(const FClothVisualizationNoGC::FDrawFilter& Filter, const TSet<int32>* const DebugFilterSet, const FUintVector2& Edge, int32 EdgeToGlobalOffset, int32 EdgeToLocalOffset)
+		{
+			return (Filter.SingleVertex == INDEX_NONE ||
+				Filter.SingleVertex == Edge[0] + EdgeToGlobalOffset ||
+				Filter.SingleVertex == Edge[1] + EdgeToGlobalOffset)
+
+				&&
+
+				(!DebugFilterSet ||
+					DebugFilterSet->Contains(Edge[0] + EdgeToLocalOffset) ||
+					DebugFilterSet->Contains(Edge[1] + EdgeToLocalOffset));
+		}
+
+		static bool ShouldDrawConstraint2(const FClothVisualizationNoGC::FDrawFilter& Filter, const TSet<int32>* const DebugFilterSet, const TVec2<int32>& Constraint, int32 ConstraintToGlobalOffset, int32 ConstraintToLocalOffset)
+		{
+			return (Filter.SingleVertex == INDEX_NONE ||
+				Filter.SingleVertex == Constraint[0] + ConstraintToGlobalOffset ||
+				Filter.SingleVertex == Constraint[1] + ConstraintToGlobalOffset)
+
+				&&
+
+				(!DebugFilterSet ||
+					DebugFilterSet->Contains(Constraint[0] + ConstraintToLocalOffset) ||
+					DebugFilterSet->Contains(Constraint[1] + ConstraintToLocalOffset));
+		}
+
+		static bool ShouldDrawElement(const FClothVisualizationNoGC::FDrawFilter& Filter, const TSet<int32>* const DebugFilterSet, const TVec3<int32>& Element, int32 ElementToGlobalOffset, int32 ElementToLocalOffset)
+		{
+			return (Filter.SingleVertex == INDEX_NONE ||
+				Filter.SingleVertex == Element.X + ElementToGlobalOffset ||
+				Filter.SingleVertex == Element.Y + ElementToGlobalOffset ||
+				Filter.SingleVertex == Element.Z + ElementToGlobalOffset)
+
+				&&
+
+				(!DebugFilterSet ||
+					DebugFilterSet->Contains(Element.X + ElementToLocalOffset) ||
+					DebugFilterSet->Contains(Element.Y + ElementToLocalOffset) ||
+					DebugFilterSet->Contains(Element.Z + ElementToLocalOffset));
+		}
+
+		static bool ShouldDrawConstraint4(const FClothVisualizationNoGC::FDrawFilter& Filter, const TSet<int32>* const DebugFilterSet, const TVec4<int32>& Constraint, int32 ConstraintToGlobalOffset, int32 ConstraintToLocalOffset)
+		{
+			return (Filter.SingleVertex == INDEX_NONE ||
+				Filter.SingleVertex == Constraint[0] + ConstraintToGlobalOffset ||
+				Filter.SingleVertex == Constraint[1] + ConstraintToGlobalOffset ||
+				Filter.SingleVertex == Constraint[2] + ConstraintToGlobalOffset ||
+				Filter.SingleVertex == Constraint[3] + ConstraintToGlobalOffset)
+
+				&&
+
+				(!DebugFilterSet ||
+					DebugFilterSet->Contains(Constraint[0] + ConstraintToLocalOffset) ||
+					DebugFilterSet->Contains(Constraint[1] + ConstraintToLocalOffset) ||
+					DebugFilterSet->Contains(Constraint[2] + ConstraintToLocalOffset) ||
+					DebugFilterSet->Contains(Constraint[3] + ConstraintToLocalOffset));
+		}
 	}  // namespace Private
 
 	// Delay loading materials to their first point of use to avoid causing sync flush issues and slow loading when starting PIE
@@ -239,14 +305,14 @@ namespace Chaos
 	}
 
 #if WITH_EDITOR
-	void FClothVisualizationNoGC::DrawPhysMeshShaded(FPrimitiveDrawInterface* PDI) const
+	void FClothVisualizationNoGC::DrawPhysMeshShaded(const FDrawContext& Context) const
 	{
-		if (!Solver || !FMaterials::GetInstance().GetClothMaterial())
+		if (!Solver || !FMaterials::GetInstance().GetClothMaterial() || !Context.PDI)
 		{
 			return;
 		}
 
-		FDynamicMeshBuilder MeshBuilder(PDI->View->GetFeatureLevel());
+		FDynamicMeshBuilder MeshBuilder(Context.PDI->View->GetFeatureLevel());
 		int32 VertexIndex = 0;
 
 		for (const FClothingSimulationCloth* const Cloth : Solver->GetCloths())
@@ -256,48 +322,59 @@ namespace Chaos
 			{
 				continue;
 			}
+
+
 			// Elements are local indexed for new solver
 			const int32 Offset = Solver->IsLegacySolver() ? ParticleRangeId : 0;
+			const TSet<int32>* const DebugFilterSet = Context.Filter.VertexSet.IsEmpty() ? nullptr : Cloth->GetVertexSetByName(Solver, Context.Filter.VertexSet);
+			const int32 GlobalParticleOffset = Solver->GetGlobalParticleOffset(ParticleRangeId);
+			auto ShouldDrawElement = [&Context, DebugFilterSet, GlobalParticleOffset, Offset](const TVec3<int32>& Element)
+				{
+					return Private::ShouldDrawElement(Context.Filter, DebugFilterSet, Element, GlobalParticleOffset - Offset, -Offset);
+				};
 
 			const TConstArrayView<TVec3<int32>> Elements = Cloth->GetTriangleMesh(Solver).GetElements();
 			const TConstArrayView<Softs::FSolverVec3> Positions = Cloth->GetParticlePositions(Solver);
 			const TConstArrayView<Softs::FSolverReal> InvMasses = Cloth->GetParticleInvMasses(Solver);
 			check(InvMasses.Num() == Positions.Num());
 
-			for (int32 ElementIndex = 0; ElementIndex < Elements.Num(); ++ElementIndex, VertexIndex += 3)
+			for (int32 ElementIndex = 0; ElementIndex < Elements.Num(); ++ElementIndex)
 			{
 				const TVec3<int32>& Element = Elements[ElementIndex];
+				if(ShouldDrawElement(Element))
+				{
+					const FVector3f Pos0(Positions[Element.X - Offset]); // TODO: Triangle Mesh shouldn't really be solver dependent (ie not use an offset)
+					const FVector3f Pos1(Positions[Element.Y - Offset]);
+					const FVector3f Pos2(Positions[Element.Z - Offset]);
 
-				const FVector3f Pos0(Positions[Element.X - Offset]); // TODO: Triangle Mesh shouldn't really be solver dependent (ie not use an offset)
-				const FVector3f Pos1(Positions[Element.Y - Offset]);
-				const FVector3f Pos2(Positions[Element.Z - Offset]);
-
-				const FVector3f Normal = FVector3f::CrossProduct(Pos2 - Pos0, Pos1 - Pos0).GetSafeNormal();
-				const FVector3f Tangent = ((Pos1 + Pos2) * 0.5f - Pos0).GetSafeNormal();
+					const FVector3f Normal = FVector3f::CrossProduct(Pos2 - Pos0, Pos1 - Pos0).GetSafeNormal();
+					const FVector3f Tangent = ((Pos1 + Pos2) * 0.5f - Pos0).GetSafeNormal();
 
-				const bool bIsKinematic0 = (InvMasses[Element.X - Offset] == (Softs::FSolverReal)0.);
-				const bool bIsKinematic1 = (InvMasses[Element.Y - Offset] == (Softs::FSolverReal)0.);
-				const bool bIsKinematic2 = (InvMasses[Element.Z - Offset] == (Softs::FSolverReal)0.);
+					const bool bIsKinematic0 = (InvMasses[Element.X - Offset] == (Softs::FSolverReal)0.);
+					const bool bIsKinematic1 = (InvMasses[Element.Y - Offset] == (Softs::FSolverReal)0.);
+					const bool bIsKinematic2 = (InvMasses[Element.Z - Offset] == (Softs::FSolverReal)0.);
 
-				MeshBuilder.AddVertex(FDynamicMeshVertex(Pos0, Tangent, Normal, FVector2f(0.f, 0.f), bIsKinematic0 ? FColor::Purple : FColor::White));
-				MeshBuilder.AddVertex(FDynamicMeshVertex(Pos1, Tangent, Normal, FVector2f(0.f, 1.f), bIsKinematic1 ? FColor::Purple : FColor::White));
-				MeshBuilder.AddVertex(FDynamicMeshVertex(Pos2, Tangent, Normal, FVector2f(1.f, 1.f), bIsKinematic2 ? FColor::Purple : FColor::White));
-				MeshBuilder.AddTriangle(VertexIndex, VertexIndex + 1, VertexIndex + 2);
+					MeshBuilder.AddVertex(FDynamicMeshVertex(Pos0, Tangent, Normal, FVector2f(0.f, 0.f), bIsKinematic0 ? FColor::Purple : FColor::White));
+					MeshBuilder.AddVertex(FDynamicMeshVertex(Pos1, Tangent, Normal, FVector2f(0.f, 1.f), bIsKinematic1 ? FColor::Purple : FColor::White));
+					MeshBuilder.AddVertex(FDynamicMeshVertex(Pos2, Tangent, Normal, FVector2f(1.f, 1.f), bIsKinematic2 ? FColor::Purple : FColor::White));
+					MeshBuilder.AddTriangle(VertexIndex, VertexIndex + 1, VertexIndex + 2);
+					VertexIndex += 3;
+				}
 			}
 		}
 
 		const FMatrix LocalSimSpaceToWorld = GetLocalSpaceToWorld(*Solver);
-		MeshBuilder.Draw(PDI, LocalSimSpaceToWorld, FMaterials::GetInstance().GetClothMaterial()->GetRenderProxy(), Chaos::Private::GetDepthPriority(), false, false);
+		MeshBuilder.Draw(Context.PDI, LocalSimSpaceToWorld, FMaterials::GetInstance().GetClothMaterial()->GetRenderProxy(), Chaos::Private::GetDepthPriority(), false, false);
 	}
 
-	void FClothVisualizationNoGC::DrawWeightMapWithName(FPrimitiveDrawInterface* PDI, const FString& Name) const
+	void FClothVisualizationNoGC::DrawWeightMapWithName(const FDrawContext& Context, const FString& Name) const
 	{
-		if (!Solver || !FMaterials::GetInstance().GetClothMaterialColor())
+		if (!Solver || !FMaterials::GetInstance().GetClothMaterialColor() || !Context.PDI)
 		{
 			return;
 		}
 
-		FDynamicMeshBuilder MeshBuilder(PDI->View->GetFeatureLevel());
+		FDynamicMeshBuilder MeshBuilder(Context.PDI->View->GetFeatureLevel());
 		int32 VertexIndex = 0;
 
 		for (const FClothingSimulationCloth* const Cloth : Solver->GetCloths())
@@ -309,66 +386,71 @@ namespace Chaos
 			}
 			// Elements are local indexed for new solver
 			const int32 Offset = Solver->IsLegacySolver() ? ParticleRangeId : 0;
+			const TSet<int32>* const DebugFilterSet = Context.Filter.VertexSet.IsEmpty() ? nullptr : Cloth->GetVertexSetByName(Solver, Context.Filter.VertexSet);
+			const int32 GlobalParticleOffset = Solver->GetGlobalParticleOffset(ParticleRangeId);
+			auto ShouldDrawElement = [&Context, DebugFilterSet, GlobalParticleOffset, Offset](const TVec3<int32>& Element)
+				{
+					return Private::ShouldDrawElement(Context.Filter, DebugFilterSet, Element, GlobalParticleOffset - Offset, -Offset);
+				};
 			const TConstArrayView<TVec3<int32>> Elements = Cloth->GetTriangleMesh(Solver).GetElements();
 			const TConstArrayView<Softs::FSolverVec3> Positions = Cloth->GetParticlePositions(Solver);
-			const TConstArrayView<FRealSingle>& WeightMap = Cloth->GetWeightMapByName(Solver, Name);
+			const TConstArrayView<FRealSingle>& WeightMap = Cloth->GetWeightMapByName(Solver, Name.IsEmpty() ? Private::WeightMapName : Name);
 
-			for (int32 ElementIndex = 0; ElementIndex < Elements.Num(); ++ElementIndex, VertexIndex += 3)
+			for (int32 ElementIndex = 0; ElementIndex < Elements.Num(); ++ElementIndex)
 			{
 				const TVec3<int32>& Element = Elements[ElementIndex];
 
-				const FVector3f Pos0(Positions[Element.X - Offset]);
-				const FVector3f Pos1(Positions[Element.Y - Offset]);
-				const FVector3f Pos2(Positions[Element.Z - Offset]);
+				if (ShouldDrawElement(Element))
+				{
+					const FVector3f Pos0(Positions[Element.X - Offset]);
+					const FVector3f Pos1(Positions[Element.Y - Offset]);
+					const FVector3f Pos2(Positions[Element.Z - Offset]);
 
-				const FVector3f Normal = FVector3f::CrossProduct(Pos2 - Pos0, Pos1 - Pos0).GetSafeNormal();
-				const FVector3f Tangent = ((Pos1 + Pos2) * 0.5f - Pos0).GetSafeNormal();
+					const FVector3f Normal = FVector3f::CrossProduct(Pos2 - Pos0, Pos1 - Pos0).GetSafeNormal();
+					const FVector3f Tangent = ((Pos1 + Pos2) * 0.5f - Pos0).GetSafeNormal();
 
-				FLinearColor VertexColor1 = FLinearColor::Black;
-				FLinearColor VertexColor2 = FLinearColor::Black;
-				FLinearColor VertexColor3 = FLinearColor::Black;
+					FLinearColor VertexColor1 = FLinearColor::Black;
+					FLinearColor VertexColor2 = FLinearColor::Black;
+					FLinearColor VertexColor3 = FLinearColor::Black;
 
-				if (!WeightMap.IsEmpty() && WeightMap.Num() == Positions.Num()) // if map with that name exists and not empty
-				{
-					const FRealSingle Value0(WeightMap[Element.X - Offset]);
-					const FRealSingle Value1(WeightMap[Element.Y - Offset]);
-					const FRealSingle Value2(WeightMap[Element.Z - Offset]);
+					if (!WeightMap.IsEmpty() && WeightMap.Num() == Positions.Num()) // if map with that name exists and not empty
+					{
+						const FRealSingle Value0(WeightMap[Element.X - Offset]);
+						const FRealSingle Value1(WeightMap[Element.Y - Offset]);
+						const FRealSingle Value2(WeightMap[Element.Z - Offset]);
 
-					VertexColor1 = FLinearColor::LerpUsingHSV(FLinearColor::Black, FLinearColor::White, (float)Value0);
-					VertexColor2 = FLinearColor::LerpUsingHSV(FLinearColor::Black, FLinearColor::White, (float)Value1);
-					VertexColor3 = FLinearColor::LerpUsingHSV(FLinearColor::Black, FLinearColor::White, (float)Value2);
+						VertexColor1 = FLinearColor::LerpUsingHSV(FLinearColor::Black, FLinearColor::White, (float)Value0);
+						VertexColor2 = FLinearColor::LerpUsingHSV(FLinearColor::Black, FLinearColor::White, (float)Value1);
+						VertexColor3 = FLinearColor::LerpUsingHSV(FLinearColor::Black, FLinearColor::White, (float)Value2);
+					}
+
+					MeshBuilder.AddVertex(FDynamicMeshVertex(Pos0, Tangent, Normal, FVector2f(0.f, 0.f), VertexColor1.ToFColor(true)));
+					MeshBuilder.AddVertex(FDynamicMeshVertex(Pos1, Tangent, Normal, FVector2f(0.f, 1.f), VertexColor2.ToFColor(true)));
+					MeshBuilder.AddVertex(FDynamicMeshVertex(Pos2, Tangent, Normal, FVector2f(1.f, 1.f), VertexColor3.ToFColor(true)));
+					MeshBuilder.AddTriangle(VertexIndex, VertexIndex + 1, VertexIndex + 2);
+					VertexIndex += 3;
 				}
-	
-				MeshBuilder.AddVertex(FDynamicMeshVertex(Pos0, Tangent, Normal, FVector2f(0.f, 0.f), VertexColor1.ToFColor(true)));
-				MeshBuilder.AddVertex(FDynamicMeshVertex(Pos1, Tangent, Normal, FVector2f(0.f, 1.f), VertexColor2.ToFColor(true)));
-				MeshBuilder.AddVertex(FDynamicMeshVertex(Pos2, Tangent, Normal, FVector2f(1.f, 1.f), VertexColor3.ToFColor(true)));
-				MeshBuilder.AddTriangle(VertexIndex, VertexIndex + 1, VertexIndex + 2);
 			}
 		}
 
 		const FMatrix LocalSimSpaceToWorld = GetLocalSpaceToWorld(*Solver);
-		MeshBuilder.Draw(PDI, LocalSimSpaceToWorld, FMaterials::GetInstance().GetClothMaterialColor()->GetRenderProxy(), Chaos::Private::GetDepthPriority(), false, false);
+		MeshBuilder.Draw(Context.PDI, LocalSimSpaceToWorld, FMaterials::GetInstance().GetClothMaterialColor()->GetRenderProxy(), Chaos::Private::GetDepthPriority(), false, false);
 	}
 
-	void FClothVisualizationNoGC::DrawWeightMap(FPrimitiveDrawInterface* PDI) const
+	void FClothVisualizationNoGC::DrawInpaintWeightsMatched(const FDrawContext& Context) const
 	{
-		DrawWeightMapWithName(PDI, Private::WeightMapName);
+		DrawWeightMapWithName(Context, TEXT("_InpaintWeightMask"));
 	}
 
-	void FClothVisualizationNoGC::DrawInpaintWeightsMatched(FPrimitiveDrawInterface* PDI) const
+	void FClothVisualizationNoGC::DrawSelfCollisionLayers(const FDrawContext& Context) const
 	{
-		DrawWeightMapWithName(PDI, TEXT("_InpaintWeightMask"));
-	}
 
-	void FClothVisualizationNoGC::DrawSelfCollisionLayers(FPrimitiveDrawInterface* PDI) const
-	{
-
-		if (!Solver || !FMaterials::GetInstance().GetClothMaterialColor())
+		if (!Solver || !FMaterials::GetInstance().GetClothMaterialColor() || !Context.PDI)
 		{
 			return;
 		}
 
-		FDynamicMeshBuilder MeshBuilder(PDI->View->GetFeatureLevel());
+		FDynamicMeshBuilder MeshBuilder(Context.PDI->View->GetFeatureLevel());
 		int32 VertexIndex = 0;
 
 		for (const FClothingSimulationCloth* const Cloth : Solver->GetCloths())
@@ -380,45 +462,53 @@ namespace Chaos
 			}
 			// Elements are local indexed for new solver
 			const int32 Offset = Solver->IsLegacySolver() ? ParticleRangeId : 0;
+			const TSet<int32>* const DebugFilterSet = Context.Filter.VertexSet.IsEmpty() ? nullptr : Cloth->GetVertexSetByName(Solver, Context.Filter.VertexSet);
+			const int32 GlobalParticleOffset = Solver->GetGlobalParticleOffset(ParticleRangeId);
+			auto ShouldDrawElement = [&Context, DebugFilterSet, GlobalParticleOffset, Offset](const TVec3<int32>& Element)
+				{
+					return Private::ShouldDrawElement(Context.Filter, DebugFilterSet, Element, GlobalParticleOffset - Offset, -Offset);
+				};
 			const TConstArrayView<TVec3<int32>> Elements = Cloth->GetTriangleMesh(Solver).GetElements();
 			const TConstArrayView<Softs::FSolverVec3> Positions = Cloth->GetParticlePositions(Solver);
 			const TConstArrayView<int32>& WeightMap = Cloth->GetFaceIntMapByProperty(Solver, Softs::FPBDCollisionSpringConstraints::SelfCollisionLayersName);
 
-			for (int32 ElementIndex = 0; ElementIndex < Elements.Num(); ++ElementIndex, VertexIndex += 3)
+			for (int32 ElementIndex = 0; ElementIndex < Elements.Num(); ++ElementIndex)
 			{
-
 				const TVec3<int32>& Element = Elements[ElementIndex];
+				if (ShouldDrawElement(Element))
+				{
+					const FVector3f Pos0(Positions[Element.X - Offset]);
+					const FVector3f Pos1(Positions[Element.Y - Offset]);
+					const FVector3f Pos2(Positions[Element.Z - Offset]);
 
-				const FVector3f Pos0(Positions[Element.X - Offset]);
-				const FVector3f Pos1(Positions[Element.Y - Offset]);
-				const FVector3f Pos2(Positions[Element.Z - Offset]);
+					const FVector3f Normal = FVector3f::CrossProduct(Pos2 - Pos0, Pos1 - Pos0).GetSafeNormal();
+					const FVector3f Tangent = ((Pos1 + Pos2) * 0.5f - Pos0).GetSafeNormal();
 
-				const FVector3f Normal = FVector3f::CrossProduct(Pos2 - Pos0, Pos1 - Pos0).GetSafeNormal();
-				const FVector3f Tangent = ((Pos1 + Pos2) * 0.5f - Pos0).GetSafeNormal();
+					FLinearColor VertexColor1 = FLinearColor::Gray;
+					FLinearColor VertexColor2 = FLinearColor::Gray;
+					FLinearColor VertexColor3 = FLinearColor::Gray;
 
-				FLinearColor VertexColor1 = FLinearColor::Gray;
-				FLinearColor VertexColor2 = FLinearColor::Gray;
-				FLinearColor VertexColor3 = FLinearColor::Gray;
+					if (!WeightMap.IsEmpty() && WeightMap.Num() == Elements.Num() && WeightMap[ElementIndex] != INDEX_NONE) // if map with that name exists and not empty
+					{
+						VertexColor1 = VertexColor2 = VertexColor3 = Chaos::Private::PseudoRandomColor(WeightMap[ElementIndex]);
+					}
 
-				if (!WeightMap.IsEmpty() && WeightMap.Num() == Elements.Num() && WeightMap[ElementIndex]!=INDEX_NONE) // if map with that name exists and not empty
-				{
-					VertexColor1 = VertexColor2 = VertexColor3 = Chaos::Private::PseudoRandomColor(WeightMap[ElementIndex]);
+					MeshBuilder.AddVertex(FDynamicMeshVertex(Pos0, Tangent, Normal, FVector2f(0.f, 0.f), VertexColor1.ToFColor(true)));
+					MeshBuilder.AddVertex(FDynamicMeshVertex(Pos1, Tangent, Normal, FVector2f(0.f, 1.f), VertexColor2.ToFColor(true)));
+					MeshBuilder.AddVertex(FDynamicMeshVertex(Pos2, Tangent, Normal, FVector2f(1.f, 1.f), VertexColor3.ToFColor(true)));
+					MeshBuilder.AddTriangle(VertexIndex, VertexIndex + 1, VertexIndex + 2);
+					VertexIndex += 3;
 				}
-
-				MeshBuilder.AddVertex(FDynamicMeshVertex(Pos0, Tangent, Normal, FVector2f(0.f, 0.f), VertexColor1.ToFColor(true)));
-				MeshBuilder.AddVertex(FDynamicMeshVertex(Pos1, Tangent, Normal, FVector2f(0.f, 1.f), VertexColor2.ToFColor(true)));
-				MeshBuilder.AddVertex(FDynamicMeshVertex(Pos2, Tangent, Normal, FVector2f(1.f, 1.f), VertexColor3.ToFColor(true)));
-				MeshBuilder.AddTriangle(VertexIndex, VertexIndex + 1, VertexIndex + 2);
 			}
 		}
 
 		FMatrix LocalSimSpaceToWorld(FMatrix::Identity);
 		LocalSimSpaceToWorld.SetOrigin(Solver->GetLocalSpaceLocation());
-		MeshBuilder.Draw(PDI, LocalSimSpaceToWorld, FMaterials::GetInstance().GetClothMaterialColor()->GetRenderProxy(), Chaos::Private::GetDepthPriority(), false, false);
+		MeshBuilder.Draw(Context.PDI, LocalSimSpaceToWorld, FMaterials::GetInstance().GetClothMaterialColor()->GetRenderProxy(), Chaos::Private::GetDepthPriority(), false, false);
 	}
 #endif  // #if WITH_EDITOR
 
-	void FClothVisualizationNoGC::DrawParticleIndices(FCanvas* Canvas, const FSceneView* SceneView) const
+	void FClothVisualizationNoGC::DrawParticleIndices(const FDrawTextsContext& Context) const
 	{
 		if (!Solver)
 		{
@@ -435,7 +525,12 @@ namespace Chaos
 			{
 				continue;
 			}
+			const TSet<int32>* const DebugFilterSet = Context.Filter.VertexSet.IsEmpty() ? nullptr : Cloth->GetVertexSetByName(Solver, Context.Filter.VertexSet);
 			const int32 Offset = Solver->GetGlobalParticleOffset(ParticleRangeId);
+			auto ShouldDrawParticle = [&Context, DebugFilterSet, Offset](const int32 Index)
+				{
+					return Private::ShouldDrawParticle(Context.Filter, DebugFilterSet, Index, Offset);
+				};
 
 			const TConstArrayView<Softs::FSolverVec3> Positions = Cloth->GetParticlePositions(Solver);
 			const TConstArrayView<Softs::FSolverReal> InvMasses = Cloth->GetParticleInvMasses(Solver);
@@ -443,15 +538,18 @@ namespace Chaos
 
 			for (int32 Index = 0; Index < Positions.Num(); ++Index)
 			{
-				const FVector Position = GetWorldPosition(*Solver, Positions[Index]);
+				if (ShouldDrawParticle(Index))
+				{
+					const FVector Position = GetWorldPosition(*Solver, Positions[Index]);
 
-				const FText Text = FText::AsNumber(Offset + Index);
-				DrawText(Canvas, SceneView, Position, Text, InvMasses[Index] == (Softs::FSolverReal)0. ? KinematicColor : DynamicColor);
+					const FText Text = FText::AsNumber(Offset + Index);
+					DrawText(Context.Canvas, Context.SceneView, Position, Text, InvMasses[Index] == (Softs::FSolverReal)0. ? KinematicColor : DynamicColor);
+				}
 			}
 		}
 	}
 
-	void FClothVisualizationNoGC::DrawElementIndices(FCanvas* Canvas, const FSceneView* SceneView) const
+	void FClothVisualizationNoGC::DrawElementIndices(const FDrawTextsContext& Context) const
 	{
 		if (!Solver)
 		{
@@ -470,6 +568,12 @@ namespace Chaos
 			}
 			// Elements are local indexed for new solver
 			const int32 Offset = Solver->IsLegacySolver() ? ParticleRangeId : 0;
+			const TSet<int32>* const DebugFilterSet = Context.Filter.VertexSet.IsEmpty() ? nullptr : Cloth->GetVertexSetByName(Solver, Context.Filter.VertexSet);
+			const int32 GlobalParticleOffset = Solver->GetGlobalParticleOffset(ParticleRangeId);
+			auto ShouldDrawElement = [&Context, DebugFilterSet, GlobalParticleOffset, Offset](const TVec3<int32>& Element)
+				{
+					return Private::ShouldDrawElement(Context.Filter, DebugFilterSet, Element, GlobalParticleOffset - Offset, -Offset);
+				};
 
 			const TArray<TVec3<int32>>& Elements = Cloth->GetTriangleMesh(Solver).GetElements();
 			const TConstArrayView<Softs::FSolverVec3> Positions = Cloth->GetParticlePositions(Solver);
@@ -479,22 +583,25 @@ namespace Chaos
 			for (int32 Index = 0; Index < Elements.Num(); ++Index)
 			{
 				const TVec3<int32>& Element = Elements[Index];
-				const FVector Position = GetWorldPosition(*Solver,
-					(Positions[Element[0] - Offset] +
-					Positions[Element[1] - Offset] +
-					Positions[Element[2] - Offset]) / (Softs::FSolverReal)3.f);
-
-				const bool bIsKinematic0 = (InvMasses[Element.X - Offset] == (Softs::FSolverReal)0.);
-				const bool bIsKinematic1 = (InvMasses[Element.Y - Offset] == (Softs::FSolverReal)0.);
-				const bool bIsKinematic2 = (InvMasses[Element.Z - Offset] == (Softs::FSolverReal)0.);
-				const FLinearColor& Color = (bIsKinematic0 && bIsKinematic1 && bIsKinematic2) ? KinematicColor : DynamicColor;
-				const FText Text = FText::AsNumber(Index);
-				DrawText(Canvas, SceneView, Position, Text, Color);
+				if (ShouldDrawElement(Element))
+				{
+					const FVector Position = GetWorldPosition(*Solver,
+						(Positions[Element[0] - Offset] +
+							Positions[Element[1] - Offset] +
+							Positions[Element[2] - Offset]) / (Softs::FSolverReal)3.f);
+
+					const bool bIsKinematic0 = (InvMasses[Element.X - Offset] == (Softs::FSolverReal)0.);
+					const bool bIsKinematic1 = (InvMasses[Element.Y - Offset] == (Softs::FSolverReal)0.);
+					const bool bIsKinematic2 = (InvMasses[Element.Z - Offset] == (Softs::FSolverReal)0.);
+					const FLinearColor& Color = (bIsKinematic0 && bIsKinematic1 && bIsKinematic2) ? KinematicColor : DynamicColor;
+					const FText Text = FText::AsNumber(Index);
+					DrawText(Context.Canvas, Context.SceneView, Position, Text, Color);
+				}
 			}
 		}
 	}
 
-	void FClothVisualizationNoGC::DrawMaxDistanceValues(FCanvas* Canvas, const FSceneView* SceneView) const
+	void FClothVisualizationNoGC::DrawMaxDistanceValues(const FDrawTextsContext& Context) const
 	{
 		if (!Solver)
 		{
@@ -515,11 +622,18 @@ namespace Chaos
 
 		for (const FClothingSimulationCloth* const Cloth : Solver->GetCloths())
 		{
-			if (Cloth->GetParticleRangeId(Solver) == INDEX_NONE)
+			const int32 ParticleRangeId = Cloth->GetParticleRangeId(Solver);
+			if (ParticleRangeId == INDEX_NONE)
 			{
 				continue;
 			}
 
+			const TSet<int32>* const DebugFilterSet = Context.Filter.VertexSet.IsEmpty() ? nullptr : Cloth->GetVertexSetByName(Solver, Context.Filter.VertexSet);
+			const int32 GlobalParticleOffset = Solver->GetGlobalParticleOffset(ParticleRangeId);
+			auto ShouldDrawParticle = [&Context, DebugFilterSet, GlobalParticleOffset](int32 Index)
+				{
+					return Private::ShouldDrawParticle(Context.Filter, DebugFilterSet, Index, GlobalParticleOffset);
+				};
 			const int32 NumParticles = Cloth->GetNumParticles(Solver);
 			const Softs::FPBDFlatWeightMapView MaxDistances(
 				(Softs::FSolverVec2)(Softs::FPBDSphericalConstraint::GetWeightedFloatMaxDistance(Cloth->GetConfig()->GetProperties(Cloth->GetLODIndex(Solver)), FVector2f(0.f, 1.f))),
@@ -533,16 +647,19 @@ namespace Chaos
 
 			for (int32 Index = 0; Index < MaxDistances.Num(); ++Index)
 			{
-				const FReal MaxDistance = (FReal)MaxDistances.GetValue(Index);
-				const FVector Position = GetWorldPosition(*Solver, Positions[Index]);
+				if (ShouldDrawParticle(Index))
+				{
+					const FReal MaxDistance = (FReal)MaxDistances.GetValue(Index);
+					const FVector Position = GetWorldPosition(*Solver, Positions[Index]);
 
-				const FText Text = FText::AsNumber(MaxDistance, &NumberFormattingOptions);
-				DrawText(Canvas, SceneView, Position, Text, InvMasses[Index] == (Softs::FSolverReal)0. ? KinematicColor : DynamicColor);
+					const FText Text = FText::AsNumber(MaxDistance, &NumberFormattingOptions);
+					DrawText(Context.Canvas, Context.SceneView, Position, Text, InvMasses[Index] == (Softs::FSolverReal)0. ? KinematicColor : DynamicColor);
+				}
 			}
 		}
 	}
 
-	void FClothVisualizationNoGC::DrawLocalSpaceBoneNames(FCanvas* Canvas, const FSceneView* SceneView) const
+	void FClothVisualizationNoGC::DrawLocalSpaceBoneNames(const FDrawTextsContext& Context) const
 	{
 		if (!Solver)
 		{
@@ -550,7 +667,7 @@ namespace Chaos
 		}
 
 		// Draw local space
-		DrawText(Canvas, SceneView, Solver->GetLocalSpaceLocation() + FVec3(2.), LOCTEXT("SimulationSpace", "Simulation Space"), FLinearColor::White, 1.5f);
+		DrawText(Context.Canvas, Context.SceneView, Solver->GetLocalSpaceLocation() + FVec3(2.), LOCTEXT("SimulationSpace", "Simulation Space"), FLinearColor::White, 1.5f);
 
 		// Draw reference spaces
 		for (const FClothingSimulationCloth* const Cloth : Solver->GetCloths())
@@ -562,7 +679,7 @@ namespace Chaos
 			}
 			const FRigidTransform3& ReferenceSpaceTransform = Cloth->GetReferenceSpaceTransform();
 			check(Cloth->GetMesh());
-			DrawText(Canvas, SceneView, ReferenceSpaceTransform.GetLocation() + FVec3(-2.), FText::FromName(Cloth->GetMesh()->GetReferenceBoneName()), FLinearColor::White, 1.5f);
+			DrawText(Context.Canvas, Context.SceneView, ReferenceSpaceTransform.GetLocation() + FVec3(-2.), FText::FromName(Cloth->GetMesh()->GetReferenceBoneName()), FLinearColor::White, 1.5f);
 		}
 	}
 
@@ -1199,7 +1316,7 @@ namespace Chaos
 		}
 	}
 
-	void FClothVisualizationNoGC::DrawBounds(FPrimitiveDrawInterface* PDI) const
+	void FClothVisualizationNoGC::DrawBounds(const FDrawContext& Context) const
 	{
 		if (!Solver)
 		{
@@ -1210,8 +1327,8 @@ namespace Chaos
 		const FBoxSphereBounds Bounds = Solver->CalculateBounds();
 
 		// Draw bounds
-		DrawBox(PDI, FAABB3(-Bounds.BoxExtent, Bounds.BoxExtent), FQuat::Identity, Bounds.Origin, 1., FLinearColor(FColor::Purple));
-		DrawSphere(PDI, FSphere(FVector::ZeroVector, Bounds.SphereRadius), FQuat::Identity, Bounds.Origin, 1., FLinearColor(FColor::Orange));
+		DrawBox(Context.PDI, FAABB3(-Bounds.BoxExtent, Bounds.BoxExtent), FQuat::Identity, Bounds.Origin, 1., FLinearColor(FColor::Purple));
+		DrawSphere(Context.PDI, FSphere(FVector::ZeroVector, Bounds.SphereRadius), FQuat::Identity, Bounds.Origin, 1., FLinearColor(FColor::Orange));
 
 		// Draw individual cloth bounds
 		static const FLinearColor Color = FLinearColor(FColor::Purple).Desaturate(0.5);
@@ -1224,11 +1341,11 @@ namespace Chaos
 			}
 
 			const FAABB3 BoundingBox = Cloth->CalculateBoundingBox(Solver);
-			DrawBox(PDI, BoundingBox, FQuat::Identity, FVector::ZeroVector, 1.f, Color);  // TODO: Express bounds in local coordinates for LWC
+			DrawBox(Context.PDI, BoundingBox, FQuat::Identity, FVector::ZeroVector, 1.f, Color);  // TODO: Express bounds in local coordinates for LWC
 		}
 	}
 
-	void FClothVisualizationNoGC::DrawGravity(FPrimitiveDrawInterface* PDI) const
+	void FClothVisualizationNoGC::DrawGravity(const FDrawContext& Context) const
 	{
 		if (!Solver)
 		{
@@ -1250,12 +1367,21 @@ namespace Chaos
 				check(!Solver->IsLegacySolver());
 				if (ExternalForces->HasPerParticleGravity())
 				{
+					const TSet<int32>* const DebugFilterSet = Context.Filter.VertexSet.IsEmpty() ? nullptr : Cloth->GetVertexSetByName(Solver, Context.Filter.VertexSet);
+					const int32 GlobalParticleOffset = Solver->GetGlobalParticleOffset(ParticleRangeId);
+					auto ShouldDrawParticle = [&Context, &DebugFilterSet, GlobalParticleOffset](int32 ParticleIndex)
+						{
+							return Private::ShouldDrawParticle(Context.Filter, DebugFilterSet, ParticleIndex, GlobalParticleOffset);
+						};
 					const TConstArrayView<Softs::FSolverVec3> Positions = Solver->GetParticleXsView(ParticleRangeId);
 					for (int32 ParticleIndex = 0; ParticleIndex < Positions.Num(); ++ParticleIndex)
 					{
-						const FVector Pos0 = GetWorldPosition(*Solver, Positions[ParticleIndex]);
-						const FVector Pos1 = Pos0 + GravityVectorLengthMultiplier * GetWorldVector(*Solver, ExternalForces->GetScaledGravity(ParticleIndex));
-						DrawArrow(PDI, Pos0, Pos1, FLinearColor::Red);
+						if (ShouldDrawParticle(ParticleIndex))
+						{
+							const FVector Pos0 = GetWorldPosition(*Solver, Positions[ParticleIndex]);
+							const FVector Pos1 = Pos0 + GravityVectorLengthMultiplier * GetWorldVector(*Solver, ExternalForces->GetScaledGravity(ParticleIndex));
+							DrawArrow(Context.PDI, Pos0, Pos1, FLinearColor::Red);
+						}
 					}
 				}
 				else
@@ -1263,7 +1389,7 @@ namespace Chaos
 					const FRigidTransform3& ReferenceSpaceTransform = Cloth->GetReferenceSpaceTransform();
 					const FVector Pos0 = ReferenceSpaceTransform.GetLocation();
 					const FVector Pos1 = Pos0 + GravityVectorLengthMultiplier * GetWorldVector(*Solver, ExternalForces->GetScaledGravity(0));
-					DrawArrow(PDI, Pos0, Pos1, FLinearColor::Red);
+					DrawArrow(Context.PDI, Pos0, Pos1, FLinearColor::Red);
 				}
 			}
 			else
@@ -1271,13 +1397,13 @@ namespace Chaos
 				const FRigidTransform3& ReferenceSpaceTransform = Cloth->GetReferenceSpaceTransform();
 				const FVector Pos0 = ReferenceSpaceTransform.GetLocation();
 				const FVector Pos1 = Pos0 + GravityVectorLengthMultiplier * GetWorldVector(*Solver, Cloth->GetGravity(Solver));
-				DrawArrow(PDI, Pos0, Pos1, FLinearColor::Red);
+				DrawArrow(Context.PDI, Pos0, Pos1, FLinearColor::Red);
 			}
 		}
 	}
 
 
-	void FClothVisualizationNoGC::DrawFictitiousAngularForces(FPrimitiveDrawInterface* PDI) const
+	void FClothVisualizationNoGC::DrawFictitiousAngularForces(const FDrawContext& Context) const
 	{
 		if (!Solver)
 		{
@@ -1296,24 +1422,34 @@ namespace Chaos
 			{
 				check(!Solver->IsLegacySolver());
 				{
+					const TSet<int32>* const DebugFilterSet = Context.Filter.VertexSet.IsEmpty() ? nullptr : Cloth->GetVertexSetByName(Solver, Context.Filter.VertexSet);
+					const int32 GlobalParticleOffset = Solver->GetGlobalParticleOffset(ParticleRangeId);
+					auto ShouldDrawParticle = [&Context, DebugFilterSet, GlobalParticleOffset](const int32 Index)
+						{
+							return Private::ShouldDrawParticle(Context.Filter, DebugFilterSet, Index, GlobalParticleOffset);
+						};
+
 					const TConstArrayView<Softs::FSolverVec3> Positions = Solver->GetParticleXsView(ParticleRangeId);
 					const Softs::FSolverVec3& FictitousAngularVelocity = ExternalForces->GetFictitiousAngularVelocity();
 					const Softs::FSolverVec3& ReferenceSpaceLocation = ExternalForces->GetReferenceSpaceLocation();
 
 					for (int32 ParticleIndex = 0; ParticleIndex < Positions.Num(); ++ParticleIndex)
 					{
-						const FVector Pos0 = GetWorldPosition(*Solver, Positions[ParticleIndex]);
-						const Softs::FSolverVec3 CentrifugalAccel = -Softs::FSolverVec3::CrossProduct(FictitousAngularVelocity, Softs::FSolverVec3::CrossProduct(FictitousAngularVelocity, Positions[ParticleIndex] - ReferenceSpaceLocation));
+						if (ShouldDrawParticle(ParticleIndex))
+						{
+							const FVector Pos0 = GetWorldPosition(*Solver, Positions[ParticleIndex]);
+							const Softs::FSolverVec3 CentrifugalAccel = -Softs::FSolverVec3::CrossProduct(FictitousAngularVelocity, Softs::FSolverVec3::CrossProduct(FictitousAngularVelocity, Positions[ParticleIndex] - ReferenceSpaceLocation));
 
-						const FVector Pos1 = Pos0 + GetWorldVector(*Solver, CentrifugalAccel);
-						DrawLine(PDI, Pos0, Pos1, FLinearColor::Red);
+							const FVector Pos1 = Pos0 + GetWorldVector(*Solver, CentrifugalAccel);
+							DrawLine(Context.PDI, Pos0, Pos1, FLinearColor::Red);
+						}
 					}
 				}
 			}
 		}
 	}
 
-	void FClothVisualizationNoGC::DrawPhysMeshWired(FPrimitiveDrawInterface* PDI) const
+	void FClothVisualizationNoGC::DrawPhysMeshWired(const FDrawContext& Context) const
 	{
 		if (!Solver)
 		{
@@ -1332,6 +1468,12 @@ namespace Chaos
 			}
 			// Elements are local indexed for new solver
 			const int32 Offset = Solver->IsLegacySolver() ? ParticleRangeId : 0;
+			const TSet<int32>* const DebugFilterSet = Context.Filter.VertexSet.IsEmpty() ? nullptr : Cloth->GetVertexSetByName(Solver, Context.Filter.VertexSet);
+			const int32 GlobalParticleOffset = Solver->GetGlobalParticleOffset(ParticleRangeId);
+			auto ShouldDrawElement = [&Context, DebugFilterSet, GlobalParticleOffset, Offset](const TVec3<int32>& Element)
+				{
+					return Private::ShouldDrawElement(Context.Filter, DebugFilterSet, Element, GlobalParticleOffset - Offset, -Offset);
+				};
 			const TConstArrayView<TVec3<int32>> Elements = Cloth->GetTriangleMesh(Solver).GetElements();
 			const TConstArrayView<Softs::FSolverVec3> Positions = Cloth->GetParticlePositions(Solver);
 			const TConstArrayView<Softs::FSolverReal> InvMasses = Cloth->GetParticleInvMasses(Solver);
@@ -1341,22 +1483,25 @@ namespace Chaos
 			{
 				const TVec3<int32>& Element = Elements[ElementIndex];
 
-				const FVector Pos0 = GetWorldPosition(*Solver, Positions[Element.X - Offset]); // TODO: Triangle Mesh shouldn't really be solver dependent (ie not use an offset)
-				const FVector Pos1 = GetWorldPosition(*Solver, Positions[Element.Y - Offset]);
-				const FVector Pos2 = GetWorldPosition(*Solver, Positions[Element.Z - Offset]);
+				if (ShouldDrawElement(Element))
+				{
+					const FVector Pos0 = GetWorldPosition(*Solver, Positions[Element.X - Offset]); // TODO: Triangle Mesh shouldn't really be solver dependent (ie not use an offset)
+					const FVector Pos1 = GetWorldPosition(*Solver, Positions[Element.Y - Offset]);
+					const FVector Pos2 = GetWorldPosition(*Solver, Positions[Element.Z - Offset]);
 
-				const bool bIsKinematic0 = (InvMasses[Element.X - Offset] == (Softs::FSolverReal)0.);
-				const bool bIsKinematic1 = (InvMasses[Element.Y - Offset] == (Softs::FSolverReal)0.);
-				const bool bIsKinematic2 = (InvMasses[Element.Z - Offset] == (Softs::FSolverReal)0.);
+					const bool bIsKinematic0 = (InvMasses[Element.X - Offset] == (Softs::FSolverReal)0.);
+					const bool bIsKinematic1 = (InvMasses[Element.Y - Offset] == (Softs::FSolverReal)0.);
+					const bool bIsKinematic2 = (InvMasses[Element.Z - Offset] == (Softs::FSolverReal)0.);
 
-				DrawLine(PDI, Pos0, Pos1, bIsKinematic0 && bIsKinematic1 ? KinematicColor : DynamicColor);
-				DrawLine(PDI, Pos1, Pos2, bIsKinematic1 && bIsKinematic2 ? KinematicColor : DynamicColor);
-				DrawLine(PDI, Pos2, Pos0, bIsKinematic2 && bIsKinematic0 ? KinematicColor : DynamicColor);
+					DrawLine(Context.PDI, Pos0, Pos1, bIsKinematic0 && bIsKinematic1 ? KinematicColor : DynamicColor);
+					DrawLine(Context.PDI, Pos1, Pos2, bIsKinematic1 && bIsKinematic2 ? KinematicColor : DynamicColor);
+					DrawLine(Context.PDI, Pos2, Pos0, bIsKinematic2 && bIsKinematic0 ? KinematicColor : DynamicColor);
+				}
 			}
 		}
 	}
 
-	void FClothVisualizationNoGC::DrawAnimMeshWired(FPrimitiveDrawInterface* PDI) const
+	void FClothVisualizationNoGC::DrawAnimMeshWired(const FDrawContext& Context) const
 	{
 		if (!Solver)
 		{
@@ -1374,6 +1519,12 @@ namespace Chaos
 			}
 			// Elements are local indexed for new solver
 			const int32 Offset = Solver->IsLegacySolver() ? ParticleRangeId : 0;
+			const TSet<int32>* const DebugFilterSet = Context.Filter.VertexSet.IsEmpty() ? nullptr : Cloth->GetVertexSetByName(Solver, Context.Filter.VertexSet);
+			const int32 GlobalParticleOffset = Solver->GetGlobalParticleOffset(ParticleRangeId);
+			auto ShouldDrawElement = [&Context, DebugFilterSet, GlobalParticleOffset, Offset](const TVec3<int32>& Element)
+				{
+					return Private::ShouldDrawElement(Context.Filter, DebugFilterSet, Element, GlobalParticleOffset - Offset, -Offset);
+				};
 
 			const TConstArrayView<TVec3<int32>> Elements = Cloth->GetTriangleMesh(Solver).GetElements();
 			const TConstArrayView<Softs::FSolverVec3> Positions = Cloth->GetAnimationPositions(Solver);
@@ -1382,18 +1533,22 @@ namespace Chaos
 			{
 				const TVec3<int32>& Element = Elements[ElementIndex];
 
-				const FVector Pos0 = GetWorldPosition(*Solver, Positions[Element.X - Offset]); // TODO: Triangle Mesh shouldn't really be solver dependent (ie not use an offset)
-				const FVector Pos1 = GetWorldPosition(*Solver, Positions[Element.Y - Offset]);
-				const FVector Pos2 = GetWorldPosition(*Solver, Positions[Element.Z - Offset]);
+				if (ShouldDrawElement(Element))
+				{
+
+					const FVector Pos0 = GetWorldPosition(*Solver, Positions[Element.X - Offset]); // TODO: Triangle Mesh shouldn't really be solver dependent (ie not use an offset)
+					const FVector Pos1 = GetWorldPosition(*Solver, Positions[Element.Y - Offset]);
+					const FVector Pos2 = GetWorldPosition(*Solver, Positions[Element.Z - Offset]);
 
-				DrawLine(PDI, Pos0, Pos1, KinematicColor);
-				DrawLine(PDI, Pos1, Pos2, KinematicColor);
-				DrawLine(PDI, Pos2, Pos0, KinematicColor);
+					DrawLine(Context.PDI, Pos0, Pos1, KinematicColor);
+					DrawLine(Context.PDI, Pos1, Pos2, KinematicColor);
+					DrawLine(Context.PDI, Pos2, Pos0, KinematicColor);
+				}
 			}
 		}
 	}
 
-	void FClothVisualizationNoGC::DrawAccessoryMesh(FPrimitiveDrawInterface* PDI, const FName& AccessoryMeshName) const
+	void FClothVisualizationNoGC::DrawAccessoryMesh(const FDrawContext& Context, const FName& AccessoryMeshName) const
 	{
 		if (!Solver)
 		{
@@ -1416,6 +1571,12 @@ namespace Chaos
 
 			// Elements are local indexed for new solver
 			const int32 Offset = Solver->IsLegacySolver() ? ParticleRangeId : 0;
+			const TSet<int32>* const DebugFilterSet = Context.Filter.VertexSet.IsEmpty() ? nullptr : Cloth->GetVertexSetByName(Solver, Context.Filter.VertexSet);
+			const int32 GlobalParticleOffset = Solver->GetGlobalParticleOffset(ParticleRangeId);
+			auto ShouldDrawElement = [&Context, &DebugFilterSet, GlobalParticleOffset, Offset](const TVec3<int32>& Element)
+				{
+					return Private::ShouldDrawElement(Context.Filter, DebugFilterSet, Element, GlobalParticleOffset - Offset, -Offset);
+				};
 
 			const TConstArrayView<TVec3<int32>> Elements = Cloth->GetTriangleMesh(Solver).GetElements();
 			const TConstArrayView<Softs::FSolverVec3> Positions = Cloth->GetAccessoryMeshPositions(Solver, AccessoryMeshName);
@@ -1423,20 +1584,23 @@ namespace Chaos
 			for (int32 ElementIndex = 0; ElementIndex < Elements.Num(); ++ElementIndex)
 			{
 				const TVec3<int32>& Element = Elements[ElementIndex];
+				if (ShouldDrawElement(Element))
+				{
 
-				const FVector Pos0 = GetWorldPosition(*Solver, Positions[Element.X - Offset]); // TODO: Triangle Mesh shouldn't really be solver dependent (ie not use an offset)
-				const FVector Pos1 = GetWorldPosition(*Solver, Positions[Element.Y - Offset]);
-				const FVector Pos2 = GetWorldPosition(*Solver, Positions[Element.Z - Offset]);
+					const FVector Pos0 = GetWorldPosition(*Solver, Positions[Element.X - Offset]); // TODO: Triangle Mesh shouldn't really be solver dependent (ie not use an offset)
+					const FVector Pos1 = GetWorldPosition(*Solver, Positions[Element.Y - Offset]);
+					const FVector Pos2 = GetWorldPosition(*Solver, Positions[Element.Z - Offset]);
 
-				DrawLine(PDI, Pos0, Pos1, AccessoryColor);
-				DrawLine(PDI, Pos1, Pos2, AccessoryColor);
-				DrawLine(PDI, Pos2, Pos0, AccessoryColor);
+					DrawLine(Context.PDI, Pos0, Pos1, AccessoryColor);
+					DrawLine(Context.PDI, Pos1, Pos2, AccessoryColor);
+					DrawLine(Context.PDI, Pos2, Pos0, AccessoryColor);
+				}
 			}
 		}
 	}
 
 
-	void FClothVisualizationNoGC::DrawAccessoryMeshNormals(FPrimitiveDrawInterface* PDI, const FName& AccessoryMeshName, const FReal NormalLength) const
+	void FClothVisualizationNoGC::DrawAccessoryMeshNormals(const FDrawContext& Context, const FName& AccessoryMeshName, const FReal NormalLength) const
 	{
 		if (!Solver)
 		{
@@ -1447,7 +1611,8 @@ namespace Chaos
 
 		for (const FClothingSimulationCloth* const Cloth : Solver->GetCloths())
 		{
-			if (Cloth->GetParticleRangeId(Solver) == INDEX_NONE)
+			const int32 ParticleRangeId = Cloth->GetParticleRangeId(Solver);
+			if (ParticleRangeId == INDEX_NONE)
 			{
 				continue;
 			}
@@ -1458,19 +1623,28 @@ namespace Chaos
 
 			const TConstArrayView<Softs::FSolverVec3> Positions = Cloth->GetAccessoryMeshPositions(Solver, AccessoryMeshName);
 			const TConstArrayView<Softs::FSolverVec3> Normals = Cloth->GetAccessoryMeshNormals(Solver, AccessoryMeshName);
+			const TSet<int32>* const DebugFilterSet = Context.Filter.VertexSet.IsEmpty() ? nullptr : Cloth->GetVertexSetByName(Solver, Context.Filter.VertexSet);
+			const int32 Offset = Solver->GetGlobalParticleOffset(ParticleRangeId);
+			auto ShouldDrawParticle = [&Context, DebugFilterSet, Offset](const int32 Index)
+				{
+					return Private::ShouldDrawParticle(Context.Filter, DebugFilterSet, Index, Offset);
+				};
 			check(Normals.Num() == Positions.Num());
 
 			for (int32 Index = 0; Index < Positions.Num(); ++Index)
 			{
-				const FVector Pos0 = GetWorldPosition(*Solver, Positions[Index]);
-				const FVector Pos1 = Pos0 + GetWorldVector(*Solver, Normals[Index]) * NormalLength;
+				if (ShouldDrawParticle(Index))
+				{
+					const FVector Pos0 = GetWorldPosition(*Solver, Positions[Index]);
+					const FVector Pos1 = Pos0 + GetWorldVector(*Solver, Normals[Index]) * NormalLength;
 
-				DrawLine(PDI, Pos0, Pos1, AccessoryColor);
+					DrawLine(Context.PDI, Pos0, Pos1, AccessoryColor);
+				}
 			}
 		}
 	}
 
-	void FClothVisualizationNoGC::DrawOpenEdges(FPrimitiveDrawInterface* PDI) const
+	void FClothVisualizationNoGC::DrawOpenEdges(const FDrawContext& Context) const
 	{
 		auto MakeSortedUintVector2 = [](uint32 Index0, uint32 Index1) -> FUintVector2
 			{
@@ -1515,6 +1689,12 @@ namespace Chaos
 			}
 			// Elements are local indexed for new solver
 			const int32 Offset = Solver->IsLegacySolver() ? ParticleRangeId : 0;
+			const TSet<int32>* const DebugFilterSet = Context.Filter.VertexSet.IsEmpty() ? nullptr : Cloth->GetVertexSetByName(Solver, Context.Filter.VertexSet);
+			const int32 GlobalParticleOffset = Solver->GetGlobalParticleOffset(ParticleRangeId);
+			auto ShouldDrawEdge = [&Context, DebugFilterSet, GlobalParticleOffset, Offset](const  FUintVector2& Edge)
+				{
+					return Private::ShouldDrawEdge(Context.Filter, DebugFilterSet, Edge, GlobalParticleOffset - Offset, -Offset);
+				};
 
 			const TConstArrayView<TVec3<int32>> Elements = Cloth->GetTriangleMesh(Solver).GetElements();
 			const TConstArrayView<Softs::FSolverVec3> Positions = Cloth->GetAnimationPositions(Solver);
@@ -1525,18 +1705,21 @@ namespace Chaos
 			for (const TPair<FUintVector2, TArray<uint32>>& EdgeToTriangles : EdgeToTrianglesMap)
 			{
 				const FUintVector2& Edge = EdgeToTriangles.Key;
-				const TArray<uint32>& Triangles = EdgeToTriangles.Value;
+				if (ShouldDrawEdge(Edge))
+				{
+					const TArray<uint32>& Triangles = EdgeToTriangles.Value;
 
-				const FVector Pos0 = GetWorldPosition(*Solver, Positions[Edge[0] - Offset]); // TODO: Triangle Mesh shouldn't really be solver dependent (ie not use an offset)
-				const FVector Pos1 = GetWorldPosition(*Solver, Positions[Edge[1] - Offset]);
-				const FLinearColor& Color = (Triangles.Num() > 1) ? ClosedEdgeColor : OpenedEdgeColor;
+					const FVector Pos0 = GetWorldPosition(*Solver, Positions[Edge[0] - Offset]); // TODO: Triangle Mesh shouldn't really be solver dependent (ie not use an offset)
+					const FVector Pos1 = GetWorldPosition(*Solver, Positions[Edge[1] - Offset]);
+					const FLinearColor& Color = (Triangles.Num() > 1) ? ClosedEdgeColor : OpenedEdgeColor;
 
-				DrawLine(PDI, Pos0, Pos1, Color);
+					DrawLine(Context.PDI, Pos0, Pos1, Color);
+				}
 			}
 		}
 	}
 
-	void FClothVisualizationNoGC::DrawMultiResConstraint(FPrimitiveDrawInterface* PDI) const
+	void FClothVisualizationNoGC::DrawMultiResConstraint(const FDrawContext& Context) const
 	{
 		if (!Solver)
 		{
@@ -1575,9 +1758,9 @@ namespace Chaos
 						continue;
 					}
 
-					DrawLine(PDI, Pos0, Pos1, bIsKinematic0 && bIsKinematic1 ? KinematicColor : DynamicColor);
-					DrawLine(PDI, Pos1, Pos2, bIsKinematic1 && bIsKinematic2 ? KinematicColor : DynamicColor);
-					DrawLine(PDI, Pos2, Pos0, bIsKinematic2 && bIsKinematic0 ? KinematicColor : DynamicColor);
+					DrawLine(Context.PDI, Pos0, Pos1, bIsKinematic0 && bIsKinematic1 ? KinematicColor : DynamicColor);
+					DrawLine(Context.PDI, Pos1, Pos2, bIsKinematic1 && bIsKinematic2 ? KinematicColor : DynamicColor);
+					DrawLine(Context.PDI, Pos2, Pos0, bIsKinematic2 && bIsKinematic0 ? KinematicColor : DynamicColor);
 				}
 
 				// Draw springs to targets
@@ -1593,15 +1776,15 @@ namespace Chaos
 						const FVector P1 = GetWorldPosition(*Solver, Positions[Index]);
 						const FVector P2 = GetWorldPosition(*Solver, TargetPositions[Index]);
 
-						DrawPoint(PDI, P2, Red, nullptr, 2.f);
-						DrawLine(PDI, P1, P2, Brown);
+						DrawPoint(Context.PDI, P2, Red, nullptr, 2.f);
+						DrawLine(Context.PDI, P1, P2, Brown);
 					}
 				}
 			}
 		}
 	}
 
-	void FClothVisualizationNoGC::DrawAnimNormals(FPrimitiveDrawInterface* PDI, const FReal NormalLength) const
+	void FClothVisualizationNoGC::DrawAnimNormals(const FDrawContext& Context, const FReal Length) const
 	{
 		if (!Solver)
 		{
@@ -1612,26 +1795,36 @@ namespace Chaos
 
 		for (const FClothingSimulationCloth* const Cloth : Solver->GetCloths())
 		{
-			if (Cloth->GetParticleRangeId(Solver) == INDEX_NONE)
+			const int32 ParticleRangeId = Cloth->GetParticleRangeId(Solver);
+			if (ParticleRangeId == INDEX_NONE)
 			{
 				continue;
 			}
 
+			const TSet<int32>* const DebugFilterSet = Context.Filter.VertexSet.IsEmpty() ? nullptr : Cloth->GetVertexSetByName(Solver, Context.Filter.VertexSet);
+			const int32 GlobalParticleOffset = Solver->GetGlobalParticleOffset(ParticleRangeId);
+			auto ShouldDrawParticle = [&Context, DebugFilterSet, GlobalParticleOffset](const int32 Index)
+				{
+					return Private::ShouldDrawParticle(Context.Filter, DebugFilterSet, Index, GlobalParticleOffset);
+				};
 			const TConstArrayView<Softs::FSolverVec3> Positions = Cloth->GetAnimationPositions(Solver);
 			const TConstArrayView<Softs::FSolverVec3> Normals = Cloth->GetAnimationNormals(Solver);
 			check(Normals.Num() == Positions.Num());
 
 			for (int32 Index = 0; Index < Positions.Num(); ++Index)
 			{
-				const FVector Pos0 = GetWorldPosition(*Solver, Positions[Index]);
-				const FVector Pos1 = Pos0 + GetWorldVector(*Solver, Normals[Index]) * NormalLength;
+				if (ShouldDrawParticle(Index))
+				{
+					const FVector Pos0 = GetWorldPosition(*Solver, Positions[Index]);
+					const FVector Pos1 = Pos0 + GetWorldVector(*Solver, Normals[Index]) * Length;
 
-				DrawLine(PDI, Pos0, Pos1, KinematicColor);
+					DrawLine(Context.PDI, Pos0, Pos1, KinematicColor);
+				}
 			}
 		}
 	}
 
-	void FClothVisualizationNoGC::DrawAnimVelocities(FPrimitiveDrawInterface* PDI) const
+	void FClothVisualizationNoGC::DrawAnimVelocities(const FDrawContext& Context) const
 	{
 		if (!Solver)
 		{
@@ -1649,35 +1842,51 @@ namespace Chaos
 				continue;
 			}
 
+			const TSet<int32>* const DebugFilterSet = Context.Filter.VertexSet.IsEmpty() ? nullptr : Cloth->GetVertexSetByName(Solver, Context.Filter.VertexSet);
+			const int32 GlobalParticleOffset = Solver->GetGlobalParticleOffset(ParticleRangeId);
+			auto ShouldDrawParticle = [&Context, &DebugFilterSet, GlobalParticleOffset](const int32 Index)
+				{
+					return Private::ShouldDrawParticle(Context.Filter, DebugFilterSet, Index, GlobalParticleOffset);
+				};
 			const TConstArrayView<Softs::FSolverVec3> Positions = Cloth->GetAnimationPositions(Solver);
 			const TConstArrayView<Softs::FSolverVec3> OldPositions = Cloth->GetOldAnimationPositions(Solver);
 			checkSlow(OldPositions.Num() == Positions.Num());
 
 			for (int32 Index = 0; Index < Positions.Num(); ++Index)
 			{
-				const FVector Pos0 = GetWorldPosition(*Solver, Positions[Index]);
-				const FVector Pos1 = GetWorldPosition(*Solver, OldPositions[Index]);
-				DrawLine(PDI, Pos0, Pos1, KinematicColor);
+				if (ShouldDrawParticle(Index))
+				{
+					const FVector Pos0 = GetWorldPosition(*Solver, Positions[Index]);
+					const FVector Pos1 = GetWorldPosition(*Solver, OldPositions[Index]);
+					DrawLine(Context.PDI, Pos0, Pos1, KinematicColor);
+				}
 			}
 
 			const int32 Offset = Solver->IsLegacySolver() ? ParticleRangeId : 0;
+			auto ShouldDrawElement = [&Context, DebugFilterSet, GlobalParticleOffset, Offset](const TVec3<int32>& Element)
+				{
+					return Private::ShouldDrawElement(Context.Filter, DebugFilterSet, Element, GlobalParticleOffset - Offset, -Offset);
+				};
 			const TConstArrayView<TVec3<int32>> Elements = Cloth->GetTriangleMesh(Solver).GetElements();
 			for (int32 ElementIndex = 0; ElementIndex < Elements.Num(); ++ElementIndex)
 			{
 				const TVec3<int32>& Element = Elements[ElementIndex];
+				if (ShouldDrawElement(Element))
+				{
 
-				const FVector Pos0 = GetWorldPosition(*Solver, OldPositions[Element.X - Offset]); // TODO: Triangle Mesh shouldn't really be solver dependent (ie not use an offset)
-				const FVector Pos1 = GetWorldPosition(*Solver, OldPositions[Element.Y - Offset]);
-				const FVector Pos2 = GetWorldPosition(*Solver, OldPositions[Element.Z - Offset]);
+					const FVector Pos0 = GetWorldPosition(*Solver, OldPositions[Element.X - Offset]); // TODO: Triangle Mesh shouldn't really be solver dependent (ie not use an offset)
+					const FVector Pos1 = GetWorldPosition(*Solver, OldPositions[Element.Y - Offset]);
+					const FVector Pos2 = GetWorldPosition(*Solver, OldPositions[Element.Z - Offset]);
 
-				DrawLine(PDI, Pos0, Pos1, KinematicMeshColor);
-				DrawLine(PDI, Pos1, Pos2, KinematicMeshColor);
-				DrawLine(PDI, Pos2, Pos0, KinematicMeshColor);
+					DrawLine(Context.PDI, Pos0, Pos1, KinematicMeshColor);
+					DrawLine(Context.PDI, Pos1, Pos2, KinematicMeshColor);
+					DrawLine(Context.PDI, Pos2, Pos0, KinematicMeshColor);
+				}
 			}
 		}
 	}
 
-	void FClothVisualizationNoGC::DrawPointNormals(FPrimitiveDrawInterface* PDI, const FReal NormalLength) const
+	void FClothVisualizationNoGC::DrawPointNormals(const FDrawContext& Context, const FReal NormalLength) const
 	{
 		if (!Solver)
 		{
@@ -1689,7 +1898,8 @@ namespace Chaos
 
 		for (const FClothingSimulationCloth* const Cloth : Solver->GetCloths())
 		{
-			if (Cloth->GetParticleRangeId(Solver) == INDEX_NONE)
+			const int32 ParticleRangeId = Cloth->GetParticleRangeId(Solver);
+			if (ParticleRangeId == INDEX_NONE)
 			{
 				continue;
 			}
@@ -1697,21 +1907,30 @@ namespace Chaos
 			const TConstArrayView<Softs::FSolverVec3> Positions = Cloth->GetParticlePositions(Solver);
 			const TConstArrayView<Softs::FSolverVec3> Normals = Cloth->GetParticleNormals(Solver);
 			const TConstArrayView<Softs::FSolverReal> InvMasses = Cloth->GetParticleInvMasses(Solver);
+			const TSet<int32>* const DebugFilterSet = Context.Filter.VertexSet.IsEmpty() ? nullptr : Cloth->GetVertexSetByName(Solver, Context.Filter.VertexSet);
+			const int32 Offset = Solver->GetGlobalParticleOffset(ParticleRangeId);
+			auto ShouldDrawParticle = [&Context, DebugFilterSet, Offset](const int32 Index)
+				{
+					return Private::ShouldDrawParticle(Context.Filter, DebugFilterSet, Index, Offset);
+				};
 			check(Normals.Num() == Positions.Num());
 			check(InvMasses.Num() == Positions.Num());
 
 			for (int32 Index = 0; Index < Positions.Num(); ++Index)
 			{
-				const bool bIsKinematic = (InvMasses[Index] == (Softs::FSolverReal)0.);
-				const FVector Pos0 = GetWorldPosition(*Solver, Positions[Index]);
-				const FVector Pos1 = Pos0 + GetWorldVector(*Solver, Normals[Index]) * NormalLength;
+				if (ShouldDrawParticle(Index))
+				{
+					const bool bIsKinematic = (InvMasses[Index] == (Softs::FSolverReal)0.);
+					const FVector Pos0 = GetWorldPosition(*Solver, Positions[Index]);
+					const FVector Pos1 = Pos0 + GetWorldVector(*Solver, Normals[Index]) * NormalLength;
 
-				DrawLine(PDI, Pos0, Pos1, bIsKinematic ? KinematicColor : DynamicColor);
+					DrawLine(Context.PDI, Pos0, Pos1, bIsKinematic ? KinematicColor : DynamicColor);
+				}
 			}
 		}
 	}
 
-	void FClothVisualizationNoGC::DrawPointVelocities(FPrimitiveDrawInterface* PDI) const
+	void FClothVisualizationNoGC::DrawPointVelocities(const FDrawContext& Context) const
 	{
 		if (!Solver)
 		{
@@ -1720,13 +1939,20 @@ namespace Chaos
 
 		for (const FClothingSimulationCloth* const Cloth : Solver->GetCloths())
 		{
-			if (Cloth->GetParticleRangeId(Solver) == INDEX_NONE)
+			const int32 ParticleRangeId = Cloth->GetParticleRangeId(Solver);
+			if (ParticleRangeId == INDEX_NONE)
 			{
 				continue;
 			}
 
 			const TConstArrayView<Softs::FSolverVec3> Positions = Cloth->GetParticlePositions(Solver);
 			const TConstArrayView<Softs::FSolverVec3> Velocities = Cloth->GetParticleVelocities(Solver);
+			const TSet<int32>* const DebugFilterSet = Context.Filter.VertexSet.IsEmpty() ? nullptr : Cloth->GetVertexSetByName(Solver, Context.Filter.VertexSet);
+			const int32 Offset = Solver->GetGlobalParticleOffset(ParticleRangeId);
+			auto ShouldDrawParticle = [&Context, DebugFilterSet, Offset](const int32 Index)
+				{
+					return Private::ShouldDrawParticle(Context.Filter, DebugFilterSet, Index, Offset);
+				};
 			check(Velocities.Num() == Positions.Num());
 
 			const TConstArrayView<Softs::FSolverReal> InvMasses = Cloth->GetParticleInvMasses(Solver);
@@ -1734,18 +1960,21 @@ namespace Chaos
 
 			for (int32 Index = 0; Index < Positions.Num(); ++Index)
 			{
-				constexpr FReal DefaultFPS = (FReal)60.;   // TODO: A CVAR would be nice for this
-				const bool bIsKinematic = (InvMasses[Index] == 0.f);
+				if (ShouldDrawParticle(Index))
+				{
+					constexpr FReal DefaultFPS = (FReal)60.;   // TODO: A CVAR would be nice for this
+					const bool bIsKinematic = (InvMasses[Index] == 0.f);
 
-				const FVec3 Pos0 = GetWorldPosition(*Solver, Positions[Index]);
-				const FVec3 Pos1 = Pos0 + GetWorldVector(*Solver, Velocities[Index]) / DefaultFPS;  // Velocity per frame if running at DefaultFPS
+					const FVec3 Pos0 = GetWorldPosition(*Solver, Positions[Index]);
+					const FVec3 Pos1 = Pos0 + GetWorldVector(*Solver, Velocities[Index]) / DefaultFPS;  // Velocity per frame if running at DefaultFPS
 
-				DrawLine(PDI, Pos0, Pos1, bIsKinematic ? FLinearColor::Black : FLinearColor::Yellow);
+					DrawLine(Context.PDI, Pos0, Pos1, bIsKinematic ? FLinearColor::Black : FLinearColor::Yellow);
+				}
 			}
 		}
 	}
 
-	void FClothVisualizationNoGC::DrawCollision(FPrimitiveDrawInterface* PDI, bool bWireframe) const
+	void FClothVisualizationNoGC::DrawCollision(const FDrawContext& Context, bool bWireframe) const
 	{
 		if (!Solver)
 		{
@@ -1753,7 +1982,7 @@ namespace Chaos
 		}
 
 		auto DrawCollision =
-			[this, PDI, bWireframe](const FClothingSimulationCollider* Collider, const FClothingSimulationCloth* Cloth, FClothingSimulationCollider::ECollisionDataType CollisionDataType)
+			[this, &Context, bWireframe](const FClothingSimulationCollider* Collider, const FClothingSimulationCloth* Cloth, FClothingSimulationCollider::ECollisionDataType CollisionDataType)
 			{
 				static const FLinearColor GlobalColor(FColor::Cyan);
 				static const FLinearColor DynamicColor(FColor::Orange);
@@ -1782,15 +2011,15 @@ namespace Chaos
 						switch (Object->GetType())
 						{
 						case ImplicitObjectType::Sphere:
-							DrawSphere(PDI, Object->GetObjectChecked<FSphere>(), Rotation, Position, Solver->GetLocalSpaceScale(), Color);
+							DrawSphere(Context.PDI, Object->GetObjectChecked<FSphere>(), Rotation, Position, Solver->GetLocalSpaceScale(), Color);
 							break;
 
 						case ImplicitObjectType::Box:
-							DrawBox(PDI, Object->GetObjectChecked<TBox<FReal, 3>>().BoundingBox(), Rotation, Position, Solver->GetLocalSpaceScale(), Color);
+							DrawBox(Context.PDI, Object->GetObjectChecked<TBox<FReal, 3>>().BoundingBox(), Rotation, Position, Solver->GetLocalSpaceScale(), Color);
 							break;
 
 						case ImplicitObjectType::Capsule:
-							DrawCapsule(PDI, Object->GetObjectChecked<FCapsule>(), Rotation, Position, Solver->GetLocalSpaceScale(), Color);
+							DrawCapsule(Context.PDI, Object->GetObjectChecked<FCapsule>(), Rotation, Position, Solver->GetLocalSpaceScale(), Color);
 							break;
 
 						case ImplicitObjectType::Union:  // Union only used as old style tapered capsules
@@ -1801,11 +2030,11 @@ namespace Chaos
 									switch (SubObject->GetType())
 									{
 									case ImplicitObjectType::Sphere:
-										DrawSphere(PDI, SubObject->GetObjectChecked<FSphere>(), Rotation, Position, Solver->GetLocalSpaceScale(), Color);
+										DrawSphere(Context.PDI, SubObject->GetObjectChecked<FSphere>(), Rotation, Position, Solver->GetLocalSpaceScale(), Color);
 										break;
 
 									case ImplicitObjectType::TaperedCylinder:
-										DrawTaperedCylinder(PDI, SubObject->GetObjectChecked<FTaperedCylinder>(), Rotation, Position, Solver->GetLocalSpaceScale(), Color);
+										DrawTaperedCylinder(Context.PDI, SubObject->GetObjectChecked<FTaperedCylinder>(), Rotation, Position, Solver->GetLocalSpaceScale(), Color);
 										break;
 
 									default:
@@ -1822,20 +2051,20 @@ namespace Chaos
 								const FVec3 X2 = TaperedCapsule.GetX2();
 								const FReal Radius1 = TaperedCapsule.GetRadius1();
 								const FReal Radius2 = TaperedCapsule.GetRadius2();
-								DrawSphere(PDI, FSphere(X1, Radius1), Rotation, Position, Solver->GetLocalSpaceScale(), Color);
-								DrawSphere(PDI, FSphere(X2, Radius2), Rotation, Position, Solver->GetLocalSpaceScale(), Color);
-								DrawTaperedCylinder(PDI, FTaperedCylinder(X1, X2, Radius1, Radius2), Rotation, Position, Solver->GetLocalSpaceScale(), Color);
+								DrawSphere(Context.PDI, FSphere(X1, Radius1), Rotation, Position, Solver->GetLocalSpaceScale(), Color);
+								DrawSphere(Context.PDI, FSphere(X2, Radius2), Rotation, Position, Solver->GetLocalSpaceScale(), Color);
+								DrawTaperedCylinder(Context.PDI, FTaperedCylinder(X1, X2, Radius1, Radius2), Rotation, Position, Solver->GetLocalSpaceScale(), Color);
 								if (TaperedCapsule.IsOneSided())
 								{
 									const FVec3 Center = Position + Rotation.RotateVector(Solver->GetLocalSpaceScale() * .5 * (X1 + X2));
 									const FVec3 NormalEnd = Center + Rotation.RotateVector(Solver->GetLocalSpaceScale() * .5 * (Radius1 + Radius2) * TaperedCapsule.GetOneSidedPlaneNormal());
-									DrawLine(PDI, Center, NormalEnd, Color);
+									DrawLine(Context.PDI, Center, NormalEnd, Color);
 								}
 							}
 							break;
 
 						case ImplicitObjectType::Convex:
-							DrawConvex(PDI, Object->GetObjectChecked<FConvex>(), Rotation, Position, Solver->GetLocalSpaceScale(), Color);
+							DrawConvex(Context.PDI, Object->GetObjectChecked<FConvex>(), Rotation, Position, Solver->GetLocalSpaceScale(), Color);
 							break;
 
 						case ImplicitObjectType::Transformed: // Transformed only used for levelsets
@@ -1850,7 +2079,7 @@ namespace Chaos
 #else
 									nullptr;
 #endif
-								DrawLevelSet(PDI, CombinedTransform, MaterialRenderProxy, LevelSet);
+								DrawLevelSet(Context.PDI, CombinedTransform, MaterialRenderProxy, LevelSet);
 							}
 							break;
 						case (ImplicitObjectType::LevelSet | ImplicitObjectType::IsWeightedLattice):
@@ -1862,7 +2091,7 @@ namespace Chaos
 #else
 									nullptr;
 #endif
-								DrawSkinnedLevelSet(PDI, WeightedLevelset, Rotation, Position, Solver->GetLocalSpaceScale(), MaterialRenderProxy);
+								DrawSkinnedLevelSet(Context.PDI, WeightedLevelset, Rotation, Position, Solver->GetLocalSpaceScale(), MaterialRenderProxy);
 
 							}
 							break;
@@ -1876,7 +2105,7 @@ namespace Chaos
 #else
 									nullptr;
 #endif
-								DrawMLLevelSet(PDI, CombinedTransform, MaterialRenderProxy, MLLevelSet);
+								DrawMLLevelSet(Context.PDI, CombinedTransform, MaterialRenderProxy, MLLevelSet);
 							}
 							break;
 						case ImplicitObjectType::SkinnedTriangleMesh:
@@ -1888,18 +2117,18 @@ namespace Chaos
 #else
 								nullptr;
 #endif
-							DrawSkinnedTriangleMesh(PDI, MaterialRenderProxy, Solver, Color, Mesh);
+							DrawSkinnedTriangleMesh(Context.PDI, MaterialRenderProxy, Solver, Color, Mesh);
 
 						}
 						break;
 						case ImplicitObjectType::ExtrudedTaperedCapsule:
 						{
 							const FExtrudedTaperedCapsule& Capsule = Object->GetObjectChecked<FExtrudedTaperedCapsule>();
-							DrawExtrudedTaperedCapsule(PDI, Capsule, Rotation, Position, Solver->GetLocalSpaceScale(), Color);
+							DrawExtrudedTaperedCapsule(Context.PDI, Capsule, Rotation, Position, Solver->GetLocalSpaceScale(), Color);
 						}
 						break;
 						default:
-							DrawCoordinateSystem(PDI, Rotation, Position);  // Draw everything else as a coordinate for now
+							DrawCoordinateSystem(Context.PDI, Rotation, Position);  // Draw everything else as a coordinate for now
 							break;
 						}
 					}
@@ -1931,40 +2160,40 @@ namespace Chaos
 			FVec3 TangentU, TangentV;
 			Normal.FindBestAxisVectors(TangentU, TangentV);
 
-			DrawLine(PDI, Pos0 + TangentU, Pos0 + TangentV, FLinearColor::Black);
-			DrawLine(PDI, Pos0 + TangentU, Pos0 - TangentV, FLinearColor::Black);
-			DrawLine(PDI, Pos0 - TangentU, Pos0 - TangentV, FLinearColor::Black);
-			DrawLine(PDI, Pos0 - TangentU, Pos0 + TangentV, FLinearColor::Black);
+			DrawLine(Context.PDI, Pos0 + TangentU, Pos0 + TangentV, FLinearColor::Black);
+			DrawLine(Context.PDI, Pos0 + TangentU, Pos0 - TangentV, FLinearColor::Black);
+			DrawLine(Context.PDI, Pos0 - TangentU, Pos0 - TangentV, FLinearColor::Black);
+			DrawLine(Context.PDI, Pos0 - TangentU, Pos0 + TangentV, FLinearColor::Black);
 
 			// Draw normal
 			static const FLinearColor Brown(0.1f, 0.05f, 0.f);
 			const FVec3 Pos1 = Pos0 + NormalLength * Normal;
-			DrawLine(PDI, Pos0, Pos1, Brown);
+			DrawLine(Context.PDI, Pos0, Pos1, Brown);
 
 			if (bDrawPhis)
 			{
 				const FVec3 PhiLocation = Pos0 - Solver->GetCollisionPhis()[i] * Normal;
-				DrawLine(PDI, Pos0, PhiLocation, Brown);
-				DrawPoint(PDI, PhiLocation, FLinearColor::Red, nullptr, 5);
+				DrawLine(Context.PDI, Pos0, PhiLocation, Brown);
+				DrawPoint(Context.PDI, PhiLocation, FLinearColor::Red, nullptr, 5);
 			}
 		}
 	}
 
-	void FClothVisualizationNoGC::DrawBackstops(FPrimitiveDrawInterface* PDI) const
+	void FClothVisualizationNoGC::DrawBackstops(const FDrawContext& Context) const
 	{
 		if (!Solver)
 		{
 			return;
 		}
 
-		auto DrawBackstop = [PDI](const FVector& Position, const FVector& Normal, FReal Radius, const FVector& Axis, const FLinearColor& Color)
+		auto DrawBackstop = [&Context](const FVector& Position, const FVector& Normal, FReal Radius, const FVector& Axis, const FLinearColor& Color)
 			{
 				static const FReal MaxCosAngle = (FReal)0.99;
 				if (FMath::Abs(FVector::DotProduct(Normal, Axis)) < MaxCosAngle)
 				{
 					static const FReal ArcLength = (FReal)5.; // Arch length in cm
 					const FReal ArcAngle = (FReal)360. * ArcLength / FMath::Max((Radius * (FReal)2. * (FReal)PI), ArcLength);
-					DrawArc(PDI, Position, Normal, FVector::CrossProduct(Axis, Normal).GetSafeNormal(), -ArcAngle / (FReal)2., ArcAngle / (FReal)2., Radius, Color);
+					DrawArc(Context.PDI, Position, Normal, FVector::CrossProduct(Axis, Normal).GetSafeNormal(), -ArcAngle / (FReal)2., ArcAngle / (FReal)2., Radius, Color);
 				}
 			};
 
@@ -1988,64 +2217,74 @@ namespace Chaos
 				const TConstArrayView<Softs::FSolverVec3> AnimationPositions = AccessoryMeshExists ? Cloth->GetAccessoryMeshPositions(Solver, BackstopMeshName) : Cloth->GetAnimationPositions(Solver);
 				const TConstArrayView<Softs::FSolverVec3> AnimationNormals = AccessoryMeshExists ? Cloth->GetAccessoryMeshNormals(Solver, BackstopMeshName) : Cloth->GetAnimationNormals(Solver);
 				const TConstArrayView<Softs::FSolverVec3> ParticlePositions = Cloth->GetParticlePositions(Solver);
+				const int32 Offset = Solver->GetGlobalParticleOffset(ParticleRangeId);
 
-				if (AnimationPositions.IsValidIndex(Private::BackstopVertexIndex))
+				if (Context.Filter.SingleVertex != INDEX_NONE && AnimationPositions.IsValidIndex(Context.Filter.SingleVertex-Offset))
 				{
 					const FLinearColor ColorLight = FLinearColor::MakeFromHSV8(ColorSeed, 160, 128);
 					const FLinearColor ColorDark = FLinearColor::MakeFromHSV8(ColorSeed, 160, 64);
+					const int32 BackstopVertexIndex = Context.Filter.SingleVertex - Offset;
 
-					const FReal BackstopRadius = BackstopConstraint->GetBackstopRadius(Private::BackstopVertexIndex) * BackstopConstraint->GetScale() * Solver->GetLocalSpaceScale();
-					const FReal BackstopDistance = BackstopConstraint->GetBackstopDistance(Private::BackstopVertexIndex) * BackstopConstraint->GetScale() * Solver->GetLocalSpaceScale();
+					const FReal BackstopRadius = BackstopConstraint->GetBackstopRadius(BackstopVertexIndex) * BackstopConstraint->GetScale() * Solver->GetLocalSpaceScale();
+					const FReal BackstopDistance = BackstopConstraint->GetBackstopDistance(BackstopVertexIndex) * BackstopConstraint->GetScale() * Solver->GetLocalSpaceScale();
 
-					const FVector AnimationNormal(AnimationNormals[Private::BackstopVertexIndex]);
+					const FVector AnimationNormal(AnimationNormals[BackstopVertexIndex]);
 
 					// Draw a line to show the current distance to the sphere
-					const FVector Pos0 = GetWorldPosition(*Solver, AnimationPositions[Private::BackstopVertexIndex]);
+					const FVector Pos0 = GetWorldPosition(*Solver, AnimationPositions[BackstopVertexIndex]);
 					const FVector Pos1 = Pos0 - (bUseLegacyBackstop ? BackstopDistance - BackstopRadius : BackstopDistance) * AnimationNormal;
-					const FVector Pos2 = GetWorldPosition(*Solver, ParticlePositions[Private::BackstopVertexIndex]);
-					DrawPoint(PDI, Pos2, ColorLight, nullptr, 5.f);
-					DrawLine(PDI, Pos1, Pos2, ColorLight);
+					const FVector Pos2 = GetWorldPosition(*Solver, ParticlePositions[BackstopVertexIndex]);
+					DrawPoint(Context.PDI, Pos2, ColorLight, nullptr, 5.f);
+					DrawLine(Context.PDI, Pos1, Pos2, ColorLight);
 
 					// Draw the sphere
 					if (BackstopRadius > 0.f)
 					{
 						const FVector Center = Pos0 - (bUseLegacyBackstop ? BackstopDistance : BackstopRadius + BackstopDistance) * AnimationNormal;
-						DrawLine(PDI, Pos1, Center, ColorDark);
+						DrawLine(Context.PDI, Pos1, Center, ColorDark);
 						const FMaterialRenderProxy* MaterialRenderProxy =
 #if WITH_EDITOR
 							FMaterials::GetInstance().GetCollisionMaterial() ? FMaterials::GetInstance().GetCollisionMaterial()->GetRenderProxy() : nullptr;
 #else
 							nullptr;
 #endif
-						DrawSphere(PDI, FSphere(FVector::ZeroVector, BackstopRadius), FQuat::Identity, Center, 1.f, ColorDark, MaterialRenderProxy);
+						DrawSphere(Context.PDI, FSphere(FVector::ZeroVector, BackstopRadius), FQuat::Identity, Center, 1.f, ColorDark, MaterialRenderProxy);
 					}
 				}
 				else
 				{
+					const TSet<int32>* const DebugFilterSet = Context.Filter.VertexSet.IsEmpty() ? nullptr : Cloth->GetVertexSetByName(Solver, Context.Filter.VertexSet);
+					auto ShouldDrawParticle = [&Context, &DebugFilterSet, Offset](const int32 Index)
+						{
+							return Private::ShouldDrawParticle(Context.Filter, DebugFilterSet, Index, Offset);
+						};
 					for (int32 Index = 0; Index < AnimationPositions.Num(); ++Index)
 					{
-						ColorSeed += 157;  // Prime number that gives a good spread of colors without getting too similar as a rand might do.
-						const FLinearColor ColorLight = FLinearColor::MakeFromHSV8(ColorSeed, 160, 128);
-						const FLinearColor ColorDark = FLinearColor::MakeFromHSV8(ColorSeed, 160, 64);
+						if (ShouldDrawParticle(Index))
+						{
+							ColorSeed += 157;  // Prime number that gives a good spread of colors without getting too similar as a rand might do.
+							const FLinearColor ColorLight = FLinearColor::MakeFromHSV8(ColorSeed, 160, 128);
+							const FLinearColor ColorDark = FLinearColor::MakeFromHSV8(ColorSeed, 160, 64);
 
-						const FReal BackstopRadius = BackstopConstraint->GetBackstopRadius(Index) * BackstopConstraint->GetScale() * Solver->GetLocalSpaceScale();
-						const FReal BackstopDistance = BackstopConstraint->GetBackstopDistance(Index) * BackstopConstraint->GetScale() * Solver->GetLocalSpaceScale();
+							const FReal BackstopRadius = BackstopConstraint->GetBackstopRadius(Index) * BackstopConstraint->GetScale() * Solver->GetLocalSpaceScale();
+							const FReal BackstopDistance = BackstopConstraint->GetBackstopDistance(Index) * BackstopConstraint->GetScale() * Solver->GetLocalSpaceScale();
 
-						const FVector AnimationNormal(AnimationNormals[Index]);
+							const FVector AnimationNormal(AnimationNormals[Index]);
 
-						// Draw a line to show the current distance to the sphere
-						const FVector Pos0 = GetWorldPosition(*Solver, AnimationPositions[Index]);
-						const FVector Pos1 = Pos0 - (bUseLegacyBackstop ? BackstopDistance - BackstopRadius : BackstopDistance) * AnimationNormal;
-						const FVector Pos2 = GetWorldPosition(*Solver, ParticlePositions[Index]);
-						DrawLine(PDI, Pos1, Pos2, ColorLight);
+							// Draw a line to show the current distance to the sphere
+							const FVector Pos0 = GetWorldPosition(*Solver, AnimationPositions[Index]);
+							const FVector Pos1 = Pos0 - (bUseLegacyBackstop ? BackstopDistance - BackstopRadius : BackstopDistance) * AnimationNormal;
+							const FVector Pos2 = GetWorldPosition(*Solver, ParticlePositions[Index]);
+							DrawLine(Context.PDI, Pos1, Pos2, ColorLight);
 
-						// Draw the sphere
-						if (BackstopRadius > 0.f)
-						{
-							const FVector Center = Pos0 - (bUseLegacyBackstop ? BackstopDistance : BackstopRadius + BackstopDistance) * AnimationNormal;
-							DrawBackstop(Center, AnimationNormal, BackstopRadius, FVector::ForwardVector, ColorDark);
-							DrawBackstop(Center, AnimationNormal, BackstopRadius, FVector::UpVector, ColorDark);
-							DrawBackstop(Center, AnimationNormal, BackstopRadius, FVector::RightVector, ColorDark);
+							// Draw the sphere
+							if (BackstopRadius > 0.f)
+							{
+								const FVector Center = Pos0 - (bUseLegacyBackstop ? BackstopDistance : BackstopRadius + BackstopDistance) * AnimationNormal;
+								DrawBackstop(Center, AnimationNormal, BackstopRadius, FVector::ForwardVector, ColorDark);
+								DrawBackstop(Center, AnimationNormal, BackstopRadius, FVector::UpVector, ColorDark);
+								DrawBackstop(Center, AnimationNormal, BackstopRadius, FVector::RightVector, ColorDark);
+							}
 						}
 					}
 				}
@@ -2053,7 +2292,7 @@ namespace Chaos
 		}
 	}
 
-	void FClothVisualizationNoGC::DrawBackstopDistances(FPrimitiveDrawInterface* PDI) const
+	void FClothVisualizationNoGC::DrawBackstopDistances(const FDrawContext& Context) const
 	{
 		if (!Solver)
 		{
@@ -2079,28 +2318,37 @@ namespace Chaos
 
 				const TConstArrayView<Softs::FSolverVec3> AnimationPositions = AccessoryMeshExists ? Cloth->GetAccessoryMeshPositions(Solver, BackstopMeshName) : Cloth->GetAnimationPositions(Solver);
 				const TConstArrayView<Softs::FSolverVec3> AnimationNormals = AccessoryMeshExists ? Cloth->GetAccessoryMeshNormals(Solver, BackstopMeshName) : Cloth->GetAnimationNormals(Solver);
+				const TSet<int32>* const DebugFilterSet = Context.Filter.VertexSet.IsEmpty() ? nullptr : Cloth->GetVertexSetByName(Solver, Context.Filter.VertexSet);
+				const int32 Offset = Solver->GetGlobalParticleOffset(ParticleRangeId);
+				auto ShouldDrawParticle = [&Context, DebugFilterSet, Offset](const int32 Index)
+					{
+						return Private::ShouldDrawParticle(Context.Filter, DebugFilterSet, Index, Offset);
+					};
 
 				for (int32 Index = 0; Index < AnimationPositions.Num(); ++Index)
 				{
-					ColorSeed += 157;  // Prime number that gives a good spread of colors without getting too similar as a rand might do.
-					const FLinearColor ColorLight = FLinearColor::MakeFromHSV8(ColorSeed, 160, 128);
-					const FLinearColor ColorDark = FLinearColor::MakeFromHSV8(ColorSeed, 160, 64);
+					if (ShouldDrawParticle(Index))
+					{
+						ColorSeed += 157;  // Prime number that gives a good spread of colors without getting too similar as a rand might do.
+						const FLinearColor ColorLight = FLinearColor::MakeFromHSV8(ColorSeed, 160, 128);
+						const FLinearColor ColorDark = FLinearColor::MakeFromHSV8(ColorSeed, 160, 64);
 
-					const FReal BackstopRadius = BackstopConstraint->GetBackstopRadius(Index) * BackstopConstraint->GetScale() * Solver->GetLocalSpaceScale();
-					const FReal BackstopDistance = BackstopConstraint->GetBackstopDistance(Index) * BackstopConstraint->GetScale() * Solver->GetLocalSpaceScale();
+						const FReal BackstopRadius = BackstopConstraint->GetBackstopRadius(Index) * BackstopConstraint->GetScale() * Solver->GetLocalSpaceScale();
+						const FReal BackstopDistance = BackstopConstraint->GetBackstopDistance(Index) * BackstopConstraint->GetScale() * Solver->GetLocalSpaceScale();
 
-					const FVector AnimationNormal(AnimationNormals[Index]);
+						const FVector AnimationNormal(AnimationNormals[Index]);
 
-					// Draw a line to the sphere boundary
-					const FVector Pos0 = GetWorldPosition(*Solver, AnimationPositions[Index]);
-					const FVector Pos1 = Pos0 - (bUseLegacyBackstop ? BackstopDistance - BackstopRadius : BackstopDistance) * AnimationNormal;
-					DrawLine(PDI, Pos0, Pos1, ColorDark);
+						// Draw a line to the sphere boundary
+						const FVector Pos0 = GetWorldPosition(*Solver, AnimationPositions[Index]);
+						const FVector Pos1 = Pos0 - (bUseLegacyBackstop ? BackstopDistance - BackstopRadius : BackstopDistance) * AnimationNormal;
+						DrawLine(Context.PDI, Pos0, Pos1, ColorDark);
+					}
 				}
 			}
 		}
 	}
 
-	void FClothVisualizationNoGC::DrawMaxDistances(FPrimitiveDrawInterface* PDI) const
+	void FClothVisualizationNoGC::DrawMaxDistances(const FDrawContext& Context) const
 	{
 		if (!Solver)
 		{
@@ -2116,6 +2364,12 @@ namespace Chaos
 				continue;
 			}
 
+			const TSet<int32>* const DebugFilterSet = Context.Filter.VertexSet.IsEmpty() ? nullptr : Cloth->GetVertexSetByName(Solver, Context.Filter.VertexSet);
+			const int32 Offset = Solver->GetGlobalParticleOffset(ParticleRangeId);
+			auto ShouldDrawParticle = [&Context, DebugFilterSet, Offset](const int32 Index)
+				{
+					return Private::ShouldDrawParticle(Context.Filter, DebugFilterSet, Index, Offset);
+				};
 			const int32 NumParticles = Cloth->GetNumParticles(Solver);
 			const Softs::FPBDFlatWeightMapView MaxDistances(
 				(Softs::FSolverVec2)(Softs::FPBDSphericalConstraint::GetWeightedFloatMaxDistance(Cloth->GetConfig()->GetProperties(Cloth->GetLODIndex(Solver)), FVector2f(0.f, 1.f))),
@@ -2131,25 +2385,28 @@ namespace Chaos
 
 			for (int32 Index = 0; Index < MaxDistances.Num(); ++Index)
 			{
-				const FReal MaxDistance = (FReal)MaxDistances.GetValue(Index);
-				const FVector Position = GetWorldPosition(*Solver, Positions[Index]);
-				if (InvMasses[Index] == (Softs::FSolverReal)0.)
+				if (ShouldDrawParticle(Index))
 				{
+					const FReal MaxDistance = (FReal)MaxDistances.GetValue(Index);
+					const FVector Position = GetWorldPosition(*Solver, Positions[Index]);
+					if (InvMasses[Index] == (Softs::FSolverReal)0.)
+					{
 #if WITH_EDITOR
-					DrawPoint(PDI, Position, FLinearColor::Red, FMaterials::GetInstance().GetClothMaterialVertex());
+						DrawPoint(Context.PDI, Position, FLinearColor::Red, FMaterials::GetInstance().GetClothMaterialVertex());
 #else
-					DrawPoint(nullptr, Position, FLinearColor::Red, nullptr);
+						DrawPoint(nullptr, Position, FLinearColor::Red, nullptr);
 #endif
-				}
-				else
-				{
-					DrawLine(PDI, Position, Position + GetWorldVector(*Solver, Normals[Index]) * MaxDistance, FLinearColor::White);
+					}
+					else
+					{
+						DrawLine(Context.PDI, Position, Position + GetWorldVector(*Solver, Normals[Index]) * MaxDistance, FLinearColor::White);
+					}
 				}
 			}
 		}
 	}
 
-	void FClothVisualizationNoGC::DrawAnimDrive(FPrimitiveDrawInterface* PDI) const
+	void FClothVisualizationNoGC::DrawAnimDrive(const FDrawContext& Context) const
 	{
 		if (!Solver)
 		{
@@ -2176,21 +2433,31 @@ namespace Chaos
 				const FRealSingle StiffnessOffset = AnimDriveStiffness[0];
 				const FRealSingle StiffnessRange = AnimDriveStiffness[1] - AnimDriveStiffness[0];
 
+				const TSet<int32>* const DebugFilterSet = Context.Filter.VertexSet.IsEmpty() ? nullptr : Cloth->GetVertexSetByName(Solver, Context.Filter.VertexSet);
+				const int32 Offset = Solver->GetGlobalParticleOffset(ParticleRangeId);
+				auto ShouldDrawParticle = [&Context, DebugFilterSet, Offset](const int32 Index)
+					{
+						return Private::ShouldDrawParticle(Context.Filter, DebugFilterSet, Index, Offset);
+					};
+
 				for (int32 Index = 0; Index < ParticlePositions.Num(); ++Index)
 				{
-					const FRealSingle Stiffness = AnimDriveStiffnessMultipliers.IsValidIndex(Index) ?
-						StiffnessOffset + AnimDriveStiffnessMultipliers[Index] * StiffnessRange :
-						StiffnessOffset;
+					if (ShouldDrawParticle(Index))
+					{
+						const FRealSingle Stiffness = AnimDriveStiffnessMultipliers.IsValidIndex(Index) ?
+							StiffnessOffset + AnimDriveStiffnessMultipliers[Index] * StiffnessRange :
+							StiffnessOffset;
 
-					const FVector AnimationPosition = GetWorldPosition(*Solver, AnimationPositions[Index]);
-					const FVector ParticlePosition = GetWorldPosition(*Solver, ParticlePositions[Index]);
-					DrawLine(PDI, AnimationPosition, ParticlePosition, FLinearColor(FColor::Cyan) * Stiffness);
+						const FVector AnimationPosition = GetWorldPosition(*Solver, AnimationPositions[Index]);
+						const FVector ParticlePosition = GetWorldPosition(*Solver, ParticlePositions[Index]);
+						DrawLine(Context.PDI, AnimationPosition, ParticlePosition, FLinearColor(FColor::Cyan) * Stiffness);
+					}
 				}
 			}
 		}
 	}
-	template<typename SpringConstraintType, typename GetEndPointsFuncType>
-	static void DrawSpringConstraintColors(FPrimitiveDrawInterface* PDI, const ::Chaos::FClothingSimulationSolver& Solver, const SpringConstraintType* const SpringConstraints, GetEndPointsFuncType GetEndPoints)
+	template<typename SpringConstraintType, typename GetEndPointsFuncType, typename ShouldDrawConstraintFuncType>
+	static void DrawSpringConstraintColors(FPrimitiveDrawInterface* PDI, const ::Chaos::FClothingSimulationSolver& Solver, const SpringConstraintType* const SpringConstraints, GetEndPointsFuncType GetEndPoints, ShouldDrawConstraintFuncType ShouldDrawConstraint)
 	{
 		check(SpringConstraints);
 
@@ -2212,13 +2479,16 @@ namespace Chaos
 				const FLinearColor DrawColor = ConstraintColor(ConstraintColorIndex);
 				for (int32 ConstraintIndex = ColorStart; ConstraintIndex < ColorEnd; ++ConstraintIndex)
 				{
-					// Draw line
-					Softs::FSolverVec3 P1, P2;
-					GetEndPoints(ConstraintIndex, P1, P2);
+					if (ShouldDrawConstraint(ConstraintIndex))
+					{
+						// Draw line
+						Softs::FSolverVec3 P1, P2;
+						GetEndPoints(ConstraintIndex, P1, P2);
 
-					const FVec3 Pos0 = GetWorldPosition(Solver, P1);
-					const FVec3 Pos1 = GetWorldPosition(Solver, P2);
-					DrawLine(PDI, Pos0, Pos1, DrawColor);
+						const FVec3 Pos0 = GetWorldPosition(Solver, P1);
+						const FVec3 Pos1 = GetWorldPosition(Solver, P2);
+						DrawLine(PDI, Pos0, Pos1, DrawColor);
+					}
 				}
 			}
 		}
@@ -2226,33 +2496,41 @@ namespace Chaos
 		{
 			for (int32 ConstraintIndex = 0; ConstraintIndex < SpringConstraints->GetConstraints().Num(); ++ConstraintIndex)
 			{
-				// Draw line
-				Softs::FSolverVec3 P1, P2;
-				GetEndPoints(ConstraintIndex, P1, P2);
+				if (ShouldDrawConstraint(ConstraintIndex))
+				{
+					// Draw line
+					Softs::FSolverVec3 P1, P2;
+					GetEndPoints(ConstraintIndex, P1, P2);
 
-				const FVec3 Pos0 = GetWorldPosition(Solver, P1);
-				const FVec3 Pos1 = GetWorldPosition(Solver, P2);
+					const FVec3 Pos0 = GetWorldPosition(Solver, P1);
+					const FVec3 Pos1 = GetWorldPosition(Solver, P2);
 
-				DrawLine(PDI, Pos0, Pos1, FLinearColor::Black);
+					DrawLine(PDI, Pos0, Pos1, FLinearColor::Black);
+				}
 			}
 		}
 	}
 
-	template<typename SpringConstraintType>
-	static void DrawSpringConstraintColors(FPrimitiveDrawInterface* PDI, const TConstArrayView<::Chaos::Softs::FSolverVec3>& Positions, const ::Chaos::FClothingSimulationSolver& Solver, const SpringConstraintType* const SpringConstraints)
+	template<typename SpringConstraintType, typename ShouldDrawConstraintFuncType>
+	static void DrawSpringConstraintColors(FPrimitiveDrawInterface* PDI, const TConstArrayView<::Chaos::Softs::FSolverVec3>& Positions, const ::Chaos::FClothingSimulationSolver& Solver, const SpringConstraintType* const SpringConstraints, const ShouldDrawConstraintFuncType& ShouldDrawConstraint)
 	{
 		DrawSpringConstraintColors(PDI, Solver, SpringConstraints, [&SpringConstraints, &Positions](const int32 ConstraintIndex, Softs::FSolverVec3& P1, Softs::FSolverVec3& P2)
 		{
 			P1 = Positions[SpringConstraints->GetConstraints()[ConstraintIndex][0]];
 			P2 = Positions[SpringConstraints->GetConstraints()[ConstraintIndex][1]];
-		});
+		}, ShouldDrawConstraint);
 	}
 
-	static void DrawStretchBiasConstraints_ParallelGraphColor(FPrimitiveDrawInterface* PDI, const TConstArrayView<::Chaos::Softs::FSolverVec3>& Positions, const ::Chaos::FClothingSimulationSolver& Solver, const FMaterialRenderProxy* MaterialRenderProxy, const Softs::FXPBDStretchBiasElementConstraints* const SpringConstraints)
+	static void DrawStretchBiasConstraints_ParallelGraphColor(const FClothVisualizationNoGC::FDrawContext& Context, const TConstArrayView<::Chaos::Softs::FSolverVec3>& Positions, const ::Chaos::FClothingSimulationSolver& Solver, const FMaterialRenderProxy* MaterialRenderProxy, const Softs::FXPBDStretchBiasElementConstraints* const SpringConstraints, const TSet<int32>* const DebugFilterSet, int32 ConstraintToGlobalOffset, int32 ConstraintToLocalOffset)
 	{
 #if WITH_EDITOR
-		if (PDI && MaterialRenderProxy)
+		if (Context.PDI && MaterialRenderProxy)
 		{
+			auto ShouldDrawElement = [&Context, DebugFilterSet, ConstraintToGlobalOffset, ConstraintToLocalOffset](const TVec3<int32>& Element)
+				{
+					return Private::ShouldDrawElement(Context.Filter, DebugFilterSet, Element, ConstraintToGlobalOffset, ConstraintToLocalOffset);
+				};
+
 			const TArray<TVec3<int32>>& Constraints = SpringConstraints->GetConstraints();
 			const TArray<int32>& ConstraintsPerColorStartIndex = SpringConstraints->GetConstraintsPerColorStartIndex();
 
@@ -2265,7 +2543,7 @@ namespace Chaos
 				return FLinearColor::MakeFromHSV8(ColorIndex * HueOffset, 255, 255).ToFColor(true);
 			};
 
-			FDynamicMeshBuilder MeshBuilder(PDI->View->GetFeatureLevel());
+			FDynamicMeshBuilder MeshBuilder(Context.PDI->View->GetFeatureLevel());
 			int32 VertexIndex = 0;
 			for (int32 ConstraintColorIndex = 0; ConstraintColorIndex < ConstraintColorNum; ++ConstraintColorIndex)
 			{
@@ -2275,30 +2553,37 @@ namespace Chaos
 				for (int32 ConstraintIndex = ColorStart; ConstraintIndex < ColorEnd; ++ConstraintIndex, VertexIndex += 3)
 				{
 					const TVec3<int32>& Constraint = Constraints[ConstraintIndex];
-					const FVector3f Pos0 = FVector3f(Positions[Constraint[0]]);
-					const FVector3f Pos1 = FVector3f(Positions[Constraint[1]]);
-					const FVector3f Pos2 = FVector3f(Positions[Constraint[2]]);
-					const FVector3f Normal = FVector3f::CrossProduct(Pos2 - Pos0, Pos1 - Pos0).GetSafeNormal();
-					const FVector3f Tangent = ((Pos1 + Pos2) * 0.5f - Pos0).GetSafeNormal();
-
-					MeshBuilder.AddVertex(FDynamicMeshVertex(Pos0, Tangent, Normal, FVector2f(0.f, 0.f), DrawColor));
-					MeshBuilder.AddVertex(FDynamicMeshVertex(Pos1, Tangent, Normal, FVector2f(0.f, 1.f), DrawColor));
-					MeshBuilder.AddVertex(FDynamicMeshVertex(Pos2, Tangent, Normal, FVector2f(1.f, 1.f), DrawColor));
-					MeshBuilder.AddTriangle(VertexIndex, VertexIndex + 1, VertexIndex + 2);
+					if (ShouldDrawElement(Constraint))
+					{
+						const FVector3f Pos0 = FVector3f(Positions[Constraint[0]]);
+						const FVector3f Pos1 = FVector3f(Positions[Constraint[1]]);
+						const FVector3f Pos2 = FVector3f(Positions[Constraint[2]]);
+						const FVector3f Normal = FVector3f::CrossProduct(Pos2 - Pos0, Pos1 - Pos0).GetSafeNormal();
+						const FVector3f Tangent = ((Pos1 + Pos2) * 0.5f - Pos0).GetSafeNormal();
+
+						MeshBuilder.AddVertex(FDynamicMeshVertex(Pos0, Tangent, Normal, FVector2f(0.f, 0.f), DrawColor));
+						MeshBuilder.AddVertex(FDynamicMeshVertex(Pos1, Tangent, Normal, FVector2f(0.f, 1.f), DrawColor));
+						MeshBuilder.AddVertex(FDynamicMeshVertex(Pos2, Tangent, Normal, FVector2f(1.f, 1.f), DrawColor));
+						MeshBuilder.AddTriangle(VertexIndex, VertexIndex + 1, VertexIndex + 2);
+					}
 				}
 			}
 
 			const FMatrix LocalSimSpaceToWorld = GetLocalSpaceToWorld(Solver);
-			MeshBuilder.Draw(PDI, LocalSimSpaceToWorld, MaterialRenderProxy, Chaos::Private::GetDepthPriority(), false, false);
+			MeshBuilder.Draw(Context.PDI, LocalSimSpaceToWorld, MaterialRenderProxy, Chaos::Private::GetDepthPriority(), false, false);
 		}
 #endif
 	}
 
-	static void DrawStretchBiasConstraints_WarpWeftStretch(FPrimitiveDrawInterface* PDI, const TConstArrayView<::Chaos::Softs::FSolverVec3>& Positions, const::Chaos::FClothingSimulationSolver& Solver, const FMaterialRenderProxy* MaterialRenderProxy, const Softs::FXPBDStretchBiasElementConstraints* const SpringConstraints)
+	static void DrawStretchBiasConstraints_WarpWeftStretch(const FClothVisualizationNoGC::FDrawContext& Context, const TConstArrayView<::Chaos::Softs::FSolverVec3>& Positions, const::Chaos::FClothingSimulationSolver& Solver, const FMaterialRenderProxy* MaterialRenderProxy, const Softs::FXPBDStretchBiasElementConstraints* const SpringConstraints, const TSet<int32>* const DebugFilterSet, int32 ConstraintToGlobalOffset, int32 ConstraintToLocalOffset)
 	{
 #if WITH_EDITOR
-		if (PDI && MaterialRenderProxy)
+		if (Context.PDI && MaterialRenderProxy)
 		{
+			auto ShouldDrawElement = [&Context, DebugFilterSet, ConstraintToGlobalOffset, ConstraintToLocalOffset](const TVec3<int32>& Element)
+				{
+					return Private::ShouldDrawElement(Context.Filter, DebugFilterSet, Element, ConstraintToGlobalOffset, ConstraintToLocalOffset);
+				};
 			const TArray<TVec3<int32>>& Constraints = SpringConstraints->GetConstraints();
 			const float StretchRangeMinClamped = Private::StretchBiasDrawRangeMin <= Private::StretchBiasDrawRangeMax ? Private::StretchBiasDrawRangeMin : 0.f;
 			const float StretchRangeMaxClamped = Private::StretchBiasDrawRangeMin <= Private::StretchBiasDrawRangeMax ? Private::StretchBiasDrawRangeMax : 0.f;
@@ -2349,44 +2634,51 @@ namespace Chaos
 				return FLinearColor(Hue, 1.f, StretchRatio < 0 ? CompressedValue : StretchedValue).HSVToLinearRGB().ToFColor(true);
 			};
 
-			FDynamicMeshBuilder MeshBuilder(PDI->View->GetFeatureLevel());
+			FDynamicMeshBuilder MeshBuilder(Context.PDI->View->GetFeatureLevel());
 			int32 VertexIndex = 0;
 			for (int32 ConstraintIndex = 0; ConstraintIndex < Constraints.Num(); ++ConstraintIndex)
 			{
-
-				bool bIsOutOfRange;
-				const FColor DrawColor = ConstraintColor(ConstraintIndex, bIsOutOfRange);
-				if (bIsOutOfRange && !Private::bStretchBiasDrawOutOfRange)
+				const TVec3<int32>& Constraint = Constraints[ConstraintIndex];
+				if (ShouldDrawElement(Constraint))
 				{
-					continue;
-				}
 
-				const TVec3<int32>& Constraint = Constraints[ConstraintIndex];
-				const FVector3f Pos0 = FVector3f(Positions[Constraint[0]]);
-				const FVector3f Pos1 = FVector3f(Positions[Constraint[1]]);
-				const FVector3f Pos2 = FVector3f(Positions[Constraint[2]]);
-				const FVector3f Normal = FVector3f::CrossProduct(Pos2 - Pos0, Pos1 - Pos0).GetSafeNormal();
-				const FVector3f Tangent = ((Pos1 + Pos2) * 0.5f - Pos0).GetSafeNormal();
+					bool bIsOutOfRange;
+					const FColor DrawColor = ConstraintColor(ConstraintIndex, bIsOutOfRange);
+					if (bIsOutOfRange && !Private::bStretchBiasDrawOutOfRange)
+					{
+						continue;
+					}
 
-				MeshBuilder.AddVertex(FDynamicMeshVertex(Pos0, Tangent, Normal, FVector2f(0.f, 0.f), DrawColor));
-				MeshBuilder.AddVertex(FDynamicMeshVertex(Pos1, Tangent, Normal, FVector2f(0.f, 1.f), DrawColor));
-				MeshBuilder.AddVertex(FDynamicMeshVertex(Pos2, Tangent, Normal, FVector2f(1.f, 1.f), DrawColor));
-				MeshBuilder.AddTriangle(VertexIndex, VertexIndex + 1, VertexIndex + 2);
-				VertexIndex += 3;
+					const FVector3f Pos0 = FVector3f(Positions[Constraint[0]]);
+					const FVector3f Pos1 = FVector3f(Positions[Constraint[1]]);
+					const FVector3f Pos2 = FVector3f(Positions[Constraint[2]]);
+					const FVector3f Normal = FVector3f::CrossProduct(Pos2 - Pos0, Pos1 - Pos0).GetSafeNormal();
+					const FVector3f Tangent = ((Pos1 + Pos2) * 0.5f - Pos0).GetSafeNormal();
+
+					MeshBuilder.AddVertex(FDynamicMeshVertex(Pos0, Tangent, Normal, FVector2f(0.f, 0.f), DrawColor));
+					MeshBuilder.AddVertex(FDynamicMeshVertex(Pos1, Tangent, Normal, FVector2f(0.f, 1.f), DrawColor));
+					MeshBuilder.AddVertex(FDynamicMeshVertex(Pos2, Tangent, Normal, FVector2f(1.f, 1.f), DrawColor));
+					MeshBuilder.AddTriangle(VertexIndex, VertexIndex + 1, VertexIndex + 2);
+					VertexIndex += 3;
+				}
 			}
 
 			const FMatrix LocalSimSpaceToWorld = GetLocalSpaceToWorld(Solver);
-			MeshBuilder.Draw(PDI, LocalSimSpaceToWorld, MaterialRenderProxy, Chaos::Private::GetDepthPriority(), false, false);
+			MeshBuilder.Draw(Context.PDI, LocalSimSpaceToWorld, MaterialRenderProxy, Chaos::Private::GetDepthPriority(), false, false);
 		}
 #endif
 	}
 
-	static void DrawStretchBiasConstraints_BiasStretch(FPrimitiveDrawInterface* PDI, const TConstArrayView<::Chaos::Softs::FSolverVec3>& Positions, const::Chaos::FClothingSimulationSolver& Solver, const FMaterialRenderProxy* MaterialRenderProxy, const Softs::FXPBDStretchBiasElementConstraints* const SpringConstraints)
+	static void DrawStretchBiasConstraints_BiasStretch(const FClothVisualizationNoGC::FDrawContext& Context, const TConstArrayView<::Chaos::Softs::FSolverVec3>& Positions, const::Chaos::FClothingSimulationSolver& Solver, const FMaterialRenderProxy* MaterialRenderProxy, const Softs::FXPBDStretchBiasElementConstraints* const SpringConstraints, const TSet<int32>* const DebugFilterSet, int32 ConstraintToGlobalOffset, int32 ConstraintToLocalOffset)
 	{
 
 #if WITH_EDITOR
-		if (PDI && MaterialRenderProxy)
+		if (Context.PDI && MaterialRenderProxy)
 		{
+			auto ShouldDrawElement = [&Context, DebugFilterSet, ConstraintToGlobalOffset, ConstraintToLocalOffset](const TVec3<int32>& Element)
+				{
+					return Private::ShouldDrawElement(Context.Filter, DebugFilterSet, Element, ConstraintToGlobalOffset, ConstraintToLocalOffset);
+				};
 			const TArray<TVec3<int32>>& Constraints = SpringConstraints->GetConstraints();
 			auto ConstraintColor = [&Constraints, &Positions, &SpringConstraints](int32 ConstraintIndex)->FColor
 			{
@@ -2407,81 +2699,97 @@ namespace Chaos
 				return FLinearColor(UndeformedHue + Shear * (MaxDeformedHue - UndeformedHue), 1.f, 1.f).HSVToLinearRGB().ToFColor(true);
 			};
 
-			FDynamicMeshBuilder MeshBuilder(PDI->View->GetFeatureLevel());
+			FDynamicMeshBuilder MeshBuilder(Context.PDI->View->GetFeatureLevel());
 			int32 VertexIndex = 0;
-			for (int32 ConstraintIndex = 0; ConstraintIndex < Constraints.Num(); ++ConstraintIndex, VertexIndex += 3)
+			for (int32 ConstraintIndex = 0; ConstraintIndex < Constraints.Num(); ++ConstraintIndex)
 			{
-				const FColor DrawColor = ConstraintColor(ConstraintIndex);
-
 				const TVec3<int32>& Constraint = Constraints[ConstraintIndex];
-				const FVector3f Pos0 = FVector3f(Positions[Constraint[0]]);
-				const FVector3f Pos1 = FVector3f(Positions[Constraint[1]]);
-				const FVector3f Pos2 = FVector3f(Positions[Constraint[2]]);
-				const FVector3f Normal = FVector3f::CrossProduct(Pos2 - Pos0, Pos1 - Pos0).GetSafeNormal();
-				const FVector3f Tangent = ((Pos1 + Pos2) * 0.5f - Pos0).GetSafeNormal();
+				if (ShouldDrawElement(Constraint))
+				{
+					const FColor DrawColor = ConstraintColor(ConstraintIndex);
+
+					const FVector3f Pos0 = FVector3f(Positions[Constraint[0]]);
+					const FVector3f Pos1 = FVector3f(Positions[Constraint[1]]);
+					const FVector3f Pos2 = FVector3f(Positions[Constraint[2]]);
+					const FVector3f Normal = FVector3f::CrossProduct(Pos2 - Pos0, Pos1 - Pos0).GetSafeNormal();
+					const FVector3f Tangent = ((Pos1 + Pos2) * 0.5f - Pos0).GetSafeNormal();
 
-				MeshBuilder.AddVertex(FDynamicMeshVertex(Pos0, Tangent, Normal, FVector2f(0.f, 0.f), DrawColor));
-				MeshBuilder.AddVertex(FDynamicMeshVertex(Pos1, Tangent, Normal, FVector2f(0.f, 1.f), DrawColor));
-				MeshBuilder.AddVertex(FDynamicMeshVertex(Pos2, Tangent, Normal, FVector2f(1.f, 1.f), DrawColor));
-				MeshBuilder.AddTriangle(VertexIndex, VertexIndex + 1, VertexIndex + 2);
+					MeshBuilder.AddVertex(FDynamicMeshVertex(Pos0, Tangent, Normal, FVector2f(0.f, 0.f), DrawColor));
+					MeshBuilder.AddVertex(FDynamicMeshVertex(Pos1, Tangent, Normal, FVector2f(0.f, 1.f), DrawColor));
+					MeshBuilder.AddVertex(FDynamicMeshVertex(Pos2, Tangent, Normal, FVector2f(1.f, 1.f), DrawColor));
+					MeshBuilder.AddTriangle(VertexIndex, VertexIndex + 1, VertexIndex + 2);
+					VertexIndex += 3;
+				}
 			}
 
 			const FMatrix LocalSimSpaceToWorld = GetLocalSpaceToWorld(Solver);
-			MeshBuilder.Draw(PDI, LocalSimSpaceToWorld, MaterialRenderProxy, Chaos::Private::GetDepthPriority(), false, false);
+			MeshBuilder.Draw(Context.PDI, LocalSimSpaceToWorld, MaterialRenderProxy, Chaos::Private::GetDepthPriority(), false, false);
 		}
 #endif
 	}
 
-	template<typename ConstraintType>
-	static void DrawEdgeAnisotropy(FPrimitiveDrawInterface* PDI, const TConstArrayView<::Chaos::Softs::FSolverVec3>& Positions, const ::Chaos::FClothingSimulationSolver& Solver, const ConstraintType* const BendingConstraints)
+	template<typename ConstraintType, typename ShouldDrawConstraintFuncType>
+	static void DrawEdgeAnisotropy(const FClothVisualizationNoGC::FDrawContext& Context, const TConstArrayView<::Chaos::Softs::FSolverVec3>& Positions, const ::Chaos::FClothingSimulationSolver& Solver, const ConstraintType* const BendingConstraints, ShouldDrawConstraintFuncType ShouldDrawConstraint)
 	{
 		const auto& Constraints = BendingConstraints->GetConstraints(); // auto because this could be Vec2 or Vec4, but we always just care about first two indices
 		const TArray<Softs::FSolverVec3>& WarpWeftBiasBaseMultipliers = BendingConstraints->GetWarpWeftBiasBaseMultipliers();
+
 		for (int32 ConstraintIndex = 0; ConstraintIndex < Constraints.Num(); ++ConstraintIndex)
 		{
-			const Softs::FSolverVec3& P1 = Positions[Constraints[ConstraintIndex][0]];
-			const Softs::FSolverVec3& P2 = Positions[Constraints[ConstraintIndex][1]];
-			const Softs::FSolverVec3& Multiplier = WarpWeftBiasBaseMultipliers[ConstraintIndex];
+			if (ShouldDrawConstraint(ConstraintIndex))
+			{
+				const Softs::FSolverVec3& P1 = Positions[Constraints[ConstraintIndex][0]];
+				const Softs::FSolverVec3& P2 = Positions[Constraints[ConstraintIndex][1]];
+				const Softs::FSolverVec3& Multiplier = WarpWeftBiasBaseMultipliers[ConstraintIndex];
 
-			const FVec3 Pos0 = GetWorldPosition(Solver, P1);
-			const FVec3 Pos1 = GetWorldPosition(Solver, P2);
-			DrawLine(PDI, Pos0, Pos1, FLinearColor(Multiplier[0], Multiplier[1], Multiplier[2]));
+				const FVec3 Pos0 = GetWorldPosition(Solver, P1);
+				const FVec3 Pos1 = GetWorldPosition(Solver, P2);
+				DrawLine(Context.PDI, Pos0, Pos1, FLinearColor(Multiplier[0], Multiplier[1], Multiplier[2]));
+			}
 		}
 	}
 
 	template<typename ConstraintType>
-	static void DrawAxialSpringAnisotropy(FPrimitiveDrawInterface* PDI, const TConstArrayView<::Chaos::Softs::FSolverVec3>& Positions, const ::Chaos::FClothingSimulationSolver& Solver, const ConstraintType* const AxialConstraints)
+	static void DrawAxialSpringAnisotropy(const FClothVisualizationNoGC::FDrawContext& Context, const TConstArrayView<::Chaos::Softs::FSolverVec3>& Positions, const ::Chaos::FClothingSimulationSolver& Solver, const ConstraintType* const AxialConstraints, const TSet<int32>* const DebugFilterSet, int32 ConstraintToGlobalOffset, int32 ConstraintToLocalOffset)
 	{
+		auto ShouldDrawElement = [&Context, DebugFilterSet, ConstraintToGlobalOffset, ConstraintToLocalOffset](const TVec3<int32>& Element)
+			{
+				return Private::ShouldDrawElement(Context.Filter, DebugFilterSet, Element, ConstraintToGlobalOffset, ConstraintToLocalOffset);
+			};
 		const TArray<TVec3<int32>>& Constraints = AxialConstraints->GetConstraints();
 		const TArray<Softs::FSolverVec3>& WarpWeftBiasBaseMultipliers = AxialConstraints->GetWarpWeftBiasBaseMultipliers();
 		const TArray<Softs::FSolverReal>& Barys = AxialConstraints->GetBarys();
 		for (int32 ConstraintIndex = 0; ConstraintIndex < Constraints.Num(); ++ConstraintIndex)
 		{
-			const Softs::FSolverVec3& P1 = Positions[Constraints[ConstraintIndex][0]];
-			const Softs::FSolverVec3& P2 = Positions[Constraints[ConstraintIndex][1]];
-			const Softs::FSolverVec3& P3 = Positions[Constraints[ConstraintIndex][2]];
-			const Softs::FSolverVec3 P = Barys[ConstraintIndex] * P2 + ((FSolverReal)1. - Barys[ConstraintIndex]) * P3;
-			const Softs::FSolverVec3& Multiplier = WarpWeftBiasBaseMultipliers[ConstraintIndex];
+			if (ShouldDrawElement(Constraints[ConstraintIndex]))
+			{
+				const Softs::FSolverVec3& P1 = Positions[Constraints[ConstraintIndex][0]];
+				const Softs::FSolverVec3& P2 = Positions[Constraints[ConstraintIndex][1]];
+				const Softs::FSolverVec3& P3 = Positions[Constraints[ConstraintIndex][2]];
+				const Softs::FSolverVec3 P = Barys[ConstraintIndex] * P2 + ((FSolverReal)1. - Barys[ConstraintIndex]) * P3;
+				const Softs::FSolverVec3& Multiplier = WarpWeftBiasBaseMultipliers[ConstraintIndex];
 
-			const FVec3 Pos0 = GetWorldPosition(Solver, P1);
-			const FVec3 Pos1 = GetWorldPosition(Solver, P);
-			DrawLine(PDI, Pos0, Pos1, FLinearColor(Multiplier[0], Multiplier[1], Multiplier[2]));
+				const FVec3 Pos0 = GetWorldPosition(Solver, P1);
+				const FVec3 Pos1 = GetWorldPosition(Solver, P);
+				DrawLine(Context.PDI, Pos0, Pos1, FLinearColor(Multiplier[0], Multiplier[1], Multiplier[2]));
+			}
 		}
 	}
 
-	template<typename SpringConstraintType>
-	static void DrawAxialSpringConstraintColors(FPrimitiveDrawInterface* PDI, const TConstArrayView<::Chaos::Softs::FSolverVec3>& Positions, const ::Chaos::FClothingSimulationSolver& Solver, const SpringConstraintType* const SpringConstraints)
+	template<typename SpringConstraintType, typename ShouldDrawConstraintFuncType>
+	static void DrawAxialSpringConstraintColors(const FClothVisualizationNoGC::FDrawContext& Context, const TConstArrayView<::Chaos::Softs::FSolverVec3>& Positions, const ::Chaos::FClothingSimulationSolver& Solver, const SpringConstraintType* const SpringConstraints, const ShouldDrawConstraintFuncType& ShouldDrawConstraint)
 	{
-		DrawSpringConstraintColors(PDI, Solver, SpringConstraints, [&SpringConstraints, &Positions](const int32 ConstraintIndex, Softs::FSolverVec3& P1, Softs::FSolverVec3& P2)
+		DrawSpringConstraintColors(Context.PDI, Solver, SpringConstraints, 
+			[&SpringConstraints, &Positions](const int32 ConstraintIndex, Softs::FSolverVec3& P1, Softs::FSolverVec3& P2)
 		{
 			P1 = Positions[SpringConstraints->GetConstraints()[ConstraintIndex][0]];
 			const Softs::FSolverReal Bary = SpringConstraints->GetBarys()[ConstraintIndex];
 			
 			P2 = Bary * Positions[SpringConstraints->GetConstraints()[ConstraintIndex][1]] + ((FSolverReal)1. - Bary)* Positions[SpringConstraints->GetConstraints()[ConstraintIndex][2]];
-		});
+		}, ShouldDrawConstraint);
 	}
 
-	void FClothVisualizationNoGC::DrawEdgeConstraint(FPrimitiveDrawInterface* PDI) const
+	void FClothVisualizationNoGC::DrawEdgeConstraint(const FDrawContext& Context) const
 	{
 		if (!Solver)
 		{
@@ -2499,15 +2807,28 @@ namespace Chaos
 
 			// Constraints are locally indexed for new solver
 			const TConstArrayView<Softs::FSolverVec3> Positions = Solver->IsLegacySolver() ? TConstArrayView<Softs::FSolverVec3>(Solver->GetParticleXs()) : Solver->GetParticleXsView(ParticleRangeId);
+			const int32 Offset = Solver->IsLegacySolver() ? ParticleRangeId : 0;
+			const TSet<int32>* const DebugFilterSet = Context.Filter.VertexSet.IsEmpty() ? nullptr : Cloth->GetVertexSetByName(Solver, Context.Filter.VertexSet);
+			const int32 GlobalParticleOffset = Solver->GetGlobalParticleOffset(ParticleRangeId);
 
 			if (const Softs::FPBDEdgeSpringConstraints* const EdgeConstraints = ClothConstraints.GetEdgeSpringConstraints().Get())
 			{
-				DrawSpringConstraintColors(PDI, Positions, *Solver, EdgeConstraints);
+				DrawSpringConstraintColors(Context.PDI, Positions, *Solver, EdgeConstraints,
+					[&Context, EdgeConstraints, DebugFilterSet, GlobalParticleOffset, Offset](int32 ConstraintIndex)
+					{
+						const TVec2<int32>& Constraint = EdgeConstraints->GetConstraints()[ConstraintIndex];
+						return Private::ShouldDrawConstraint2(Context.Filter, DebugFilterSet, Constraint, GlobalParticleOffset - Offset, -Offset);
+					});
 			}
 
 			if (const Softs::FXPBDEdgeSpringConstraints* const EdgeConstraints = ClothConstraints.GetXEdgeSpringConstraints().Get())
 			{
-				DrawSpringConstraintColors(PDI, Positions, *Solver, EdgeConstraints);
+				DrawSpringConstraintColors(Context.PDI, Positions, *Solver, EdgeConstraints,
+					[&Context, EdgeConstraints, DebugFilterSet, GlobalParticleOffset, Offset](int32 ConstraintIndex)
+				{
+					const TVec2<int32>& Constraint = EdgeConstraints->GetConstraints()[ConstraintIndex];
+					return Private::ShouldDrawConstraint2(Context.Filter, DebugFilterSet, Constraint, GlobalParticleOffset - Offset, -Offset);
+				});
 			}
 
 			if (const Softs::FXPBDAnisotropicSpringConstraints* const AnisoSpringConstraints = ClothConstraints.GetXAnisoSpringConstraints().Get())
@@ -2515,13 +2836,28 @@ namespace Chaos
 				switch ((Private::EAnisoSpringDrawMode)Private::AnisoSpringDrawMode)
 				{
 				case Private::EAnisoSpringDrawMode::Anisotropy:
-					DrawEdgeAnisotropy(PDI, Positions, *Solver, &AnisoSpringConstraints->GetEdgeConstraints());
-					DrawAxialSpringAnisotropy(PDI, Positions, *Solver, &AnisoSpringConstraints->GetAxialConstraints());
+					DrawEdgeAnisotropy(Context, Positions, *Solver, &AnisoSpringConstraints->GetEdgeConstraints(),
+						[&Context, AnisoSpringConstraints, DebugFilterSet, GlobalParticleOffset, Offset](int32 ConstraintIndex)
+					{
+						const TVec2<int32>& Constraint = AnisoSpringConstraints->GetEdgeConstraints().GetConstraints()[ConstraintIndex];
+						return Private::ShouldDrawConstraint2(Context.Filter, DebugFilterSet, Constraint, GlobalParticleOffset - Offset, -Offset);
+					});
+					DrawAxialSpringAnisotropy(Context, Positions, *Solver, &AnisoSpringConstraints->GetAxialConstraints(), DebugFilterSet, GlobalParticleOffset - Offset, -Offset);
 					break;
 				case Private::EAnisoSpringDrawMode::ParallelGraphColor: // fallthrough
 				default:
-					DrawSpringConstraintColors(PDI, Positions, *Solver, &AnisoSpringConstraints->GetEdgeConstraints());
-					DrawAxialSpringConstraintColors(PDI, Positions, *Solver, &AnisoSpringConstraints->GetAxialConstraints());
+					DrawSpringConstraintColors(Context.PDI, Positions, *Solver, &AnisoSpringConstraints->GetEdgeConstraints(),
+						[&Context, AnisoSpringConstraints, DebugFilterSet, GlobalParticleOffset, Offset](int32 ConstraintIndex)
+						{
+							const TVec2<int32>& Constraint = AnisoSpringConstraints->GetEdgeConstraints().GetConstraints()[ConstraintIndex];
+							return Private::ShouldDrawConstraint2(Context.Filter, DebugFilterSet, Constraint, GlobalParticleOffset - Offset, -Offset);
+						});
+					DrawAxialSpringConstraintColors(Context, Positions, *Solver, &AnisoSpringConstraints->GetAxialConstraints(),
+						[&Context, AnisoSpringConstraints, DebugFilterSet, GlobalParticleOffset, Offset](int32 ConstraintIndex)
+					{
+						const TVec3<int32>& Constraint = AnisoSpringConstraints->GetAxialConstraints().GetConstraints()[ConstraintIndex];
+						return Private::ShouldDrawElement(Context.Filter, DebugFilterSet, Constraint, GlobalParticleOffset - Offset, -Offset);
+					});
 					break;
 				}
 			}
@@ -2538,59 +2874,75 @@ namespace Chaos
 				{
 				case Private::EStretchBiasDrawMode::WarpStretch: // fallthrough
 				case Private::EStretchBiasDrawMode::WeftStretch:
-					DrawStretchBiasConstraints_WarpWeftStretch(PDI, Positions, *Solver, MaterialRenderProxy, StretchConstraints);
+					DrawStretchBiasConstraints_WarpWeftStretch(Context, Positions, *Solver, MaterialRenderProxy, StretchConstraints, DebugFilterSet, GlobalParticleOffset - Offset, -Offset);
 					break;
 				case Private::EStretchBiasDrawMode::BiasStretch:
-					DrawStretchBiasConstraints_BiasStretch(PDI, Positions, *Solver, MaterialRenderProxy, StretchConstraints);
+					DrawStretchBiasConstraints_BiasStretch(Context, Positions, *Solver, MaterialRenderProxy, StretchConstraints, DebugFilterSet, GlobalParticleOffset - Offset, -Offset);
 					break;
 				case Private::EStretchBiasDrawMode::ParallelGraphColor: // fallthrough
 				default:
-					DrawStretchBiasConstraints_ParallelGraphColor(PDI, Positions, *Solver, MaterialRenderProxy, StretchConstraints);
+					DrawStretchBiasConstraints_ParallelGraphColor(Context, Positions, *Solver, MaterialRenderProxy, StretchConstraints, DebugFilterSet, GlobalParticleOffset - Offset, -Offset);
 				}
 
 			}
 		}
 	}
 
-	static void DrawBendingElementBuckleStatus(FPrimitiveDrawInterface* PDI, const TConstArrayView<::Chaos::Softs::FSolverVec3>& Positions, const ::Chaos::FClothingSimulationSolver& Solver, const ::Chaos::Softs::FPBDBendingConstraintsBase* const BendingConstraints)
+	static void DrawBendingElementBuckleStatus(const FClothVisualizationNoGC::FDrawContext& Context, const TConstArrayView<::Chaos::Softs::FSolverVec3>& Positions, const ::Chaos::FClothingSimulationSolver& Solver, const ::Chaos::Softs::FPBDBendingConstraintsBase* const BendingConstraints, const TSet<int32>* const DebugFilterSet, int32 ConstraintToGlobalOffset, int32 ConstraintToLocalOffset)
 	{
 		const TArray<TVec4<int32>>& Constraints = BendingConstraints->GetConstraints();
 		const TArray<bool>& IsBuckled = BendingConstraints->GetIsBuckled();
 
+		auto ShouldDrawConstraint4 = [&Context, DebugFilterSet, ConstraintToGlobalOffset, ConstraintToLocalOffset](const TVec4<int32>& Element)
+			{
+				return Private::ShouldDrawConstraint4(Context.Filter, DebugFilterSet, Element, ConstraintToGlobalOffset, ConstraintToLocalOffset);
+			};
+
 		// Color constraint edge with red or blue: Red = Buckled, Blue = Not Buckled. 
 		for (int32 ConstraintIndex = 0; ConstraintIndex < Constraints.Num(); ++ConstraintIndex)
 		{
-			const Softs::FSolverVec3& P1 = Positions[Constraints[ConstraintIndex][0]];
-			const Softs::FSolverVec3& P2 = Positions[Constraints[ConstraintIndex][1]];
+			if (ShouldDrawConstraint4(Constraints[ConstraintIndex]))
+			{
+				const Softs::FSolverVec3& P1 = Positions[Constraints[ConstraintIndex][0]];
+				const Softs::FSolverVec3& P2 = Positions[Constraints[ConstraintIndex][1]];
 
-			const bool bIsBuckled = IsBuckled.IsValidIndex(ConstraintIndex) ? IsBuckled[ConstraintIndex] : false; // IsBuckled is empty if the simulation is paused.
+				const bool bIsBuckled = IsBuckled.IsValidIndex(ConstraintIndex) ? IsBuckled[ConstraintIndex] : false; // IsBuckled is empty if the simulation is paused.
 
-			const FVec3 Pos0 = GetWorldPosition(Solver, P1);
-			const FVec3 Pos1 = GetWorldPosition(Solver, P2);
-			DrawLine(PDI, Pos0, Pos1, bIsBuckled ? FLinearColor::Red : FLinearColor::Blue);
+				const FVec3 Pos0 = GetWorldPosition(Solver, P1);
+				const FVec3 Pos1 = GetWorldPosition(Solver, P2);
+				DrawLine(Context.PDI, Pos0, Pos1, bIsBuckled ? FLinearColor::Red : FLinearColor::Blue);
+			}
 		}
 	}
 
 
 
-	static void DrawBendingElementRestAngle(FPrimitiveDrawInterface* PDI, const TConstArrayView<::Chaos::Softs::FSolverVec3>& Positions, const ::Chaos::FClothingSimulationSolver& Solver, const ::Chaos::Softs::FPBDBendingConstraintsBase* const BendingConstraints)
+	static void DrawBendingElementRestAngle(const FClothVisualizationNoGC::FDrawContext& Context, const TConstArrayView<::Chaos::Softs::FSolverVec3>& Positions, const ::Chaos::FClothingSimulationSolver& Solver, const ::Chaos::Softs::FPBDBendingConstraintsBase* const BendingConstraints, const TSet<int32>* const DebugFilterSet, int32 ConstraintToGlobalOffset, int32 ConstraintToLocalOffset)
 	{
 		const TArray<TVec4<int32>>& Constraints = BendingConstraints->GetConstraints();
 		const TArray<Softs::FSolverReal>& RestAngles = BendingConstraints->GetRestAngles();
+
+		auto ShouldDrawConstraint4 = [&Context, DebugFilterSet, ConstraintToGlobalOffset, ConstraintToLocalOffset](const TVec4<int32>& Element)
+			{
+				return Private::ShouldDrawConstraint4(Context.Filter, DebugFilterSet, Element, ConstraintToGlobalOffset, ConstraintToLocalOffset);
+			};
 		for (int32 ConstraintIndex = 0; ConstraintIndex < Constraints.Num(); ++ConstraintIndex)
 		{
-			const Softs::FSolverVec3& P1 = Positions[Constraints[ConstraintIndex][0]];
-			const Softs::FSolverVec3& P2 = Positions[Constraints[ConstraintIndex][1]];
-			const Softs::FSolverReal RestAngle = RestAngles[ConstraintIndex];
-			const uint8 ColorSat = (uint8)FMath::Clamp(FMath::Abs(RestAngle) / UE_PI * 256, 0, 255);
+			if (ShouldDrawConstraint4(Constraints[ConstraintIndex]))
+			{
+				const Softs::FSolverVec3& P1 = Positions[Constraints[ConstraintIndex][0]];
+				const Softs::FSolverVec3& P2 = Positions[Constraints[ConstraintIndex][1]];
+				const Softs::FSolverReal RestAngle = RestAngles[ConstraintIndex];
+				const uint8 ColorSat = (uint8)FMath::Clamp(FMath::Abs(RestAngle) / UE_PI * 256, 0, 255);
 
-			const FVec3 Pos0 = GetWorldPosition(Solver, P1);
-			const FVec3 Pos1 = GetWorldPosition(Solver, P2);
-			DrawLine(PDI, Pos0, Pos1, FLinearColor::MakeFromHSV8(RestAngle > 0 ? 170 : 0, ColorSat, 255));
+				const FVec3 Pos0 = GetWorldPosition(Solver, P1);
+				const FVec3 Pos1 = GetWorldPosition(Solver, P2);
+				DrawLine(Context.PDI, Pos0, Pos1, FLinearColor::MakeFromHSV8(RestAngle > 0 ? 170 : 0, ColorSat, 255));
+			}
 		}
 	}
 
-	void FClothVisualizationNoGC::DrawBendingConstraint(FPrimitiveDrawInterface* PDI) const
+	void FClothVisualizationNoGC::DrawBendingConstraint(const FDrawContext& Context) const
 	{
 		if (!Solver)
 		{
@@ -2610,15 +2962,28 @@ namespace Chaos
 
 			// Constraints are locally indexed for new solver
 			const TConstArrayView<Softs::FSolverVec3> Positions = Solver->IsLegacySolver() ? TConstArrayView<Softs::FSolverVec3>(Solver->GetParticleXs()) : Solver->GetParticleXsView(ParticleRangeId);
+			const int32 Offset = Solver->IsLegacySolver() ? ParticleRangeId : 0;
+			const TSet<int32>* const DebugFilterSet = Context.Filter.VertexSet.IsEmpty() ? nullptr : Cloth->GetVertexSetByName(Solver, Context.Filter.VertexSet);
+			const int32 GlobalParticleOffset = Solver->GetGlobalParticleOffset(ParticleRangeId);
 
 			if (const Softs::FPBDBendingSpringConstraints* const BendingConstraints = ClothConstraints.GetBendingSpringConstraints().Get())
 			{
-				DrawSpringConstraintColors(PDI, Positions, *Solver, BendingConstraints);
+				DrawSpringConstraintColors(Context.PDI, Positions, *Solver, BendingConstraints,
+					[&Context, BendingConstraints, DebugFilterSet, GlobalParticleOffset, Offset](int32 ConstraintIndex)
+					{
+						const TVec2<int32>& Constraint = BendingConstraints->GetConstraints()[ConstraintIndex];
+						return Private::ShouldDrawConstraint2(Context.Filter, DebugFilterSet, Constraint, GlobalParticleOffset - Offset, -Offset);
+					});
 			}
 
 			if (const Softs::FXPBDBendingSpringConstraints* const BendingConstraints = ClothConstraints.GetXBendingSpringConstraints().Get())
 			{
-				DrawSpringConstraintColors(PDI, Positions, *Solver, BendingConstraints);
+				DrawSpringConstraintColors(Context.PDI, Positions, *Solver, BendingConstraints,
+					[&Context, BendingConstraints, DebugFilterSet, GlobalParticleOffset, Offset](int32 ConstraintIndex)
+					{
+						const TVec2<int32>& Constraint = BendingConstraints->GetConstraints()[ConstraintIndex];
+						return Private::ShouldDrawConstraint2(Context.Filter, DebugFilterSet, Constraint, GlobalParticleOffset - Offset, -Offset);
+					});
 			}
 
 			if (const Softs::FPBDBendingConstraints* const BendingConstraints = ClothConstraints.GetBendingElementConstraints().Get())
@@ -2626,14 +2991,19 @@ namespace Chaos
 				switch ((Private::EBendingDrawMode)Private::BendingDrawMode)
 				{
 				case Private::EBendingDrawMode::RestAngle:
-					DrawBendingElementRestAngle(PDI, Positions, *Solver, BendingConstraints);
+					DrawBendingElementRestAngle(Context, Positions, *Solver, BendingConstraints, DebugFilterSet, GlobalParticleOffset - Offset, -Offset);
 					break;
 				case Private::EBendingDrawMode::ParallelGraphColor:
-					DrawSpringConstraintColors(PDI, Positions, *Solver, BendingConstraints);
+					DrawSpringConstraintColors(Context.PDI, Positions, *Solver, BendingConstraints,
+						[&Context, BendingConstraints, DebugFilterSet, GlobalParticleOffset, Offset](int32 ConstraintIndex)
+						{
+							const TVec4<int32>& Constraint = BendingConstraints->GetConstraints()[ConstraintIndex];
+							return Private::ShouldDrawConstraint4(Context.Filter, DebugFilterSet, Constraint, GlobalParticleOffset - Offset, -Offset);
+						});
 					break;
 				case Private::EBendingDrawMode::BuckleStatus:
 				default:
-					DrawBendingElementBuckleStatus(PDI, Positions, *Solver, BendingConstraints);
+					DrawBendingElementBuckleStatus(Context, Positions, *Solver, BendingConstraints, DebugFilterSet, GlobalParticleOffset - Offset, -Offset);
 					break;
 				}
 			}
@@ -2643,14 +3013,19 @@ namespace Chaos
 				switch ((Private::EBendingDrawMode)Private::BendingDrawMode)
 				{
 				case Private::EBendingDrawMode::RestAngle:
-					DrawBendingElementRestAngle(PDI, Positions, *Solver, BendingConstraints);
+					DrawBendingElementRestAngle(Context, Positions, *Solver, BendingConstraints, DebugFilterSet, GlobalParticleOffset - Offset, -Offset);
 					break;
 				case Private::EBendingDrawMode::ParallelGraphColor:
-					DrawSpringConstraintColors(PDI, Positions, *Solver, BendingConstraints);
+					DrawSpringConstraintColors(Context.PDI, Positions, *Solver, BendingConstraints,
+						[&Context, BendingConstraints, DebugFilterSet, GlobalParticleOffset, Offset](int32 ConstraintIndex)
+						{
+							const TVec4<int32>& Constraint = BendingConstraints->GetConstraints()[ConstraintIndex];
+							return Private::ShouldDrawConstraint4(Context.Filter, DebugFilterSet, Constraint, GlobalParticleOffset - Offset, -Offset);
+						});
 					break;
 				case Private::EBendingDrawMode::BuckleStatus:
 				default:
-					DrawBendingElementBuckleStatus(PDI, Positions, *Solver, BendingConstraints);
+					DrawBendingElementBuckleStatus(Context, Positions, *Solver, BendingConstraints, DebugFilterSet, GlobalParticleOffset - Offset, -Offset);
 					break;
 				}
 			}
@@ -2660,24 +3035,34 @@ namespace Chaos
 				switch ((Private::EBendingDrawMode)Private::BendingDrawMode)
 				{
 				case Private::EBendingDrawMode::RestAngle:
-					DrawBendingElementRestAngle(PDI, Positions, *Solver, BendingConstraints);
+					DrawBendingElementRestAngle(Context, Positions, *Solver, BendingConstraints, DebugFilterSet, GlobalParticleOffset - Offset, -Offset);
 					break;
 				case Private::EBendingDrawMode::ParallelGraphColor:
-					DrawSpringConstraintColors(PDI, Positions, *Solver, BendingConstraints);
+					DrawSpringConstraintColors(Context.PDI, Positions, *Solver, BendingConstraints,
+						[&Context, BendingConstraints, DebugFilterSet, GlobalParticleOffset, Offset](int32 ConstraintIndex)
+						{
+							const TVec4<int32>& Constraint = BendingConstraints->GetConstraints()[ConstraintIndex];
+							return Private::ShouldDrawConstraint4(Context.Filter, DebugFilterSet, Constraint, GlobalParticleOffset - Offset, -Offset);
+						});
 					break;
 				case Private::EBendingDrawMode::Anisotropy:
-					DrawEdgeAnisotropy(PDI, Positions, *Solver, BendingConstraints);
+					DrawEdgeAnisotropy(Context, Positions, *Solver, BendingConstraints,
+						[&Context, BendingConstraints, DebugFilterSet, GlobalParticleOffset, Offset](int32 ConstraintIndex)
+						{
+							const TVec4<int32>& Constraint = BendingConstraints->GetConstraints()[ConstraintIndex];
+							return Private::ShouldDrawConstraint4(Context.Filter, DebugFilterSet, Constraint, GlobalParticleOffset - Offset, -Offset);
+						});
 					break;
 				case Private::EBendingDrawMode::BuckleStatus:
 				default:
-					DrawBendingElementBuckleStatus(PDI, Positions, *Solver, BendingConstraints);
+					DrawBendingElementBuckleStatus(Context, Positions, *Solver, BendingConstraints, DebugFilterSet, GlobalParticleOffset - Offset, -Offset);
 					break;
 				}
 			}
 		}
 	}
 
-	void FClothVisualizationNoGC::DrawLongRangeConstraint(FPrimitiveDrawInterface* PDI) const
+	void FClothVisualizationNoGC::DrawLongRangeConstraint(const FDrawContext& Context) const
 	{
 		if (!Solver)
 		{
@@ -2721,6 +3106,13 @@ namespace Chaos
 			{
 				const TArray<TConstArrayView<Softs::FPBDLongRangeConstraints::FTether>>& Tethers = LongRangeConstraints->GetTethers();
 
+				const TSet<int32>* const DebugFilterSet = Context.Filter.VertexSet.IsEmpty() ? nullptr : Cloth->GetVertexSetByName(Solver, Context.Filter.VertexSet);
+				const int32 Offset = Solver->GetGlobalParticleOffset(ParticleRangeId);
+				auto ShouldDrawParticle = [&Context, DebugFilterSet, Offset](const int32 Index)
+					{
+						return Private::ShouldDrawParticle(Context.Filter, DebugFilterSet, Index, Offset);
+					};
+
 				for (int32 BatchIndex = 0; BatchIndex < Tethers.Num(); ++BatchIndex)
 				{
 					const FLinearColor Color = PseudoRandomColor(ColorOffset + BatchIndex);
@@ -2734,25 +3126,29 @@ namespace Chaos
 						const int32 KinematicIndex = LongRangeConstraints->GetStartIndex(Tether);
 						const int32 DynamicIndex = LongRangeConstraints->GetEndIndex(Tether);
 						const FReal TargetLength = LongRangeConstraints->GetTargetLength(Tether);
+						if (ShouldDrawParticle(KinematicIndex) ||
+							ShouldDrawParticle(DynamicIndex))
+						{
 
-						const Softs::FSolverVec3 KinematicPos = Positions[KinematicIndex];
-						const Softs::FSolverVec3 DynamicPos = Positions[DynamicIndex];
+							const Softs::FSolverVec3 KinematicPos = Positions[KinematicIndex];
+							const Softs::FSolverVec3 DynamicPos = Positions[DynamicIndex];
 
-						const FVec3 Pos0 = GetWorldPosition(*Solver, KinematicPos);
-						const FVec3 Pos1 = GetWorldPosition(*Solver, DynamicPos);
+							const FVec3 Pos0 = GetWorldPosition(*Solver, KinematicPos);
+							const FVec3 Pos1 = GetWorldPosition(*Solver, DynamicPos);
 
-						DrawLine(PDI, Pos0, Pos1, Color);
+							DrawLine(Context.PDI, Pos0, Pos1, Color);
 #if WITH_EDITOR
-						DrawPoint(PDI, Pos1, Color, FMaterials::GetInstance().GetClothMaterialVertex());
+							DrawPoint(Context.PDI, Pos1, Color, FMaterials::GetInstance().GetClothMaterialVertex());
 #else
-						DrawPoint(nullptr, Pos1, Color, nullptr);
+							DrawPoint(nullptr, Pos1, Color, nullptr);
 #endif
-						Softs::FSolverVec3 Direction = DynamicPos - KinematicPos;
-						const float Length = Direction.SafeNormalize();
-						if (Length > SMALL_NUMBER)
-						{
-							const FVec3 Pos2 = Pos1 + GetWorldVector(*Solver, Direction * (TargetLength - Length));
-							DrawLine(PDI, Pos1, Pos2, DarkenedColor);
+							Softs::FSolverVec3 Direction = DynamicPos - KinematicPos;
+							const float Length = Direction.SafeNormalize();
+							if (Length > SMALL_NUMBER)
+							{
+								const FVec3 Pos2 = Pos1 + GetWorldVector(*Solver, Direction * (TargetLength - Length));
+								DrawLine(Context.PDI, Pos1, Pos2, DarkenedColor);
+							}
 						}
 					}
 				}
@@ -2763,7 +3159,7 @@ namespace Chaos
 		}
 	}
 
-	void FClothVisualizationNoGC::DrawWindAndPressureForces(FPrimitiveDrawInterface* PDI, const FReal ForceLength) const
+	void FClothVisualizationNoGC::DrawWindAndPressureForces(const FDrawContext& Context, const FReal ForceLength) const
 	{
 		if (!Solver)
 		{
@@ -2798,6 +3194,14 @@ namespace Chaos
 			const TConstArrayView<Softs::FSolverReal> InvMasses = Solver->IsLegacySolver() ? TConstArrayView<Softs::FSolverReal>(Solver->GetParticleInvMasses()) : Solver->GetParticleInvMassesView(ParticleRangeId);
 			const TConstArrayView<Softs::FSolverVec3> Velocities = Solver->IsLegacySolver() ? TConstArrayView<Softs::FSolverVec3>(Solver->GetParticleVs()) : Solver->GetParticleVsView(ParticleRangeId);
 
+			const int32 Offset = Solver->IsLegacySolver() ? ParticleRangeId : 0;
+			const TSet<int32>* const DebugFilterSet = Context.Filter.VertexSet.IsEmpty() ? nullptr : Cloth->GetVertexSetByName(Solver, Context.Filter.VertexSet);
+			const int32 GlobalParticleOffset = Solver->GetGlobalParticleOffset(ParticleRangeId);
+			auto ShouldDrawElement = [&Context, DebugFilterSet, GlobalParticleOffset, Offset](const TVec3<int32>& Element)
+				{
+					return Private::ShouldDrawElement(Context.Filter, DebugFilterSet, Element, GlobalParticleOffset - Offset, -Offset);
+				};
+
 			const TConstArrayView<TVec3<int32>>& Elements = VelocityField->GetElements();
 			check(InvMasses.Num() == Positions.Num());
 			check(InvMasses.Num() == Velocities.Num());
@@ -2805,23 +3209,26 @@ namespace Chaos
 			for (int32 ElementIndex = 0; ElementIndex < Elements.Num(); ++ElementIndex)
 			{
 				const TVec3<int32>& Element = Elements[ElementIndex];
-				const FVec3 Position = GetWorldPosition(*Solver, 
-					(Positions[Element.X] +
-					Positions[Element.Y] +
-					Positions[Element.Z]) / (Softs::FSolverReal)3.f);
-
-				const bool bIsKinematic0 = !InvMasses[Element.X];
-				const bool bIsKinematic1 = !InvMasses[Element.Y];
-				const bool bIsKinematic2 = !InvMasses[Element.Z];
-				const bool bIsKinematic = bIsKinematic0 || bIsKinematic1 || bIsKinematic2;
-
-				const FVec3 Force = GetWorldVector(*Solver, FVec3(VelocityField->CalculateForce(Positions, Velocities, ElementIndex)) * ForceLength);
-				DrawLine(PDI, Position, Position + Force, bIsKinematic ? FColor::Cyan : FColor::Green);
+				if (ShouldDrawElement(Element))
+				{
+					const FVec3 Position = GetWorldPosition(*Solver,
+						(Positions[Element.X] +
+							Positions[Element.Y] +
+							Positions[Element.Z]) / (Softs::FSolverReal)3.f);
+
+					const bool bIsKinematic0 = !InvMasses[Element.X];
+					const bool bIsKinematic1 = !InvMasses[Element.Y];
+					const bool bIsKinematic2 = !InvMasses[Element.Z];
+					const bool bIsKinematic = bIsKinematic0 || bIsKinematic1 || bIsKinematic2;
+
+					const FVec3 Force = GetWorldVector(*Solver, FVec3(VelocityField->CalculateForce(Positions, Velocities, ElementIndex)) * ForceLength);
+					DrawLine(Context.PDI, Position, Position + Force, bIsKinematic ? FColor::Cyan : FColor::Green);
+				}
 			}
 		}
 	}
 
-	void FClothVisualizationNoGC::DrawWindVelocity(FPrimitiveDrawInterface* PDI, const FReal LengthScale) const
+	void FClothVisualizationNoGC::DrawWindVelocity(const FDrawContext& Context, const FReal LengthScale) const
 	{
 		if (!Solver)
 		{
@@ -2852,12 +3259,12 @@ namespace Chaos
 				const FRigidTransform3& ReferenceSpaceTransform = Cloth->GetReferenceSpaceTransform();
 				const FVector Pos0 = ReferenceSpaceTransform.GetLocation();
 				const FVector Pos1 = Pos0 + LengthScale * GetWorldVector(*Solver, VelocityField->GetVelocity());
-				DrawArrow(PDI, Pos0, Pos1, FLinearColor::Blue);
+				DrawArrow(Context.PDI, Pos0, Pos1, FLinearColor::Blue);
 			}
 		}
 	}
 
-	void FClothVisualizationNoGC::DrawLocalSpace(FPrimitiveDrawInterface* PDI, const FReal LocalSpaceLengthScale) const
+	void FClothVisualizationNoGC::DrawLocalSpace(const FDrawContext& Context, const FReal LocalSpaceLengthScale) const
 	{
 		if (!Solver)
 		{
@@ -2865,8 +3272,8 @@ namespace Chaos
 		}
 
 		// Draw local space
-		DrawCoordinateSystem(PDI, FQuat::Identity, Solver->GetLocalSpaceLocation(), 5.f * Solver->GetLocalSpaceScale() * LocalSpaceLengthScale, 0.25f);
-		DrawCoordinateSystem(PDI, Solver->GetLocalSpaceRotation(), Solver->GetLocalSpaceLocation(), 7.f * Solver->GetLocalSpaceScale() * LocalSpaceLengthScale, 0.5f);
+		DrawCoordinateSystem(Context.PDI, FQuat::Identity, Solver->GetLocalSpaceLocation(), 5.f * Solver->GetLocalSpaceScale() * LocalSpaceLengthScale, 0.25f);
+		DrawCoordinateSystem(Context.PDI, Solver->GetLocalSpaceRotation(), Solver->GetLocalSpaceLocation(), 7.f * Solver->GetLocalSpaceScale() * LocalSpaceLengthScale, 0.5f);
 		DrawText(nullptr, nullptr, Solver->GetLocalSpaceLocation() + FVec3(2.), LOCTEXT("SimulationSpace", "Simulation Space"), FLinearColor::White, 1.5f);
 
 		// Draw reference spaces
@@ -2878,13 +3285,13 @@ namespace Chaos
 				continue;
 			}
 			const FRigidTransform3& ReferenceSpaceTransform = Cloth->GetReferenceSpaceTransform();
-			DrawCoordinateSystem(PDI, ReferenceSpaceTransform.GetRotation(), ReferenceSpaceTransform.GetLocation(), 10.f * Solver->GetLocalSpaceScale() * LocalSpaceLengthScale);
+			DrawCoordinateSystem(Context.PDI, ReferenceSpaceTransform.GetRotation(), ReferenceSpaceTransform.GetLocation(), 10.f * Solver->GetLocalSpaceScale() * LocalSpaceLengthScale);
 			check(Cloth->GetMesh());
 			DrawText(nullptr, nullptr, ReferenceSpaceTransform.GetLocation() + FVec3(-2.), FText::FromName(Cloth->GetMesh()->GetReferenceBoneName()), FLinearColor::White, 1.5f);
 		}
 	}
 
-	void FClothVisualizationNoGC::DrawVelocityScale(FPrimitiveDrawInterface* PDI, const FReal LengthScale) const
+	void FClothVisualizationNoGC::DrawVelocityScale(const FDrawContext& Context, const FReal LengthScale) const
 	{
 		if (!Solver)
 		{
@@ -2921,28 +3328,28 @@ namespace Chaos
 			const FReal YApplied = FVector::DotProduct(AppliedVelocity, YDir);
 			const FReal ZApplied = FVector::DotProduct(AppliedVelocity, ZDir);
 
-			DrawArrow(PDI, ReferenceSpaceTransform.GetLocation(), ReferenceSpaceTransform.GetLocation() + LengthScale * XApplied * XDir, FLinearColor::Red);
-			DrawArrow(PDI, ReferenceSpaceTransform.GetLocation(), ReferenceSpaceTransform.GetLocation() + LengthScale * YApplied * YDir, FLinearColor::Green);
-			DrawArrow(PDI, ReferenceSpaceTransform.GetLocation(), ReferenceSpaceTransform.GetLocation() + LengthScale * ZApplied * ZDir, FLinearColor::Blue);
+			DrawArrow(Context.PDI, ReferenceSpaceTransform.GetLocation(), ReferenceSpaceTransform.GetLocation() + LengthScale * XApplied * XDir, FLinearColor::Red);
+			DrawArrow(Context.PDI, ReferenceSpaceTransform.GetLocation(), ReferenceSpaceTransform.GetLocation() + LengthScale * YApplied * YDir, FLinearColor::Green);
+			DrawArrow(Context.PDI, ReferenceSpaceTransform.GetLocation(), ReferenceSpaceTransform.GetLocation() + LengthScale * ZApplied * ZDir, FLinearColor::Blue);
 
 			// Draw OrigVelocity lighter
 			constexpr float OrigColorScale = 0.5f;
 			if (XOrig > XApplied)
 			{
-				DrawArrow(PDI, ReferenceSpaceTransform.GetLocation() + LengthScale * XApplied * XDir, ReferenceSpaceTransform.GetLocation() + LengthScale * XOrig * XDir, FLinearColor::Red * OrigColorScale);
+				DrawArrow(Context.PDI, ReferenceSpaceTransform.GetLocation() + LengthScale * XApplied * XDir, ReferenceSpaceTransform.GetLocation() + LengthScale * XOrig * XDir, FLinearColor::Red * OrigColorScale);
 			}
 			if (YOrig > YApplied)
 			{
-				DrawArrow(PDI, ReferenceSpaceTransform.GetLocation() + LengthScale * YApplied * YDir, ReferenceSpaceTransform.GetLocation() + LengthScale * YOrig * YDir, FLinearColor::Green * OrigColorScale);
+				DrawArrow(Context.PDI, ReferenceSpaceTransform.GetLocation() + LengthScale * YApplied * YDir, ReferenceSpaceTransform.GetLocation() + LengthScale * YOrig * YDir, FLinearColor::Green * OrigColorScale);
 			}
 			if (ZOrig > ZApplied)
 			{
-				DrawArrow(PDI, ReferenceSpaceTransform.GetLocation() + LengthScale * ZApplied * ZDir, ReferenceSpaceTransform.GetLocation() + LengthScale * ZOrig * ZDir, FLinearColor::Blue * OrigColorScale);
+				DrawArrow(Context.PDI, ReferenceSpaceTransform.GetLocation() + LengthScale * ZApplied * ZDir, ReferenceSpaceTransform.GetLocation() + LengthScale * ZOrig * ZDir, FLinearColor::Blue * OrigColorScale);
 			}
 
 
-			DrawArrow(PDI, ReferenceSpaceTransform.GetLocation(), ReferenceSpaceTransform.GetLocation() + OriginalAngularVelocity, FLinearColor::Yellow * OrigColorScale);
-			DrawArrow(PDI, ReferenceSpaceTransform.GetLocation(), ReferenceSpaceTransform.GetLocation() + AppliedAngularVelocity, FLinearColor::Yellow);
+			DrawArrow(Context.PDI, ReferenceSpaceTransform.GetLocation(), ReferenceSpaceTransform.GetLocation() + OriginalAngularVelocity, FLinearColor::Yellow * OrigColorScale);
+			DrawArrow(Context.PDI, ReferenceSpaceTransform.GetLocation(), ReferenceSpaceTransform.GetLocation() + AppliedAngularVelocity, FLinearColor::Yellow);
 
 			// Text
 			const FText Label = FText::Format(FText::FromString("VelocityScale\nX: {0} ({1})\nY: {2} ({3})\nZ: {4} ({5})\nA: {6} ({7})"),
@@ -2954,7 +3361,7 @@ namespace Chaos
 		}
 	}
 
-	void FClothVisualizationNoGC::DrawSelfCollision(FPrimitiveDrawInterface* PDI) const
+	void FClothVisualizationNoGC::DrawSelfCollision(const FDrawContext& Context) const
 	{
 		if (!Solver)
 		{
@@ -2976,37 +3383,46 @@ namespace Chaos
 				// Constraints are locally indexed for new solver
 				const TConstArrayView<Softs::FSolverVec3> Positions = Solver->IsLegacySolver() ? TConstArrayView<Softs::FSolverVec3>(Solver->GetParticleXs()) : Solver->GetParticleXsView(ParticleRangeId);
 				const int32 Offset = Solver->IsLegacySolver() ? ParticleRangeId : 0;
+				const TSet<int32>* const DebugFilterSet = Context.Filter.VertexSet.IsEmpty() ? nullptr : Cloth->GetVertexSetByName(Solver, Context.Filter.VertexSet);
+				const int32 GlobalParticleOffset = Solver->GetGlobalParticleOffset(ParticleRangeId);
+				auto ShouldDrawElement = [&Context, DebugFilterSet, GlobalParticleOffset, Offset](const TVec4<int32>& Element)
+					{
+						return Private::ShouldDrawConstraint4(Context.Filter, DebugFilterSet, Element, GlobalParticleOffset - Offset, -Offset);
+					};
 				const TArray<TVec4<int32>>& Constraints = SelfCollisionConstraints->GetConstraints();
 				const TArray<Softs::FSolverVec3>& Barys = SelfCollisionConstraints->GetBarys();
 				const TArray<bool>& FlipNormals = SelfCollisionConstraints->GetFlipNormals();
 
 				for (int32 Index = 0; Index < Constraints.Num(); ++Index)
 				{
-					const FReal Height = (FReal)SelfCollisionConstraints->GetConstraintThickness(Index);
 					const TVec4<int32>& Constraint = Constraints[Index];
-					const FVec3 Bary(Barys[Index]);
+					if (ShouldDrawElement(Constraint))
+					{
+						const FReal Height = (FReal)SelfCollisionConstraints->GetConstraintThickness(Index);
+						const FVec3 Bary(Barys[Index]);
 
-					const FVector P = GetWorldPosition(*Solver, Positions[Constraint[0]]);
-					const FVector P0 = GetWorldPosition(*Solver, Positions[Constraint[1]]);
-					const FVector P1 = GetWorldPosition(*Solver, Positions[Constraint[2]]);
-					const FVector P2 = GetWorldPosition(*Solver, Positions[Constraint[3]]);
+						const FVector P = GetWorldPosition(*Solver, Positions[Constraint[0]]);
+						const FVector P0 = GetWorldPosition(*Solver, Positions[Constraint[1]]);
+						const FVector P1 = GetWorldPosition(*Solver, Positions[Constraint[2]]);
+						const FVector P2 = GetWorldPosition(*Solver, Positions[Constraint[3]]);
 
-					const FVector Pos0 = P0 * Bary[0] + P1 * Bary[1] + P2 * Bary[2];
+						const FVector Pos0 = P0 * Bary[0] + P1 * Bary[1] + P2 * Bary[2];
 
-					static const FLinearColor Brown(0.1f, 0.05f, 0.f);
-					static const FLinearColor Red(0.3f, 0.f, 0.f);
-					const FTriangle Triangle(P0, P1, P2);
-					const FVector Normal = FlipNormals[Index] ? -Triangle.GetNormal() : Triangle.GetNormal();
+						static const FLinearColor Brown(0.1f, 0.05f, 0.f);
+						static const FLinearColor Red(0.3f, 0.f, 0.f);
+						const FTriangle Triangle(P0, P1, P2);
+						const FVector Normal = FlipNormals[Index] ? -Triangle.GetNormal() : Triangle.GetNormal();
 
-					// Draw point to surface line (=normal)
-					const FVector Pos1 = Pos0 + GetWorldVector(*Solver, Height * Normal);
-					DrawPoint(PDI, Pos0, Brown, nullptr, 2.f);
-					DrawLine(PDI, Pos0, Pos1, FlipNormals[Index] ? Red : Brown);
+						// Draw point to surface line (=normal)
+						const FVector Pos1 = Pos0 + GetWorldVector(*Solver, Height * Normal);
+						DrawPoint(Context.PDI, Pos0, Brown, nullptr, 2.f);
+						DrawLine(Context.PDI, Pos0, Pos1, FlipNormals[Index] ? Red : Brown);
 
-					// Draw pushup to point
-					static const FLinearColor Orange(0.3f, 0.15f, 0.f);
-					DrawPoint(PDI, P, Orange, nullptr, 2.f);
-					DrawLine(PDI, Pos1, P, Orange);
+						// Draw pushup to point
+						static const FLinearColor Orange(0.3f, 0.15f, 0.f);
+						DrawPoint(Context.PDI, P, Orange, nullptr, 2.f);
+						DrawLine(Context.PDI, Pos1, P, Orange);
+					}
 				}
 
 				const TArray<int32>& KinematicCollidingParticles = SelfCollisionConstraints->GetKinematicCollidingParticles();
@@ -3018,8 +3434,7 @@ namespace Chaos
 					{
 						const FVector P = GetWorldPosition(*Solver, Positions[Index1]);
 
-						static const FLinearColor Orange(0.3f, 0.15f, 0.f);
-						DrawPoint(PDI, P, Orange, nullptr, 2.f);
+						bool bShouldDrawPoint = false;
 
 						const TMap<int32, Softs::FSolverReal>& Timers = KinematicColliderTimers[Index1 - Offset];
 						for (const TPair<int32, Softs::FSolverReal>& ElemAndTimer : Timers)
@@ -3028,18 +3443,28 @@ namespace Chaos
 							const int32 Index3 = TriangleMesh->GetElements()[ElemAndTimer.Get<0>()][1];
 							const int32 Index4 = TriangleMesh->GetElements()[ElemAndTimer.Get<0>()][2];
 
-							const Softs::FSolverVec3& P1 = Positions[Index1];
-							const Softs::FSolverVec3& P2 = Positions[Index2];
-							const Softs::FSolverVec3& P3 = Positions[Index3];
-							const Softs::FSolverVec3& P4 = Positions[Index4];
-							Softs::FSolverVec3 Bary;
-							const FVector Pos1 = GetWorldPosition(*Solver, FindClosestPointAndBaryOnTriangle(P2, P3, P4, P1, Bary));
-
-							static const FLinearColor LtRed(0.6f, 0.f, 0.f);
-							static const FLinearColor DkRed(0.3f, 0.f, 0.f);
-							const FLinearColor& Color = ElemAndTimer.Get<1>() > 0.f ? LtRed : DkRed;
-							DrawPoint(PDI, Pos1, Color, nullptr, 2.f);
-							DrawLine(PDI, Pos1, P, Color);
+							if (ShouldDrawElement(TVec4<int32>(Index1, Index2, Index3, Index4)))
+							{
+								bShouldDrawPoint = true;
+								const Softs::FSolverVec3& P1 = Positions[Index1];
+								const Softs::FSolverVec3& P2 = Positions[Index2];
+								const Softs::FSolverVec3& P3 = Positions[Index3];
+								const Softs::FSolverVec3& P4 = Positions[Index4];
+								Softs::FSolverVec3 Bary;
+								const FVector Pos1 = GetWorldPosition(*Solver, FindClosestPointAndBaryOnTriangle(P2, P3, P4, P1, Bary));
+
+								static const FLinearColor LtRed(0.6f, 0.f, 0.f);
+								static const FLinearColor DkRed(0.3f, 0.f, 0.f);
+								const FLinearColor& Color = ElemAndTimer.Get<1>() > 0.f ? LtRed : DkRed;
+								DrawPoint(Context.PDI, Pos1, Color, nullptr, 2.f);
+								DrawLine(Context.PDI, Pos1, P, Color);
+							}
+						}
+
+						if (bShouldDrawPoint)
+						{
+							static const FLinearColor Orange(0.3f, 0.15f, 0.f);
+							DrawPoint(Context.PDI, P, Orange, nullptr, 2.f);
 						}
 					}
 				}
@@ -3051,14 +3476,27 @@ namespace Chaos
 				// Constraints are locally indexed for new solver
 				const TConstArrayView<Softs::FSolverVec3> Positions = Solver->IsLegacySolver() ? TConstArrayView<Softs::FSolverVec3>(Solver->GetParticleXs()) : Solver->GetParticleXsView(ParticleRangeId);
 				const int32 Offset = Solver->IsLegacySolver() ? ParticleRangeId : 0;
+				const TSet<int32>* const DebugFilterSet = Context.Filter.VertexSet.IsEmpty() ? nullptr : Cloth->GetVertexSetByName(Solver, Context.Filter.VertexSet);
+				const int32 GlobalParticleOffset = Solver->GetGlobalParticleOffset(ParticleRangeId);
+				auto ShouldDrawConstraint2 = [&Context, DebugFilterSet, GlobalParticleOffset, Offset](const TVec2<int32>& Element)
+					{
+						return Private::ShouldDrawConstraint2(Context.Filter, DebugFilterSet, Element, GlobalParticleOffset - Offset, -Offset);
+					};
+				auto ShouldDrawParticle = [&Context, DebugFilterSet, GlobalParticleOffset](const int32 Index)
+					{
+						return Private::ShouldDrawParticle(Context.Filter, DebugFilterSet, Index, GlobalParticleOffset);
+					};
 				const TArray<TVec2<int32>>& Constraints = SelfCollisionSphereConstraints->GetConstraints();
 				for (int32 Index = 0; Index < Constraints.Num(); ++Index)
 				{
 					const TVec2<int32>& Constraint = Constraints[Index];
-					const FVector P0 = GetWorldPosition(*Solver, Positions[Constraint[0]]);
-					const FVector P1 = GetWorldPosition(*Solver, Positions[Constraint[1]]);
-					static const FLinearColor Brown(0.1f, 0.05f, 0.f);
-					DrawLine(PDI, P0, P1, Brown);
+					if (ShouldDrawConstraint2(Constraint))
+					{
+						const FVector P0 = GetWorldPosition(*Solver, Positions[Constraint[0]]);
+						const FVector P1 = GetWorldPosition(*Solver, Positions[Constraint[1]]);
+						static const FLinearColor Brown(0.1f, 0.05f, 0.f);
+						DrawLine(Context.PDI, P0, P1, Brown);
+					}
 				}
 
 				if (const TSet<int32>* const VertexSet = SelfCollisionSphereConstraints->GetVertexSet())
@@ -3066,16 +3504,19 @@ namespace Chaos
 					const FReal Radius = (FReal)SelfCollisionSphereConstraints->GetRadius() * Solver->GetLocalSpaceScale();
 					for (const int32 Vertex : *VertexSet)
 					{
-						const FVector P0 = GetWorldPosition(*Solver, Positions[Vertex + Offset]);
-						DrawSphere(PDI, FSphere(FVector::ZeroVector, Radius), FQuat::Identity,
-							P0, 1., FLinearColor(FColor::Orange));
+						if (ShouldDrawParticle(Vertex))
+						{
+							const FVector P0 = GetWorldPosition(*Solver, Positions[Vertex + Offset]);
+							DrawSphere(Context.PDI, FSphere(FVector::ZeroVector, Radius), FQuat::Identity,
+								P0, 1., FLinearColor(FColor::Orange));
+						}
 					}
 				}
 			}
 		}
 	}
 
-	void FClothVisualizationNoGC::DrawSelfIntersection(FPrimitiveDrawInterface* PDI) const
+	void FClothVisualizationNoGC::DrawSelfIntersection(const FDrawContext& Context) const
 	{
 		if (!Solver)
 		{
@@ -3125,7 +3566,8 @@ namespace Chaos
 					Black
 				};
 
-				auto DrawContour = [this, PDI, &Positions](const TArray<Softs::FPBDTriangleMeshCollisions::FBarycentricPoint>& Contour,
+				// Don't filter contours
+				auto DrawContour = [this, &Context, &Positions](const TArray<Softs::FPBDTriangleMeshCollisions::FBarycentricPoint>& Contour,
 					const FLinearColor& ContourColor)
 				{
 					for (int32 PointIdx = 0; PointIdx < Contour.Num() - 1; ++PointIdx)
@@ -3134,9 +3576,9 @@ namespace Chaos
 						const FVector EndPoint0 = GetWorldPosition(*Solver, (1.f - Point0.Bary[0] - Point0.Bary[1]) * Positions[Point0.Vertices[0]] + Point0.Bary[0] * Positions[Point0.Vertices[1]] + Point0.Bary[1] * Positions[Point0.Vertices[2]]);
 						const Softs::FPBDTriangleMeshCollisions::FBarycentricPoint& Point1 = Contour[PointIdx + 1];
 						const FVector EndPoint1 = GetWorldPosition(*Solver, (1.f - Point1.Bary[0] - Point1.Bary[1]) * Positions[Point1.Vertices[0]] + Point1.Bary[0] * Positions[Point1.Vertices[1]] + Point1.Bary[1] * Positions[Point1.Vertices[2]]);
-						DrawLine(PDI, EndPoint0, EndPoint1, ContourColor);
-						DrawPoint(PDI, EndPoint0, ContourColor, nullptr, 1.f);
-						DrawPoint(PDI, EndPoint1, ContourColor, nullptr, 1.f);
+						DrawLine(Context.PDI, EndPoint0, EndPoint1, ContourColor);
+						DrawPoint(Context.PDI, EndPoint0, ContourColor, nullptr, 1.f);
+						DrawPoint(Context.PDI, EndPoint1, ContourColor, nullptr, 1.f);
 					}
 
 				};
@@ -3153,7 +3595,7 @@ namespace Chaos
 					DrawContour(Contour, Orange);
 				}
 
-				// Draw GIA colors
+				// Draw GIA colors (don't filter)
 				const TConstArrayView<Softs::FPBDTriangleMeshCollisions::FGIAColor>& VertexGIAColors = SelfCollisionInit->GetVertexGIAColors();
 				static const FLinearColor Gray(0.5f, 0.5f, 0.5f);
 				if (VertexGIAColors.Num())
@@ -3169,7 +3611,7 @@ namespace Chaos
 							const bool bAnyBlack = (VertexGIAColors[ParticleIdx].ContourIndexBits & VertexGIAColors[ParticleIdx].ColorBits);
 							const FLinearColor& VertColor = bIsLoop ? Red : bIsBoundary ? Blue : (bAnyWhite && bAnyBlack) ? Gray : bAnyWhite ? White : Black;
 						
-							DrawPoint(PDI, GetWorldPosition(*Solver, Positions[ParticleIdx]), VertColor, nullptr, 5.f);
+							DrawPoint(Context.PDI, GetWorldPosition(*Solver, Positions[ParticleIdx]), VertColor, nullptr, 5.f);
 						}
 					}
 				}
@@ -3185,14 +3627,14 @@ namespace Chaos
 							const bool bAnyWhite = (TriangleGIAColors[TriangleIdx].ContourIndexBits & ~TriangleGIAColors[TriangleIdx].ColorBits);
 							const bool bAnyBlack = (TriangleGIAColors[TriangleIdx].ContourIndexBits & TriangleGIAColors[TriangleIdx].ColorBits);
 							const FLinearColor& TriColor = bIsLoop ? Red : (bAnyWhite && bAnyBlack) ? Gray : bAnyWhite ? White : Black;
-							DrawLine(PDI, GetWorldPosition(*Solver, Positions[Elements[TriangleIdx][0]]), GetWorldPosition(*Solver, Positions[Elements[TriangleIdx][1]]), TriColor);
-							DrawLine(PDI, GetWorldPosition(*Solver, Positions[Elements[TriangleIdx][1]]), GetWorldPosition(*Solver, Positions[Elements[TriangleIdx][2]]), TriColor);
-							DrawLine(PDI, GetWorldPosition(*Solver, Positions[Elements[TriangleIdx][0]]), GetWorldPosition(*Solver, Positions[Elements[TriangleIdx][2]]), TriColor);
+							DrawLine(Context.PDI, GetWorldPosition(*Solver, Positions[Elements[TriangleIdx][0]]), GetWorldPosition(*Solver, Positions[Elements[TriangleIdx][1]]), TriColor);
+							DrawLine(Context.PDI, GetWorldPosition(*Solver, Positions[Elements[TriangleIdx][1]]), GetWorldPosition(*Solver, Positions[Elements[TriangleIdx][2]]), TriColor);
+							DrawLine(Context.PDI, GetWorldPosition(*Solver, Positions[Elements[TriangleIdx][0]]), GetWorldPosition(*Solver, Positions[Elements[TriangleIdx][2]]), TriColor);
 						}
 					}
 				}
 
-				// Draw contour minimization gradients
+				// Draw contour minimization gradients (don't filter)
 				const TArray<Softs::FPBDTriangleMeshCollisions::FContourMinimizationIntersection>& ContourMinimizationIntersections = SelfCollisionInit->GetContourMinimizationIntersections();
 				constexpr FReal MaxDrawImpulse = 1.;
 				constexpr FReal RegularizeEpsilonSq = 1.;
@@ -3206,16 +3648,16 @@ namespace Chaos
 					const FVector EdgeCenter = GetWorldPosition(*Solver, .5 * (Positions[Intersection.EdgeVertices[0]] + Positions[Intersection.EdgeVertices[1]]));
 					const FVector TriCenter = GetWorldPosition(*Solver, (Positions[Intersection.FaceVertices[0]] + Positions[Intersection.FaceVertices[1]] + Positions[Intersection.FaceVertices[2]]) / 3.);
 
-					DrawPoint(PDI, EdgeCenter, Green, nullptr, 2.f);
-					DrawLine(PDI, EdgeCenter, EdgeCenter + Delta, Green);
-					DrawPoint(PDI, TriCenter, Green, nullptr, 2.f);
-					DrawLine(PDI, TriCenter, TriCenter - Delta, Green);
+					DrawPoint(Context.PDI, EdgeCenter, Green, nullptr, 2.f);
+					DrawLine(Context.PDI, EdgeCenter, EdgeCenter + Delta, Green);
+					DrawPoint(Context.PDI, TriCenter, Green, nullptr, 2.f);
+					DrawLine(Context.PDI, TriCenter, TriCenter - Delta, Green);
 				}
 			}
 		}
 	}
 
-	void FClothVisualizationNoGC::DrawSelfCollisionThickness(FPrimitiveDrawInterface* PDI) const
+	void FClothVisualizationNoGC::DrawSelfCollisionThickness(const FDrawContext& Context) const
 	{
 		if (!Solver)
 		{
@@ -3223,9 +3665,9 @@ namespace Chaos
 		}
 
 #if WITH_EDITOR
-		if (FMaterials::GetInstance().GetClothMaterialColor())
+		if (FMaterials::GetInstance().GetClothMaterialColor() && Context.PDI)
 		{
-			FDynamicMeshBuilder MeshBuilder(PDI->View->GetFeatureLevel());
+			FDynamicMeshBuilder MeshBuilder(Context.PDI->View->GetFeatureLevel());
 
 			for (const FClothingSimulationCloth* const Cloth : Solver->GetCloths())
 			{
@@ -3240,6 +3682,12 @@ namespace Chaos
 				if (const Softs::FPBDCollisionSpringConstraints* const SelfCollisionConstraints = ClothConstraints.GetSelfCollisionConstraints().Get())
 				{
 					const int32 Offset = Solver->IsLegacySolver() ? ParticleRangeId : 0;
+					const TSet<int32>* const DebugFilterSet = Context.Filter.VertexSet.IsEmpty() ? nullptr : Cloth->GetVertexSetByName(Solver, Context.Filter.VertexSet);
+					const int32 GlobalParticleOffset = Solver->GetGlobalParticleOffset(ParticleRangeId);
+					auto ShouldDrawEdge = [&Context, DebugFilterSet, GlobalParticleOffset, Offset](const TVec2<int32>& Element)
+						{
+							return Private::ShouldDrawConstraint2(Context.Filter, DebugFilterSet, Element, GlobalParticleOffset - Offset, -Offset);
+						};
 					const TConstArrayView<Softs::FSolverVec3> Positions = Cloth->GetParticlePositions(Solver);
 					const TConstArrayView<Softs::FSolverReal> InvMasses = Cloth->GetParticleInvMasses(Solver);
 					const TConstArrayView<int32>& WeightMap = Cloth->GetFaceIntMapByProperty(Solver, Softs::FPBDCollisionSpringConstraints::SelfCollisionLayersName);
@@ -3252,7 +3700,7 @@ namespace Chaos
 					
 						const bool bIsKinematic0 = (InvMasses[Edge[0] - Offset] == (Softs::FSolverReal)0.);
 						const bool bIsKinematic1 = (InvMasses[Edge[1] - Offset] == (Softs::FSolverReal)0.);
-						if (bIsKinematic0 && bIsKinematic1)
+						if ((bIsKinematic0 && bIsKinematic1) || !ShouldDrawEdge(Edge))
 						{
 							continue;
 						}
@@ -3271,7 +3719,7 @@ namespace Chaos
 				}
 			}
 			const FMatrix LocalSimSpaceToWorld = GetLocalSpaceToWorld(*Solver);
-			MeshBuilder.Draw(PDI, LocalSimSpaceToWorld, FMaterials::GetInstance().GetClothMaterialColor()->GetRenderProxy(), Chaos::Private::GetDepthPriority(), false, false);
+			MeshBuilder.Draw(Context.PDI, LocalSimSpaceToWorld, FMaterials::GetInstance().GetClothMaterialColor()->GetRenderProxy(), Chaos::Private::GetDepthPriority(), false, false);
 		}
 		else
 #endif
@@ -3289,6 +3737,16 @@ namespace Chaos
 				if (const Softs::FPBDCollisionSpringConstraints* const SelfCollisionConstraints = ClothConstraints.GetSelfCollisionConstraints().Get())
 				{
 					const int32 Offset = Solver->IsLegacySolver() ? ParticleRangeId : 0;
+					const TSet<int32>* const DebugFilterSet = Context.Filter.VertexSet.IsEmpty() ? nullptr : Cloth->GetVertexSetByName(Solver, Context.Filter.VertexSet);
+					const int32 GlobalParticleOffset = Solver->GetGlobalParticleOffset(ParticleRangeId);
+					auto ShouldDrawEdge = [&Context, DebugFilterSet, GlobalParticleOffset, Offset](const TVec2<int32>& Element)
+						{
+							return Private::ShouldDrawConstraint2(Context.Filter, DebugFilterSet, Element, GlobalParticleOffset - Offset, -Offset);
+						};
+					auto ShouldDrawParticle = [&Context, DebugFilterSet, GlobalParticleOffset](const int32 Index)
+						{
+							return Private::ShouldDrawParticle(Context.Filter, DebugFilterSet, Index, GlobalParticleOffset);
+						};
 					const TConstArrayView<Softs::FSolverVec3> Positions = Cloth->GetParticlePositions(Solver);
 					const TConstArrayView<Softs::FSolverReal> InvMasses = Cloth->GetParticleInvMasses(Solver);
 
@@ -3297,7 +3755,7 @@ namespace Chaos
 					{
 						const bool bIsKinematic0 = (InvMasses[Edge[0] - Offset] == (Softs::FSolverReal)0.);
 						const bool bIsKinematic1 = (InvMasses[Edge[1] - Offset] == (Softs::FSolverReal)0.);
-						if (bIsKinematic0 && bIsKinematic1)
+						if ((bIsKinematic0 && bIsKinematic1) || !ShouldDrawEdge(Edge))
 						{
 							continue;
 						}
@@ -3307,13 +3765,13 @@ namespace Chaos
 
 						const FReal Radius1 = (FReal)SelfCollisionConstraints->GetParticleThickness(Edge[0]) * Solver->GetLocalSpaceScale();
 						const FReal Radius2 = (FReal)SelfCollisionConstraints->GetParticleThickness(Edge[1]) * Solver->GetLocalSpaceScale();
-						DrawTaperedCylinder(PDI, GetWorldPosition(*Solver, Position1), GetWorldPosition(*Solver, Position2), Radius1, Radius2, 6, FLinearColor::Gray);
+						DrawTaperedCylinder(Context.PDI, GetWorldPosition(*Solver, Position1), GetWorldPosition(*Solver, Position2), Radius1, Radius2, 6, FLinearColor::Gray);
 					}
 
 					for (int32 VertexIndex = 0; VertexIndex < Positions.Num(); ++VertexIndex)
 					{
 						const bool bIsKinematic0 = (InvMasses[VertexIndex] == (Softs::FSolverReal)0.);
-						if (bIsKinematic0)
+						if (bIsKinematic0 || !ShouldDrawParticle(VertexIndex))
 						{
 							continue;
 						}
@@ -3321,14 +3779,14 @@ namespace Chaos
 						const FVector Position1(Positions[VertexIndex]);
 						const FReal Radius1 = (FReal)SelfCollisionConstraints->GetParticleThickness(VertexIndex + Offset);
 						const FTransform Transform(FRotator::ZeroRotator, GetWorldPosition(*Solver, Position1), FVector(Solver->GetLocalSpaceScale()));
-						DrawWireSphere(PDI, Transform, FLinearColor::Gray, Radius1, 6, Chaos::Private::GetDepthPriority(), 0.0f, 0.001f, false);
+						DrawWireSphere(Context.PDI, Transform, FLinearColor::Gray, Radius1, 6, Chaos::Private::GetDepthPriority(), 0.0f, 0.001f, false);
 					}
 				}
 			}
 		}
 	}
 
-	void FClothVisualizationNoGC::DrawKinematicColliderWired(FPrimitiveDrawInterface* PDI) const
+	void FClothVisualizationNoGC::DrawKinematicColliderWired(const FDrawContext& Context) const
 	{
 		if (!Solver)
 		{
@@ -3360,15 +3818,15 @@ namespace Chaos
 					const FVector Pos1 = GetWorldPosition(*Solver, Positions[Element.Y - Offset]);
 					const FVector Pos2 = GetWorldPosition(*Solver, Positions[Element.Z - Offset]);
 
-					DrawLine(PDI, Pos0, Pos1, WireframeColor);
-					DrawLine(PDI, Pos1, Pos2, WireframeColor);
-					DrawLine(PDI, Pos2, Pos0, WireframeColor);
+					DrawLine(Context.PDI, Pos0, Pos1, WireframeColor);
+					DrawLine(Context.PDI, Pos1, Pos2, WireframeColor);
+					DrawLine(Context.PDI, Pos2, Pos0, WireframeColor);
 				}
 			}
 		}
 	}
 
-	void FClothVisualizationNoGC::DrawClothClothConstraints(FPrimitiveDrawInterface* PDI) const
+	void FClothVisualizationNoGC::DrawClothClothConstraints(const FDrawContext& Context) const
 	{
 		if (!Solver)
 		{
@@ -3385,25 +3843,35 @@ namespace Chaos
 			}
 
 			const FClothConstraints& ClothConstraints = Solver->GetClothConstraints(ParticleRangeId);
+			const int32 Offset = Solver->IsLegacySolver() ? ParticleRangeId : 0;
+			const TSet<int32>* const DebugFilterSet = Context.Filter.VertexSet.IsEmpty() ? nullptr : Cloth->GetVertexSetByName(Solver, Context.Filter.VertexSet);
+			const int32 GlobalParticleOffset = Solver->GetGlobalParticleOffset(ParticleRangeId);
 			if (const Softs::FXPBDVertexConstraints* const ClothVertexConstraints = ClothConstraints.GetClothVertexSpringConstraints().Get())
 			{
 				const TConstArrayView<Softs::FSolverVec3> Positions = Cloth->GetParticlePositions(Solver);
 				const TArray<TVec2<int32>>& Constraints = ClothVertexConstraints->GetConstraints();
 				const TArray<Softs::FSolverReal>& SpringLengths = ClothVertexConstraints->GetSpringLengths();
+				auto ShouldDrawConstraint2 = [&Context, DebugFilterSet, GlobalParticleOffset, Offset](const TVec2<int32>& Element)
+					{
+						return Private::ShouldDrawConstraint2(Context.Filter, DebugFilterSet, Element, GlobalParticleOffset - Offset, -Offset);
+					};
 				for (int32 ConstraintIndex = 0; ConstraintIndex < Constraints.Num(); ++ConstraintIndex)
 				{
-					const Softs::FSolverVec3& P1 = Positions[Constraints[ConstraintIndex][0]];
-					const Softs::FSolverVec3& P2 = Positions[Constraints[ConstraintIndex][1]];
-
-					const Softs::FSolverReal DistSq = Softs::FSolverVec3::DistSquared(P1, P2);
-					const Softs::FSolverReal Stiffness = DistSq >= FMath::Square(SpringLengths[ConstraintIndex]) ? ClothVertexConstraints->GetExtensionStiffness(ConstraintIndex) : ClothVertexConstraints->GetCompressionStiffness(ConstraintIndex);
-					const bool bIsActive = Stiffness > (Softs::FSolverReal)0.f;
-					const FVec3 Pos0 = GetWorldPosition(*Solver, P1);
-					const FVec3 Pos1 = GetWorldPosition(*Solver, P2);
-					const FLinearColor& ConstraintColor = bIsActive ? ActiveConstraintColor : InactiveConstraintColor;
-					DrawLine(PDI, Pos0, Pos1, ConstraintColor);
-					DrawPoint(PDI, Pos0, ConstraintColor, nullptr);
-					DrawPoint(PDI, Pos1, ConstraintColor, nullptr);
+					if (ShouldDrawConstraint2(Constraints[ConstraintIndex]))
+					{
+						const Softs::FSolverVec3& P1 = Positions[Constraints[ConstraintIndex][0]];
+						const Softs::FSolverVec3& P2 = Positions[Constraints[ConstraintIndex][1]];
+
+						const Softs::FSolverReal DistSq = Softs::FSolverVec3::DistSquared(P1, P2);
+						const Softs::FSolverReal Stiffness = DistSq >= FMath::Square(SpringLengths[ConstraintIndex]) ? ClothVertexConstraints->GetExtensionStiffness(ConstraintIndex) : ClothVertexConstraints->GetCompressionStiffness(ConstraintIndex);
+						const bool bIsActive = Stiffness > (Softs::FSolverReal)0.f;
+						const FVec3 Pos0 = GetWorldPosition(*Solver, P1);
+						const FVec3 Pos1 = GetWorldPosition(*Solver, P2);
+						const FLinearColor& ConstraintColor = bIsActive ? ActiveConstraintColor : InactiveConstraintColor;
+						DrawLine(Context.PDI, Pos0, Pos1, ConstraintColor);
+						DrawPoint(Context.PDI, Pos0, ConstraintColor, nullptr);
+						DrawPoint(Context.PDI, Pos1, ConstraintColor, nullptr);
+					}
 				}
 			}
 			if (const Softs::FXPBDVertexFaceConstraints* const ClothVertexFaceConstraints = ClothConstraints.GetClothVertexFaceSpringConstraints().Get())
@@ -3412,24 +3880,31 @@ namespace Chaos
 				const TArray<TVec4<int32>>& Constraints = ClothVertexFaceConstraints->GetConstraints();
 				const TArray<TVec4<Softs::FSolverReal>>& Weights = ClothVertexFaceConstraints->GetWeights();
 				const TArray<Softs::FSolverReal>& SpringLengths = ClothVertexFaceConstraints->GetSpringLengths();
+				auto ShouldDrawConstraint4 = [&Context, DebugFilterSet, GlobalParticleOffset, Offset](const TVec4<int32>& Element)
+					{
+						return Private::ShouldDrawConstraint4(Context.Filter, DebugFilterSet, Element, GlobalParticleOffset - Offset, -Offset);
+					};
 				for (int32 ConstraintIndex = 0; ConstraintIndex < Constraints.Num(); ++ConstraintIndex)
 				{
-					const Softs::FSolverVec3& P1 = Positions[Constraints[ConstraintIndex][0]];
-					const Softs::FSolverVec3& P2 = Positions[Constraints[ConstraintIndex][1]];
-					const Softs::FSolverVec3& P3 = Positions[Constraints[ConstraintIndex][2]];
-					const Softs::FSolverVec3& P4 = Positions[Constraints[ConstraintIndex][3]];
-
-					const Softs::FSolverVec3 P = -Weights[ConstraintIndex][1] * P2 - Weights[ConstraintIndex][2] * P3 - Weights[ConstraintIndex][3] * P4;
-
-					const Softs::FSolverReal DistSq = Softs::FSolverVec3::DistSquared(P1, P);
-					const Softs::FSolverReal Stiffness = DistSq >= FMath::Square(SpringLengths[ConstraintIndex]) ? ClothVertexFaceConstraints->GetExtensionStiffness(ConstraintIndex) : ClothVertexFaceConstraints->GetCompressionStiffness(ConstraintIndex);
-					const bool bIsActive = Stiffness > (Softs::FSolverReal)0.f;
-					const FVec3 Pos0 = GetWorldPosition(*Solver, P1);
-					const FVec3 Pos1 = GetWorldPosition(*Solver, P);
-					const FLinearColor& ConstraintColor = bIsActive ? ActiveConstraintColor : InactiveConstraintColor;
-					DrawLine(PDI, Pos0, Pos1, ConstraintColor);
-					DrawPoint(PDI, Pos0, ConstraintColor, nullptr);
-					DrawPoint(PDI, Pos1, ConstraintColor, nullptr);
+					if (ShouldDrawConstraint4(Constraints[ConstraintIndex]))
+					{
+						const Softs::FSolverVec3& P1 = Positions[Constraints[ConstraintIndex][0]];
+						const Softs::FSolverVec3& P2 = Positions[Constraints[ConstraintIndex][1]];
+						const Softs::FSolverVec3& P3 = Positions[Constraints[ConstraintIndex][2]];
+						const Softs::FSolverVec3& P4 = Positions[Constraints[ConstraintIndex][3]];
+
+						const Softs::FSolverVec3 P = -Weights[ConstraintIndex][1] * P2 - Weights[ConstraintIndex][2] * P3 - Weights[ConstraintIndex][3] * P4;
+
+						const Softs::FSolverReal DistSq = Softs::FSolverVec3::DistSquared(P1, P);
+						const Softs::FSolverReal Stiffness = DistSq >= FMath::Square(SpringLengths[ConstraintIndex]) ? ClothVertexFaceConstraints->GetExtensionStiffness(ConstraintIndex) : ClothVertexFaceConstraints->GetCompressionStiffness(ConstraintIndex);
+						const bool bIsActive = Stiffness > (Softs::FSolverReal)0.f;
+						const FVec3 Pos0 = GetWorldPosition(*Solver, P1);
+						const FVec3 Pos1 = GetWorldPosition(*Solver, P);
+						const FLinearColor& ConstraintColor = bIsActive ? ActiveConstraintColor : InactiveConstraintColor;
+						DrawLine(Context.PDI, Pos0, Pos1, ConstraintColor);
+						DrawPoint(Context.PDI, Pos0, ConstraintColor, nullptr);
+						DrawPoint(Context.PDI, Pos1, ConstraintColor, nullptr);
+					}
 				}
 			}
 			if (const Softs::FXPBDFaceConstraints* const ClothFaceConstraints = ClothConstraints.GetClothFaceSpringConstraints().Get())
@@ -3438,27 +3913,35 @@ namespace Chaos
 				const TArray<TVector<int32,6>>& Constraints = ClothFaceConstraints->GetConstraints();
 				const TArray<TVector<Softs::FSolverReal,6>>& Weights = ClothFaceConstraints->GetWeights();
 				const TArray<Softs::FSolverReal>& SpringLengths = ClothFaceConstraints->GetSpringLengths();
+				auto ShouldDrawConstraint6 = [&Context, DebugFilterSet, GlobalParticleOffset, Offset](const TVector<int32, 6>& Element)
+					{
+						return Private::ShouldDrawElement(Context.Filter, DebugFilterSet, TVec3<int32>(Element[0], Element[1], Element[2]), GlobalParticleOffset - Offset, -Offset) ||
+							Private::ShouldDrawElement(Context.Filter, DebugFilterSet, TVec3<int32>(Element[3], Element[4], Element[5]), GlobalParticleOffset - Offset, -Offset);
+					};
 				for (int32 ConstraintIndex = 0; ConstraintIndex < Constraints.Num(); ++ConstraintIndex)
 				{
-					const Softs::FSolverVec3& P1 = Positions[Constraints[ConstraintIndex][0]];
-					const Softs::FSolverVec3& P2 = Positions[Constraints[ConstraintIndex][1]];
-					const Softs::FSolverVec3& P3 = Positions[Constraints[ConstraintIndex][2]];
-					const Softs::FSolverVec3& P4 = Positions[Constraints[ConstraintIndex][3]];
-					const Softs::FSolverVec3& P5 = Positions[Constraints[ConstraintIndex][4]];
-					const Softs::FSolverVec3& P6 = Positions[Constraints[ConstraintIndex][5]];
-
-					const Softs::FSolverVec3 Ps =  Weights[ConstraintIndex][0] * P1 + Weights[ConstraintIndex][1] * P2 + Weights[ConstraintIndex][2] * P3;
-					const Softs::FSolverVec3 Pt = -Weights[ConstraintIndex][3] * P4 - Weights[ConstraintIndex][4] * P5 - Weights[ConstraintIndex][5] * P6;
-
-					const Softs::FSolverReal DistSq = Softs::FSolverVec3::DistSquared(Ps, Pt);
-					const Softs::FSolverReal Stiffness = DistSq >= FMath::Square(SpringLengths[ConstraintIndex]) ? ClothFaceConstraints->GetExtensionStiffness(ConstraintIndex) : ClothFaceConstraints->GetCompressionStiffness(ConstraintIndex);
-					const bool bIsActive = Stiffness > (Softs::FSolverReal)0.f;
-					const FVec3 Pos0 = GetWorldPosition(*Solver, Ps);
-					const FVec3 Pos1 = GetWorldPosition(*Solver, Pt);
-					const FLinearColor& ConstraintColor = bIsActive ? ActiveConstraintColor : InactiveConstraintColor;
-					DrawLine(PDI, Pos0, Pos1, ConstraintColor);
-					DrawPoint(PDI, Pos0, ConstraintColor, nullptr);
-					DrawPoint(PDI, Pos1, ConstraintColor, nullptr);
+					if (ShouldDrawConstraint6(Constraints[ConstraintIndex]))
+					{
+						const Softs::FSolverVec3& P1 = Positions[Constraints[ConstraintIndex][0]];
+						const Softs::FSolverVec3& P2 = Positions[Constraints[ConstraintIndex][1]];
+						const Softs::FSolverVec3& P3 = Positions[Constraints[ConstraintIndex][2]];
+						const Softs::FSolverVec3& P4 = Positions[Constraints[ConstraintIndex][3]];
+						const Softs::FSolverVec3& P5 = Positions[Constraints[ConstraintIndex][4]];
+						const Softs::FSolverVec3& P6 = Positions[Constraints[ConstraintIndex][5]];
+
+						const Softs::FSolverVec3 Ps = Weights[ConstraintIndex][0] * P1 + Weights[ConstraintIndex][1] * P2 + Weights[ConstraintIndex][2] * P3;
+						const Softs::FSolverVec3 Pt = -Weights[ConstraintIndex][3] * P4 - Weights[ConstraintIndex][4] * P5 - Weights[ConstraintIndex][5] * P6;
+
+						const Softs::FSolverReal DistSq = Softs::FSolverVec3::DistSquared(Ps, Pt);
+						const Softs::FSolverReal Stiffness = DistSq >= FMath::Square(SpringLengths[ConstraintIndex]) ? ClothFaceConstraints->GetExtensionStiffness(ConstraintIndex) : ClothFaceConstraints->GetCompressionStiffness(ConstraintIndex);
+						const bool bIsActive = Stiffness > (Softs::FSolverReal)0.f;
+						const FVec3 Pos0 = GetWorldPosition(*Solver, Ps);
+						const FVec3 Pos1 = GetWorldPosition(*Solver, Pt);
+						const FLinearColor& ConstraintColor = bIsActive ? ActiveConstraintColor : InactiveConstraintColor;
+						DrawLine(Context.PDI, Pos0, Pos1, ConstraintColor);
+						DrawPoint(Context.PDI, Pos0, ConstraintColor, nullptr);
+						DrawPoint(Context.PDI, Pos1, ConstraintColor, nullptr);
+					}
 				}
 			}
 			if (const Softs::FPBDVertexFaceRepulsionConstraints* const RepulsionConstraints = ClothConstraints.GetRepulsionConstraints().Get())
@@ -3466,29 +3949,36 @@ namespace Chaos
 				const TConstArrayView<Softs::FSolverVec3> Positions = Cloth->GetParticlePositions(Solver);
 				const TArray<TVec4<int32>>& Constraints = RepulsionConstraints->GetConstraints();
 				const TArray<bool>& IsActive = RepulsionConstraints->GetConstraintIsActive();
+				auto ShouldDrawConstraint4 = [&Context, DebugFilterSet, GlobalParticleOffset, Offset](const TVec4<int32>& Element)
+					{
+						return Private::ShouldDrawConstraint4(Context.Filter, DebugFilterSet, Element, GlobalParticleOffset - Offset, -Offset);
+					};
 				for (int32 ConstraintIndex = 0; ConstraintIndex < Constraints.Num(); ++ConstraintIndex)
 				{
-					const Softs::FSolverVec3& P1 = Positions[Constraints[ConstraintIndex][0]];
-					const Softs::FSolverVec3& P2 = Positions[Constraints[ConstraintIndex][1]];
-					const Softs::FSolverVec3& P3 = Positions[Constraints[ConstraintIndex][2]];
-					const Softs::FSolverVec3& P4 = Positions[Constraints[ConstraintIndex][3]];
-
-					FSolverVec3 Bary;
-					const Softs::FSolverVec3 P = FindClosestPointAndBaryOnTriangle(P2, P3, P4, P1, Bary);
-
-					const Softs::FSolverReal DistSq = Softs::FSolverVec3::DistSquared(P1, P);
-					const FVec3 Pos0 = GetWorldPosition(*Solver, P1);
-					const FVec3 Pos1 = GetWorldPosition(*Solver, P);
-					const FLinearColor& ConstraintColor = IsActive[ConstraintIndex] ? ActiveConstraintColor : InactiveConstraintColor;
-					DrawLine(PDI, Pos0, Pos1, ConstraintColor);
-					DrawPoint(PDI, Pos0, ConstraintColor, nullptr);
-					DrawPoint(PDI, Pos1, ConstraintColor, nullptr);
+					if (ShouldDrawConstraint4(Constraints[ConstraintIndex]))
+					{
+						const Softs::FSolverVec3& P1 = Positions[Constraints[ConstraintIndex][0]];
+						const Softs::FSolverVec3& P2 = Positions[Constraints[ConstraintIndex][1]];
+						const Softs::FSolverVec3& P3 = Positions[Constraints[ConstraintIndex][2]];
+						const Softs::FSolverVec3& P4 = Positions[Constraints[ConstraintIndex][3]];
+
+						FSolverVec3 Bary;
+						const Softs::FSolverVec3 P = FindClosestPointAndBaryOnTriangle(P2, P3, P4, P1, Bary);
+
+						const Softs::FSolverReal DistSq = Softs::FSolverVec3::DistSquared(P1, P);
+						const FVec3 Pos0 = GetWorldPosition(*Solver, P1);
+						const FVec3 Pos1 = GetWorldPosition(*Solver, P);
+						const FLinearColor& ConstraintColor = IsActive[ConstraintIndex] ? ActiveConstraintColor : InactiveConstraintColor;
+						DrawLine(Context.PDI, Pos0, Pos1, ConstraintColor);
+						DrawPoint(Context.PDI, Pos0, ConstraintColor, nullptr);
+						DrawPoint(Context.PDI, Pos1, ConstraintColor, nullptr);
+					}
 				}
 			}
 		}
 	}
 
-	void FClothVisualizationNoGC::DrawTeleportReset(FPrimitiveDrawInterface* PDI) const
+	void FClothVisualizationNoGC::DrawTeleportReset(const FDrawContext& Context) const
 	{
 		if (!Solver)
 		{
@@ -3533,14 +4023,14 @@ namespace Chaos
 				const FVector Pos1 = GetWorldPosition(*Solver, Positions[Element.Y - Offset]);
 				const FVector Pos2 = GetWorldPosition(*Solver, Positions[Element.Z - Offset]);
 
-				DrawLine(PDI, Pos0, Pos1, Color);
-				DrawLine(PDI, Pos1, Pos2, Color);
-				DrawLine(PDI, Pos2, Pos0, Color);
+				DrawLine(Context.PDI, Pos0, Pos1, Color);
+				DrawLine(Context.PDI, Pos1, Pos2, Color);
+				DrawLine(Context.PDI, Pos2, Pos0, Color);
 			}
 		}
 	}
 
-	void FClothVisualizationNoGC::DrawExtremlyDeformedEdges(FPrimitiveDrawInterface* PDI) const
+	void FClothVisualizationNoGC::DrawExtremlyDeformedEdges(const FDrawContext& Context) const
 	{
 		static const FLinearColor DynamicColor = FColor::Red;
 		static const FLinearColor KinematicColor = FColor::Purple;
@@ -3554,6 +4044,12 @@ namespace Chaos
 			}
 			// Elements are local indexed for new solver
 			const int32 Offset = Solver->IsLegacySolver() ? ParticleRangeId : 0;
+			const TSet<int32>* const DebugFilterSet = Context.Filter.VertexSet.IsEmpty() ? nullptr : Cloth->GetVertexSetByName(Solver, Context.Filter.VertexSet);
+			const int32 GlobalParticleOffset = Solver->GetGlobalParticleOffset(ParticleRangeId);
+			auto ShouldDrawConstraint2 = [&Context, DebugFilterSet, GlobalParticleOffset, Offset](const TVec2<int32>& Element)
+				{
+					return Private::ShouldDrawConstraint2(Context.Filter, DebugFilterSet, Element, GlobalParticleOffset - Offset, -Offset);
+				};
 			const FClothConstraints& ClothConstraint = Solver->GetClothConstraints(ParticleRangeId);
 			if (const Softs::FPBDExtremeDeformationConstraints* const ExtremeDeformationConstraints = Solver->GetClothConstraints(ParticleRangeId).GetExtremeDeformationConstraints().Get())
 			{
@@ -3566,11 +4062,14 @@ namespace Chaos
 				for (int32 EdgeIndex = 0; EdgeIndex < ExtremelyDeformedEdges.Num(); ++EdgeIndex)
 				{
 					const TVec2<int32>& Edge = ExtremelyDeformedEdges[EdgeIndex];
-					const FVector Pos0 = GetWorldPosition(*Solver, Positions[Edge[0] - Offset]); // TODO: Triangle Mesh shouldn't really be solver dependent (ie not use an offset)
-					const FVector Pos1 = GetWorldPosition(*Solver, Positions[Edge[1] - Offset]);
-					const bool bIsKinematic0 = (InvMasses[Edge[0] - Offset] == (Softs::FSolverReal)0.);
-					const bool bIsKinematic1 = (InvMasses[Edge[1] - Offset] == (Softs::FSolverReal)0.);
-					DrawLine(PDI, Pos0, Pos1, bIsKinematic0 && bIsKinematic1 ? KinematicColor : DynamicColor);
+					if (ShouldDrawConstraint2(Edge))
+					{
+						const FVector Pos0 = GetWorldPosition(*Solver, Positions[Edge[0] - Offset]); // TODO: Triangle Mesh shouldn't really be solver dependent (ie not use an offset)
+						const FVector Pos1 = GetWorldPosition(*Solver, Positions[Edge[1] - Offset]);
+						const bool bIsKinematic0 = (InvMasses[Edge[0] - Offset] == (Softs::FSolverReal)0.);
+						const bool bIsKinematic1 = (InvMasses[Edge[1] - Offset] == (Softs::FSolverReal)0.);
+						DrawLine(Context.PDI, Pos0, Pos1, bIsKinematic0 && bIsKinematic1 ? KinematicColor : DynamicColor);
+					}
 				}
 			}
 		}
@@ -3597,6 +4096,16 @@ namespace Chaos
 		return AllNames.Array();
 	}
 
+	TArray<FString> FClothVisualizationNoGC::GetAllVertexSetNames() const
+	{
+		TSet<FString> AllNames;
+		for (const FClothingSimulationCloth* const Cloth : Solver->GetCloths())
+		{
+			AllNames.Append(Cloth->GetAllVertexSetNames());
+		}
+		return AllNames.Array();
+	}
+
 	TArray<FName> FClothVisualizationNoGC::GetAllAccessoryMeshNames() const
 	{
 		TSet<FName> AllNames;
@@ -3607,14 +4116,14 @@ namespace Chaos
 		return AllNames.Array();
 	}
 
-	void FClothVisualizationNoGC::DrawKinematicColliderShaded(FPrimitiveDrawInterface* PDI) const
+	void FClothVisualizationNoGC::DrawKinematicColliderShaded(const FDrawContext& Context) const
 	{
 		if (!Solver || !FMaterials::GetInstance().GetCollisionMaterial())
 		{
 			return;
 		}
 
-		FDynamicMeshBuilder MeshBuilder(PDI->View->GetFeatureLevel());
+		FDynamicMeshBuilder MeshBuilder(Context.PDI->View->GetFeatureLevel());
 		int32 VertexIndex = 0;
 
 		for (const FClothingSimulationCloth* const Cloth : Solver->GetCloths())
@@ -3652,10 +4161,10 @@ namespace Chaos
 		}
 
 		const FMatrix LocalSimSpaceToWorld = GetLocalSpaceToWorld(*Solver);
-		MeshBuilder.Draw(PDI, LocalSimSpaceToWorld, FMaterials::GetInstance().GetCollisionMaterial()->GetRenderProxy(), Chaos::Private::GetDepthPriority(), false, false);
+		MeshBuilder.Draw(Context.PDI, LocalSimSpaceToWorld, FMaterials::GetInstance().GetCollisionMaterial()->GetRenderProxy(), Chaos::Private::GetDepthPriority(), false, false);
 	}
 
-	void FClothVisualizationNoGC::DrawSimMorphTarget(FPrimitiveDrawInterface* PDI, const FString& Name) const
+	void FClothVisualizationNoGC::DrawSimMorphTarget(const FDrawContext& Context, const FString& Name) const
 	{
 		if (!Solver)
 		{
@@ -3696,19 +4205,28 @@ namespace Chaos
 
 					// Elements are local indexed for new solver
 					const int32 Offset = Solver->IsLegacySolver() ? ParticleRangeId : 0;
+					const TSet<int32>* const DebugFilterSet = Context.Filter.VertexSet.IsEmpty() ? nullptr : Cloth->GetVertexSetByName(Solver, Context.Filter.VertexSet);
+					const int32 GlobalParticleOffset = Solver->GetGlobalParticleOffset(ParticleRangeId);
+					auto ShouldDrawElement = [&Context, DebugFilterSet, GlobalParticleOffset, Offset](const TVec3<int32>& Element)
+						{
+							return Private::ShouldDrawElement(Context.Filter, DebugFilterSet, Element, GlobalParticleOffset - Offset, -Offset);
+						};
 
 					const TConstArrayView<TVec3<int32>> Elements = Cloth->GetTriangleMesh(Solver).GetElements();
 					for (int32 ElementIndex = 0; ElementIndex < Elements.Num(); ++ElementIndex)
 					{
 						const TVec3<int32>& Element = Elements[ElementIndex];
+						if (ShouldDrawElement(Element))
+						{
 
-						const FVector Pos0 = GetWorldPosition(*Solver, Positions[Element.X - Offset]); // TODO: Triangle Mesh shouldn't really be solver dependent (ie not use an offset)
-						const FVector Pos1 = GetWorldPosition(*Solver, Positions[Element.Y - Offset]);
-						const FVector Pos2 = GetWorldPosition(*Solver, Positions[Element.Z - Offset]);
+							const FVector Pos0 = GetWorldPosition(*Solver, Positions[Element.X - Offset]); // TODO: Triangle Mesh shouldn't really be solver dependent (ie not use an offset)
+							const FVector Pos1 = GetWorldPosition(*Solver, Positions[Element.Y - Offset]);
+							const FVector Pos2 = GetWorldPosition(*Solver, Positions[Element.Z - Offset]);
 
-						DrawLine(PDI, Pos0, Pos1, Color);
-						DrawLine(PDI, Pos1, Pos2, Color);
-						DrawLine(PDI, Pos2, Pos0, Color);
+							DrawLine(Context.PDI, Pos0, Pos1, Color);
+							DrawLine(Context.PDI, Pos1, Pos2, Color);
+							DrawLine(Context.PDI, Pos2, Pos0, Color);
+						}
 					}
 				}
 			}
