Chaos Cloth: Add options to filter debug draw to a single vertex or named set of vertices.

diff --git a/Engine/Plugins/ChaosCloth/Source/ChaosCloth/Private/ChaosCloth/ChaosClothingSimulation.cpp b/Engine/Plugins/ChaosCloth/Source/ChaosCloth/Private/ChaosCloth/ChaosClothingSimulation.cpp
--- a/Engine/Plugins/ChaosCloth/Source/ChaosCloth/Private/ChaosCloth/ChaosClothingSimulation.cpp
+++ b/Engine/Plugins/ChaosCloth/Source/ChaosCloth/Private/ChaosCloth/ChaosClothingSimulation.cpp
@@ -72,6 +72,8 @@ FAutoConsoleVariableRef CVarClothUseLegacyEvolution(TEXT("p.ChaosCloth.UseLegacy
 #if CHAOS_DEBUG_DRAW
 namespace ClothingSimulationCVar
 {
+	TAutoConsoleVariable<int32> DebugDrawSingleVertex       (TEXT("p.ChaosCloth.DebugDrawSingleVertex"        ), -1,    TEXT("Filter Chaos Cloth debug draw to a single vertex. Default = -1 is no filtering."), ECVF_Cheat);
+	TAutoConsoleVariable<FString> DebugDrawFilterSet        (TEXT("p.ChaosCloth.DebugDrawFilterSet"           ), "",    TEXT("Filter Chaos Cloth debug draw to a named sim vertex set. Default ='' is no filtering."), ECVF_Cheat);
 	TAutoConsoleVariable<bool> DebugDrawLocalSpace          (TEXT("p.ChaosCloth.DebugDrawLocalSpace"          ), false, TEXT("Whether to debug draw the Chaos Cloth local space"), ECVF_Cheat);
 	TAutoConsoleVariable<bool> DebugDrawBounds              (TEXT("p.ChaosCloth.DebugDrawBounds"              ), false, TEXT("Whether to debug draw the Chaos Cloth bounds"), ECVF_Cheat);
 	TAutoConsoleVariable<bool> DebugDrawGravity             (TEXT("p.ChaosCloth.DebugDrawGravity"             ), false, TEXT("Whether to debug draw the Chaos Cloth gravity acceleration vector"), ECVF_Cheat);
@@ -646,31 +648,39 @@ void FClothingSimulation::Simulate_AnyThread(const IClothingSimulationContext* I
 
 	// Debug draw
 #if CHAOS_DEBUG_DRAW
-	if (ClothingSimulationCVar::DebugDrawLocalSpace          .GetValueOnAnyThread()) { DebugDrawLocalSpace          (); }
-	if (ClothingSimulationCVar::DebugDrawBounds              .GetValueOnAnyThread()) { DebugDrawBounds              (); }
-	if (ClothingSimulationCVar::DebugDrawGravity             .GetValueOnAnyThread()) { DebugDrawGravity             (); }
-	if (ClothingSimulationCVar::DebugDrawPhysMeshWired       .GetValueOnAnyThread()) { DebugDrawPhysMeshWired       (); }
-	if (ClothingSimulationCVar::DebugDrawAnimMeshWired       .GetValueOnAnyThread()) { DebugDrawAnimMeshWired       (); }
-	if (ClothingSimulationCVar::DebugDrawPointVelocities     .GetValueOnAnyThread()) { DebugDrawPointVelocities     (); }
-	if (ClothingSimulationCVar::DebugDrawAnimNormals         .GetValueOnAnyThread()) { DebugDrawAnimNormals         (); }
-	if (ClothingSimulationCVar::DebugDrawAnimVelocities      .GetValueOnAnyThread()) { DebugDrawAnimVelocities      (); }
-	if (ClothingSimulationCVar::DebugDrawPointNormals        .GetValueOnAnyThread()) { DebugDrawPointNormals        (); }
-	if (ClothingSimulationCVar::DebugDrawCollision           .GetValueOnAnyThread()) { DebugDrawCollision           (); }
-	if (ClothingSimulationCVar::DebugDrawBackstops           .GetValueOnAnyThread()) { DebugDrawBackstops           (); }
-	if (ClothingSimulationCVar::DebugDrawBackstopDistances   .GetValueOnAnyThread()) { DebugDrawBackstopDistances   (); }
-	if (ClothingSimulationCVar::DebugDrawMaxDistances        .GetValueOnAnyThread()) { DebugDrawMaxDistances        (); }
-	if (ClothingSimulationCVar::DebugDrawMaxDistanceValues   .GetValueOnAnyThread()) { DebugDrawMaxDistanceValues   (); }
-	if (ClothingSimulationCVar::DebugDrawAnimDrive           .GetValueOnAnyThread()) { DebugDrawAnimDrive           (); }
-	if (ClothingSimulationCVar::DebugDrawEdgeConstraint      .GetValueOnAnyThread()) { DebugDrawEdgeConstraint      (); }
-	if (ClothingSimulationCVar::DebugDrawBendingConstraint   .GetValueOnAnyThread()) { DebugDrawBendingConstraint   (); }
-	if (ClothingSimulationCVar::DebugDrawLongRangeConstraint .GetValueOnAnyThread()) { DebugDrawLongRangeConstraint (); }
-	if (ClothingSimulationCVar::DebugDrawWindForces          .GetValueOnAnyThread()) { DebugDrawWindAndPressureForces(); }
-	if (ClothingSimulationCVar::DebugDrawWindVelocity        .GetValueOnAnyThread()) { DebugDrawWindVelocity        (); }
-	if (ClothingSimulationCVar::DebugDrawSelfCollision       .GetValueOnAnyThread()) { DebugDrawSelfCollision       (); }
-	if (ClothingSimulationCVar::DebugDrawSelfIntersection    .GetValueOnAnyThread()) { DebugDrawSelfIntersection    (); }
-	if (ClothingSimulationCVar::DebugDrawTeleportReset       .GetValueOnAnyThread()) { DebugDrawTeleportReset       (); }
-	if (ClothingSimulationCVar::DebugDrawExtremlyDeformedEdges.GetValueOnAnyThread()) { DebugDrawExtremlyDeformedEdges(); }
-	if (ClothingSimulationCVar::DebugDrawVelocityScale.GetValueOnAnyThread()) { DebugDrawVelocityScale(); }
+	FClothVisualizationNoGC::FDrawContext DrawContext;
+	DrawContext.Filter.SingleVertex = ClothingSimulationCVar::DebugDrawSingleVertex.GetValueOnAnyThread();
+	DrawContext.Filter.VertexSet = ClothingSimulationCVar::DebugDrawFilterSet.GetValueOnAnyThread();
+	FClothVisualizationNoGC::FDrawTextsContext DrawTextsContext;
+	DrawTextsContext.Filter = DrawContext.Filter;
+	
+	if (ClothingSimulationCVar::DebugDrawLocalSpace          .GetValueOnAnyThread()) { Visualization.DrawLocalSpace          (DrawContext); }
+	if (ClothingSimulationCVar::DebugDrawBounds              .GetValueOnAnyThread()) { Visualization.DrawBounds              (DrawContext); }
+	if (ClothingSimulationCVar::DebugDrawGravity             .GetValueOnAnyThread()) { Visualization.DrawGravity             (DrawContext); }
+	if (ClothingSimulationCVar::DebugDrawPhysMeshWired       .GetValueOnAnyThread()) { Visualization.DrawPhysMeshWired       (DrawContext); }
+	if (ClothingSimulationCVar::DebugDrawAnimMeshWired       .GetValueOnAnyThread()) { Visualization.DrawAnimMeshWired       (DrawContext); }
+	if (ClothingSimulationCVar::DebugDrawPointVelocities     .GetValueOnAnyThread()) { Visualization.DrawPointVelocities     (DrawContext); }
+	if (ClothingSimulationCVar::DebugDrawAnimNormals         .GetValueOnAnyThread()) { Visualization.DrawAnimNormals         (DrawContext); }
+	if (ClothingSimulationCVar::DebugDrawAnimVelocities      .GetValueOnAnyThread()) { Visualization.DrawAnimVelocities      (DrawContext); }
+	if (ClothingSimulationCVar::DebugDrawPointNormals        .GetValueOnAnyThread()) { Visualization.DrawPointNormals        (DrawContext); }
+	if (ClothingSimulationCVar::DebugDrawCollision           .GetValueOnAnyThread()) { Visualization.DrawCollision           (DrawContext); }
+	if (ClothingSimulationCVar::DebugDrawBackstops           .GetValueOnAnyThread()) { Visualization.DrawBackstops           (DrawContext); }
+	if (ClothingSimulationCVar::DebugDrawBackstopDistances   .GetValueOnAnyThread()) { Visualization.DrawBackstopDistances   (DrawContext); }
+	if (ClothingSimulationCVar::DebugDrawMaxDistances        .GetValueOnAnyThread()) { Visualization.DrawMaxDistances        (DrawContext); }
+	if (ClothingSimulationCVar::DebugDrawMaxDistanceValues   .GetValueOnAnyThread()) { Visualization.DrawMaxDistanceValues   (DrawTextsContext); }
+	if (ClothingSimulationCVar::DebugDrawAnimDrive           .GetValueOnAnyThread()) { Visualization.DrawAnimDrive           (DrawContext); }
+	if (ClothingSimulationCVar::DebugDrawEdgeConstraint      .GetValueOnAnyThread()) { Visualization.DrawEdgeConstraint      (DrawContext); }
+	if (ClothingSimulationCVar::DebugDrawBendingConstraint   .GetValueOnAnyThread()) { Visualization.DrawBendingConstraint   (DrawContext); }
+	if (ClothingSimulationCVar::DebugDrawLongRangeConstraint .GetValueOnAnyThread()) { Visualization.DrawLongRangeConstraint (DrawContext); }
+	if (ClothingSimulationCVar::DebugDrawWindForces          .GetValueOnAnyThread()) { Visualization.DrawWindAndPressureForces(DrawContext); }
+	if (ClothingSimulationCVar::DebugDrawWindVelocity        .GetValueOnAnyThread()) { Visualization.DrawWindVelocity        (DrawContext); }
+	if (ClothingSimulationCVar::DebugDrawSelfCollision       .GetValueOnAnyThread()) { Visualization.DrawSelfCollision       (DrawContext); }
+	if (ClothingSimulationCVar::DebugDrawSelfIntersection    .GetValueOnAnyThread()) { Visualization.DrawSelfIntersection    (DrawContext); }
+	if (ClothingSimulationCVar::DebugDrawParticleIndices     .GetValueOnAnyThread()) { Visualization.DrawParticleIndices     (DrawTextsContext); }
+	if (ClothingSimulationCVar::DebugDrawElementIndices      .GetValueOnAnyThread()) { Visualization.DrawElementIndices      (DrawTextsContext); }
+	if (ClothingSimulationCVar::DebugDrawTeleportReset       .GetValueOnAnyThread()) { Visualization.DrawTeleportReset       (DrawContext); }
+	if (ClothingSimulationCVar::DebugDrawExtremlyDeformedEdges.GetValueOnAnyThread()) { Visualization.DrawExtremlyDeformedEdges(DrawContext); }
+	if (ClothingSimulationCVar::DebugDrawVelocityScale       .GetValueOnAnyThread()) { Visualization.DrawVelocityScale       (DrawContext); }
 #endif  // #if CHAOS_DEBUG_DRAW
 }
 
diff --git a/Engine/Plugins/ChaosCloth/Source/ChaosCloth/Private/ChaosCloth/ChaosClothingSimulationCloth.cpp b/Engine/Plugins/ChaosCloth/Source/ChaosCloth/Private/ChaosCloth/ChaosClothingSimulationCloth.cpp
--- a/Engine/Plugins/ChaosCloth/Source/ChaosCloth/Private/ChaosCloth/ChaosClothingSimulationCloth.cpp
+++ b/Engine/Plugins/ChaosCloth/Source/ChaosCloth/Private/ChaosCloth/ChaosClothingSimulationCloth.cpp
@@ -930,6 +930,31 @@ TConstArrayView<int32> FClothingSimulationCloth::GetFaceIntMapByProperty(const F
 	return GetFaceIntMapByName(Solver, PropertyString);
 }
 
+TSet<FString> FClothingSimulationCloth::GetAllVertexSetNames() const
+{
+	TSet<FString> Names;
+	for (const TUniquePtr<FLODData>& Data : LODData)
+	{
+		TSet<FString> LODNames;
+		Data->VertexSets.GetKeys(LODNames);
+		Names.Append(MoveTemp(LODNames));
+	}
+	return Names;
+}
+
+const TSet<int32>* FClothingSimulationCloth::GetVertexSetByName(const FClothingSimulationSolver* Solver, const FString& Name) const
+{
+	const int32 LODIndex = LODIndices.FindChecked(Solver);
+	return LODData.IsValidIndex(LODIndex) ? LODData[LODIndex]->VertexSets.FindRef(Name) : nullptr;
+}
+
+const TSet<int32>* FClothingSimulationCloth::GetVertexSetByProperty(const FClothingSimulationSolver* Solver, const FName& Property) const
+{
+	check(Config);
+	const FString PropertyString = Config->GetProperties(GetLODIndex(Solver)).GetStringValue(Property);
+	return GetVertexSetByName(Solver, PropertyString);
+}
+
 const TArray<TConstArrayView<TTuple<int32, int32, float>>>& FClothingSimulationCloth::GetTethers(const FClothingSimulationSolver* Solver) const
 {
 	const int32 LODIndex = LODIndices.FindChecked(Solver);
diff --git a/Engine/Plugins/ChaosCloth/Source/ChaosCloth/Public/ChaosCloth/ChaosClothVisualization.h b/Engine/Plugins/ChaosCloth/Source/ChaosCloth/Public/ChaosCloth/ChaosClothVisualization.h
--- a/Engine/Plugins/ChaosCloth/Source/ChaosCloth/Public/ChaosCloth/ChaosClothVisualization.h
+++ b/Engine/Plugins/ChaosCloth/Source/ChaosCloth/Public/ChaosCloth/ChaosClothVisualization.h
@@ -26,138 +26,467 @@ namespace Chaos
 		CHAOSCLOTH_API explicit FClothVisualizationNoGC(const ::Chaos::FClothingSimulationSolver* InSolver = nullptr);
 		CHAOSCLOTH_API virtual ~FClothVisualizationNoGC();
 
+		struct FDrawFilter
+		{
+			FString VertexSet;
+			int32 SingleVertex = INDEX_NONE;
+		};
+
+		struct FDrawContext
+		{
+			explicit FDrawContext(FPrimitiveDrawInterface* InPDI = nullptr)
+				: PDI(InPDI)
+			{
+			}
+			FPrimitiveDrawInterface* PDI = nullptr;
+			FDrawFilter Filter;
+		};
+
+		struct FDrawTextsContext
+		{
+			explicit FDrawTextsContext(FCanvas* InCanvas = nullptr, const FSceneView* InSceneView = nullptr)
+				: Canvas(InCanvas), SceneView(InSceneView)
+			{
+			}
+			FCanvas* Canvas = nullptr;
+			const FSceneView* SceneView = nullptr;
+			FDrawFilter Filter;
+		};
+
 #if CHAOS_DEBUG_DRAW
 		// Editor & runtime functions
 		CHAOSCLOTH_API void SetSolver(const ::Chaos::FClothingSimulationSolver* InSolver);
 
-		CHAOSCLOTH_API void DrawParticleIndices(FCanvas* Canvas = nullptr, const FSceneView* SceneView = nullptr) const;
-		CHAOSCLOTH_API void DrawElementIndices(FCanvas* Canvas = nullptr, const FSceneView* SceneView = nullptr) const;
-		CHAOSCLOTH_API void DrawMaxDistanceValues(FCanvas* Canvas = nullptr, const FSceneView* SceneView = nullptr) const;
-		CHAOSCLOTH_API void DrawLocalSpaceBoneNames(FCanvas* Canvas = nullptr, const FSceneView* SceneView = nullptr) const;
-
-		CHAOSCLOTH_API void DrawPhysMeshWired(FPrimitiveDrawInterface* PDI = nullptr) const;
-		CHAOSCLOTH_API void DrawAnimMeshWired(FPrimitiveDrawInterface* PDI = nullptr) const;
-		CHAOSCLOTH_API void DrawAnimVelocities(FPrimitiveDrawInterface* PDI = nullptr) const;
-		CHAOSCLOTH_API void DrawAnimNormals(FPrimitiveDrawInterface* PDI, const FReal Length) const;
-		CHAOSCLOTH_API void DrawOpenEdges(FPrimitiveDrawInterface* PDI = nullptr) const;
-		CHAOSCLOTH_API void DrawPointNormals(FPrimitiveDrawInterface* PDI, const FReal Length) const;
-		CHAOSCLOTH_API void DrawPointVelocities(FPrimitiveDrawInterface* PDI = nullptr) const;
-		CHAOSCLOTH_API void DrawCollision(FPrimitiveDrawInterface* PDI, bool bWireframe) const;
-		CHAOSCLOTH_API void DrawBackstops(FPrimitiveDrawInterface* PDI = nullptr) const;
-		CHAOSCLOTH_API void DrawBackstopDistances(FPrimitiveDrawInterface* PDI = nullptr) const;
-		CHAOSCLOTH_API void DrawMaxDistances(FPrimitiveDrawInterface* PDI = nullptr) const;
-		CHAOSCLOTH_API void DrawAnimDrive(FPrimitiveDrawInterface* PDI = nullptr) const;
-		CHAOSCLOTH_API void DrawEdgeConstraint(FPrimitiveDrawInterface* PDI = nullptr) const;
-		CHAOSCLOTH_API void DrawBendingConstraint(FPrimitiveDrawInterface* PDI = nullptr) const;
-		CHAOSCLOTH_API void DrawLongRangeConstraint(FPrimitiveDrawInterface* PDI = nullptr) const;
-		CHAOSCLOTH_API void DrawWindAndPressureForces(FPrimitiveDrawInterface* PDI, const FReal ForceLengthScale) const;
-		CHAOSCLOTH_API void DrawWindVelocity(FPrimitiveDrawInterface* PDI, const FReal LengthScale) const;
-		CHAOSCLOTH_API void DrawLocalSpace(FPrimitiveDrawInterface* PDI, FReal LocalSpaceLengthScale) const;
-		CHAOSCLOTH_API void DrawVelocityScale(FPrimitiveDrawInterface* PDI, const FReal LengthScale) const;
-		CHAOSCLOTH_API void DrawSelfCollision(FPrimitiveDrawInterface* PDI = nullptr) const;
-		CHAOSCLOTH_API void DrawSelfIntersection(FPrimitiveDrawInterface* PDI = nullptr) const;
-		CHAOSCLOTH_API void DrawSelfCollisionThickness(FPrimitiveDrawInterface* PDI = nullptr) const;
-		CHAOSCLOTH_API void DrawKinematicColliderWired(FPrimitiveDrawInterface* PDI = nullptr) const;
-		CHAOSCLOTH_API void DrawBounds(FPrimitiveDrawInterface* PDI = nullptr) const;
-		CHAOSCLOTH_API void DrawGravity(FPrimitiveDrawInterface* PDI = nullptr) const;
-		CHAOSCLOTH_API void DrawFictitiousAngularForces(FPrimitiveDrawInterface* PDI = nullptr) const;
-		CHAOSCLOTH_API void DrawMultiResConstraint(FPrimitiveDrawInterface* PDI = nullptr) const;
-		CHAOSCLOTH_API void DrawClothClothConstraints(FPrimitiveDrawInterface* PDI = nullptr) const;
-		CHAOSCLOTH_API void DrawTeleportReset(FPrimitiveDrawInterface* PDI = nullptr) const;
-		CHAOSCLOTH_API void DrawExtremlyDeformedEdges(FPrimitiveDrawInterface* PDI = nullptr) const;
-		CHAOSCLOTH_API void DrawAccessoryMesh(FPrimitiveDrawInterface* PDI = nullptr, const FName& AccessoryMeshName = NAME_None) const;
-		CHAOSCLOTH_API void DrawAccessoryMeshNormals(FPrimitiveDrawInterface* PDI = nullptr, const FName& AccessoryMeshName = NAME_None, const FReal Length = 20.) const;
+		CHAOSCLOTH_API void DrawParticleIndices(const FDrawTextsContext& Context = FDrawTextsContext()) const;
+		CHAOSCLOTH_API void DrawElementIndices(const FDrawTextsContext& Context = FDrawTextsContext()) const;
+		CHAOSCLOTH_API void DrawMaxDistanceValues(const FDrawTextsContext& Context = FDrawTextsContext()) const;
+		CHAOSCLOTH_API void DrawLocalSpaceBoneNames(const FDrawTextsContext& Context = FDrawTextsContext()) const;
+
+		CHAOSCLOTH_API void DrawPhysMeshWired(const FDrawContext& Context = FDrawContext()) const;
+		CHAOSCLOTH_API void DrawAnimMeshWired(const FDrawContext& Context = FDrawContext()) const;
+		CHAOSCLOTH_API void DrawAnimNormals(const FDrawContext& Context = FDrawContext(), const FReal Length = 20.) const;
+		CHAOSCLOTH_API void DrawAnimVelocities(const FDrawContext& Context = FDrawContext()) const;
+		CHAOSCLOTH_API void DrawOpenEdges(const FDrawContext& Context = FDrawContext()) const;
+		CHAOSCLOTH_API void DrawPointNormals(const FDrawContext& Context = FDrawContext(), const FReal Length = 20.) const;
+		CHAOSCLOTH_API void DrawPointVelocities(const FDrawContext& Context = FDrawContext()) const;
+		CHAOSCLOTH_API void DrawCollision(const FDrawContext& Context = FDrawContext(), const bool bWireframe = false) const;
+		CHAOSCLOTH_API void DrawBackstops(const FDrawContext& Context = FDrawContext()) const;
+		CHAOSCLOTH_API void DrawBackstopDistances(const FDrawContext& Context = FDrawContext()) const;
+		CHAOSCLOTH_API void DrawMaxDistances(const FDrawContext& Context = FDrawContext()) const;
+		CHAOSCLOTH_API void DrawAnimDrive(const FDrawContext& Context = FDrawContext()) const;
+		CHAOSCLOTH_API void DrawEdgeConstraint(const FDrawContext& Context = FDrawContext()) const;
+		CHAOSCLOTH_API void DrawBendingConstraint(const FDrawContext& Context = FDrawContext()) const;
+		CHAOSCLOTH_API void DrawLongRangeConstraint(const FDrawContext& Context = FDrawContext()) const;
+		CHAOSCLOTH_API void DrawWindAndPressureForces(const FDrawContext& Context = FDrawContext(), const FReal LengthScale = 10.) const;
+		CHAOSCLOTH_API void DrawWindVelocity(const FDrawContext& Context = FDrawContext(), const FReal LengthScale = 0.1) const;
+		CHAOSCLOTH_API void DrawLocalSpace(const FDrawContext& Context = FDrawContext(), const FReal LengthScale = 0.1) const;
+		CHAOSCLOTH_API void DrawVelocityScale(const FDrawContext& Context = FDrawContext(), const FReal LengthScale = 1) const;
+		CHAOSCLOTH_API void DrawSelfCollision(const FDrawContext& Context = FDrawContext()) const;
+		CHAOSCLOTH_API void DrawSelfIntersection(const FDrawContext& Context = FDrawContext()) const;
+		CHAOSCLOTH_API void DrawSelfCollisionThickness(const FDrawContext& Context = FDrawContext()) const;
+		CHAOSCLOTH_API void DrawKinematicColliderWired(const FDrawContext& Context = FDrawContext()) const;
+		CHAOSCLOTH_API void DrawBounds(const FDrawContext& Context = FDrawContext()) const;
+		CHAOSCLOTH_API void DrawGravity(const FDrawContext& Context = FDrawContext()) const;
+		CHAOSCLOTH_API void DrawFictitiousAngularForces(const FDrawContext& Context = FDrawContext()) const;
+		CHAOSCLOTH_API void DrawMultiResConstraint(const FDrawContext& Context = FDrawContext()) const;
+		CHAOSCLOTH_API void DrawClothClothConstraints(const FDrawContext& Context = FDrawContext()) const;
+		CHAOSCLOTH_API void DrawTeleportReset(const FDrawContext& Context = FDrawContext()) const;
+		CHAOSCLOTH_API void DrawExtremlyDeformedEdges(const FDrawContext& Context = FDrawContext()) const;
+		CHAOSCLOTH_API void DrawAccessoryMesh(const FDrawContext& Context = FDrawContext(), const FName& AccessoryMeshName = NAME_None) const;
+		CHAOSCLOTH_API void DrawAccessoryMeshNormals(const FDrawContext& Context = FDrawContext(), const FName& AccessoryMeshName = NAME_None, const FReal Length = 20.) const;
+
+		UE_DEPRECATED(5.8, "Use Context version of this method instead.")
+		void DrawParticleIndices(FCanvas* Canvas, const FSceneView* SceneView = nullptr) const
+		{
+			return DrawParticleIndices(FDrawTextsContext(Canvas,SceneView));
+		}
+		UE_DEPRECATED(5.8, "Use Context version of this method instead.")
+		void DrawElementIndices(FCanvas* Canvas, const FSceneView* SceneView = nullptr) const
+		{
+			return DrawElementIndices(FDrawTextsContext(Canvas, SceneView));
+		}
+		UE_DEPRECATED(5.8, "Use Context version of this method instead.")
+		void DrawMaxDistanceValues(FCanvas* Canvas, const FSceneView* SceneView = nullptr) const
+		{
+			return DrawMaxDistanceValues(FDrawTextsContext(Canvas, SceneView));
+		}
+		UE_DEPRECATED(5.8, "Use Context version of this method instead.")
+		void DrawLocalSpaceBoneNames(FCanvas* Canvas, const FSceneView* SceneView = nullptr) const
+		{
+			return DrawLocalSpaceBoneNames(FDrawTextsContext(Canvas, SceneView));
+		}
+
+		UE_DEPRECATED(5.8, "Use Context version of this method instead.")
+		void DrawPhysMeshWired(FPrimitiveDrawInterface* PDI) const
+		{
+			return DrawPhysMeshWired(FDrawContext(PDI));
+		}
+		UE_DEPRECATED(5.8, "Use Context version of this method instead.")
+		void DrawAnimMeshWired(FPrimitiveDrawInterface* PDI) const
+		{
+			return DrawAnimMeshWired(FDrawContext(PDI));
+		}
+		UE_DEPRECATED(5.8, "Use Context version of this method instead.")
+		void DrawAnimNormals(FPrimitiveDrawInterface* PDI, const FReal Length) const
+		{
+			return DrawAnimNormals(FDrawContext(PDI), Length);
+		}
+		UE_DEPRECATED(5.8, "Use Context version of this method instead.")
+		void DrawAnimVelocities(FPrimitiveDrawInterface* PDI) const
+		{
+			return DrawAnimVelocities(FDrawContext(PDI));
+		}
+		UE_DEPRECATED(5.8, "Use Context version of this method instead.")
+		void DrawOpenEdges(FPrimitiveDrawInterface* PDI) const
+		{
+			return DrawOpenEdges(FDrawContext(PDI));
+		}
+		UE_DEPRECATED(5.8, "Use Context version of this method instead.")
+		void DrawPointNormals(FPrimitiveDrawInterface* PDI, const FReal Length) const
+		{
+			return DrawPointNormals(FDrawContext(PDI), Length);
+		}
+		UE_DEPRECATED(5.8, "Use Context version of this method instead.")
+		void DrawPointVelocities(FPrimitiveDrawInterface* PDI) const
+		{
+			return DrawPointVelocities(FDrawContext(PDI));
+		}
+		UE_DEPRECATED(5.8, "Use Context version of this method instead.")
+		void DrawCollision(FPrimitiveDrawInterface* PDI, bool bWireframe) const
+		{
+			return DrawCollision(FDrawContext(PDI), bWireframe);
+		}
+		UE_DEPRECATED(5.8, "Use Context version of this method instead.")
+		void DrawBackstops(FPrimitiveDrawInterface* PDI) const
+		{
+			return DrawBackstops(FDrawContext(PDI));
+		}
+		UE_DEPRECATED(5.8, "Use Context version of this method instead.")
+		void DrawBackstopDistances(FPrimitiveDrawInterface* PDI) const
+		{
+			return DrawBackstopDistances(FDrawContext(PDI));
+		}
+		UE_DEPRECATED(5.8, "Use Context version of this method instead.")
+		void DrawMaxDistances(FPrimitiveDrawInterface* PDI) const
+		{
+			return DrawMaxDistances(FDrawContext(PDI));
+		}
+		UE_DEPRECATED(5.8, "Use Context version of this method instead.")
+		void DrawAnimDrive(FPrimitiveDrawInterface* PDI) const
+		{
+			return DrawAnimDrive(FDrawContext(PDI));
+		}
+		UE_DEPRECATED(5.8, "Use Context version of this method instead.")
+		void DrawEdgeConstraint(FPrimitiveDrawInterface* PDI) const
+		{
+			return DrawEdgeConstraint(FDrawContext(PDI));
+		}
+		UE_DEPRECATED(5.8, "Use Context version of this method instead.")
+		void DrawBendingConstraint(FPrimitiveDrawInterface* PDI) const
+		{
+			return DrawBendingConstraint(FDrawContext(PDI));
+		}
+		UE_DEPRECATED(5.8, "Use Context version of this method instead.")
+		void DrawLongRangeConstraint(FPrimitiveDrawInterface* PDI) const
+		{
+			return DrawLongRangeConstraint(FDrawContext(PDI));
+		}
+		UE_DEPRECATED(5.8, "Use Context version of this method instead.")
+		void DrawWindAndPressureForces(FPrimitiveDrawInterface* PDI, const FReal LengthScale) const
+		{
+			return DrawWindAndPressureForces(FDrawContext(PDI), LengthScale);
+		}
+		UE_DEPRECATED(5.8, "Use Context version of this method instead.")
+		void DrawWindVelocity(FPrimitiveDrawInterface* PDI, const FReal LengthScale) const
+		{
+			return DrawWindVelocity(FDrawContext(PDI), LengthScale);
+		}
+		UE_DEPRECATED(5.8, "Use Context version of this method instead.")
+		void DrawLocalSpace(FPrimitiveDrawInterface* PDI, const FReal LengthScale) const
+		{
+			return DrawLocalSpace(FDrawContext(PDI), LengthScale);
+		}
+		UE_DEPRECATED(5.8, "Use Context version of this method instead.")
+		void DrawVelocityScale(FPrimitiveDrawInterface* PDI, const FReal LengthScale) const
+		{
+			return DrawVelocityScale(FDrawContext(PDI), LengthScale);
+		}
+		UE_DEPRECATED(5.8, "Use Context version of this method instead.")
+		void DrawSelfCollision(FPrimitiveDrawInterface* PDI) const
+		{
+			return DrawSelfCollision(FDrawContext(PDI));
+		}
+		UE_DEPRECATED(5.8, "Use Context version of this method instead.")
+		void DrawSelfIntersection(FPrimitiveDrawInterface* PDI) const
+		{
+			return DrawSelfIntersection(FDrawContext(PDI));
+		}
+		UE_DEPRECATED(5.8, "Use Context version of this method instead.")
+		void DrawSelfCollisionThickness(FPrimitiveDrawInterface* PDI) const
+		{
+			return DrawSelfCollisionThickness(FDrawContext(PDI));
+		}
+		UE_DEPRECATED(5.8, "Use Context version of this method instead.")
+		void DrawKinematicColliderWired(FPrimitiveDrawInterface* PDI) const
+		{
+			return DrawKinematicColliderWired(FDrawContext(PDI));
+		}
+		UE_DEPRECATED(5.8, "Use Context version of this method instead.")
+		void DrawBounds(FPrimitiveDrawInterface* PDI) const
+		{
+			return DrawBounds(FDrawContext(PDI));
+		}
+		UE_DEPRECATED(5.8, "Use Context version of this method instead.")
+		void DrawGravity(FPrimitiveDrawInterface* PDI) const
+		{
+			return DrawGravity(FDrawContext(PDI));
+		}
+		UE_DEPRECATED(5.8, "Use Context version of this method instead.")
+		void DrawFictitiousAngularForces(FPrimitiveDrawInterface* PDI) const
+		{
+			return DrawFictitiousAngularForces(FDrawContext(PDI));
+		}
+		UE_DEPRECATED(5.8, "Use Context version of this method instead.")
+		void DrawMultiResConstraint(FPrimitiveDrawInterface* PDI) const
+		{
+			return DrawMultiResConstraint(FDrawContext(PDI));
+		}
+		UE_DEPRECATED(5.8, "Use Context version of this method instead.")
+		void DrawClothClothConstraints(FPrimitiveDrawInterface* PDI) const
+		{
+			return DrawClothClothConstraints(FDrawContext(PDI));
+		}
+		UE_DEPRECATED(5.8, "Use Context version of this method instead.")
+		void DrawTeleportReset(FPrimitiveDrawInterface* PDI) const
+		{
+			return DrawTeleportReset(FDrawContext(PDI));
+		}
+		UE_DEPRECATED(5.8, "Use Context version of this method instead.")
+		void DrawExtremlyDeformedEdges(FPrimitiveDrawInterface* PDI) const
+		{
+			return DrawExtremlyDeformedEdges(FDrawContext(PDI));
+		}
+		UE_DEPRECATED(5.8, "Use Context version of this method instead.")
+		void DrawAccessoryMesh(FPrimitiveDrawInterface* PDI, const FName& AccessoryMeshName = NAME_None) const
+		{
+			return DrawAccessoryMesh(FDrawContext(PDI), AccessoryMeshName);
+		}
+		UE_DEPRECATED(5.8, "Use Context version of this method instead.")
+		void DrawAccessoryMeshNormals(FPrimitiveDrawInterface* PDI, const FName& AccessoryMeshName = NAME_None, const FReal Length = 20.) const
+		{
+			return DrawAccessoryMeshNormals(FDrawContext(PDI), AccessoryMeshName, Length);
+		}
 	
 #else  // #if CHAOS_DEBUG_DRAW
 		void SetSolver(const ::Chaos::FClothingSimulationSolver* /*InSolver*/) {}
+		void DrawParticleIndices(const FDrawTextsContext& /*Context*/ = FDrawTextsContext()) const {}
+		void DrawElementIndices(const FDrawTextsContext& /*Context*/ = FDrawTextsContext()) const {}
+		void DrawMaxDistanceValues(const FDrawTextsContext& /*Context*/ = FDrawTextsContext()) const {}
+		void DrawLocalSpaceBoneNames(const FDrawTextsContext& /*Context*/ = FDrawTextsContext()) const {}
+		void DrawPhysMeshWired(const FDrawContext& /*Context*/ = FDrawContext()) const {}
+		void DrawAnimMeshWired(const FDrawContext& /*Context*/ = FDrawContext()) const {}
+		void DrawAnimNormals(const FDrawContext& /*Context*/ = FDrawContext(), const FReal /*Length*/ = 20.) const {}
+		void DrawOpenEdges(const FDrawContext& /*Context*/ = FDrawContext()) const {}
+		void DrawPointNormals(const FDrawContext& /*Context*/ = FDrawContext(), const FReal /*Length*/ = 20.) const {}
+		void DrawPointVelocities(const FDrawContext& /*Context*/ = FDrawContext()) const {}
+		void DrawCollision(const FDrawContext& /*Context*/ = FDrawContext(), bool /*bWireframe*/ = false) const {}
+		void DrawBackstops(const FDrawContext& /*Context*/ = FDrawContext()) const {}
+		void DrawBackstopDistances(const FDrawContext& /*Context*/ = FDrawContext()) const {}
+		void DrawMaxDistances(const FDrawContext& /*Context*/ = FDrawContext()) const {}
+		void DrawAnimDrive(const FDrawContext& /*Context*/ = FDrawContext()) const {}
+		void DrawEdgeConstraint(const FDrawContext& /*Context*/ = FDrawContext()) const {}
+		void DrawBendingConstraint(const FDrawContext& /*Context*/ = FDrawContext()) const {}
+		void DrawLongRangeConstraint(const FDrawContext& /*Context*/ = FDrawContext()) const {}
+		void DrawWindAndPressureForces(const FDrawContext& /*Context*/ = FDrawContext(), const FReal /*ForceLengthScale*/ = 10.) const {}
+		void DrawWindVelocity(const FDrawContext& /*Context*/ = FDrawContext(), const FReal /*LengthScale*/ = 0.1) const {}
+		void DrawLocalSpace(const FDrawContext& /*Context*/ = FDrawContext(), FReal /*LocalSpaceLengthScale*/ = 0.1) const {}
+		void DrawVelocityScale(const FDrawContext& /*Context*/ = FDrawContext(), const FReal /*LengthScale*/ = 1) const {}
+		void DrawSelfCollision(const FDrawContext& /*Context*/ = FDrawContext()) const {}
+		void DrawSelfIntersection(const FDrawContext& /*Context*/ = FDrawContext()) const {}
+		void DrawSelfCollisionThickness(const FDrawContext& /*Context*/ = FDrawContext()) const {}
+		void DrawKinematicColliderWired(const FDrawContext& /*Context*/ = FDrawContext()) const {}
+		void DrawBounds(const FDrawContext& /*Context*/ = FDrawContext()) const {}
+		void DrawGravity(const FDrawContext& /*Context*/ = FDrawContext()) const {}
+		void DrawFictitiousAngularForces(const FDrawContext& /*Context*/ = FDrawContext()) const {}
+		void DrawMultiResConstraint(const FDrawContext& /*Context*/ = FDrawContext()) const {}
+		void DrawClothClothConstraints(const FDrawContext& /*Context*/ = FDrawContext()) const {}
+		void DrawExtremlyDeformedEdges(const FDrawContext& /*Context*/ = FDrawContext()) const {}
+		void DrawAccessoryMesh(const FDrawContext& /*Context*/ = FDrawContext(), const FName& /*AccessoryMeshName*/ = NAME_None) const {}
+		void DrawAccessoryMeshNormals(const FDrawContext& /*Context*/ = FDrawContext(), const FName& /*AccessoryMeshName*/ = NAME_None, const FReal /*Length*/ = 20.) const;
 
+		UE_DEPRECATED(5.8, "Use Context version of this method instead.")
 		void DrawParticleIndices(FCanvas* /*Canvas*/, const FSceneView* /*SceneView*/) const {}
+		UE_DEPRECATED(5.8, "Use Context version of this method instead.")
 		void DrawElementIndices(FCanvas* /*Canvas*/, const FSceneView* /*SceneView*/) const {}
+		UE_DEPRECATED(5.8, "Use Context version of this method instead.")
 		void DrawMaxDistanceValues(FCanvas* /*Canvas*/, const FSceneView* /*SceneView*/) const {}
+		UE_DEPRECATED(5.8, "Use Context version of this method instead.")
 		void DrawLocalSpaceBoneNames(FCanvas* /*Canvas*/, const FSceneView* /*SceneView*/) const {}
-		void DrawPhysMeshWired(FPrimitiveDrawInterface* /*PDI*/ = nullptr) const {}
-		void DrawAnimMeshWired(FPrimitiveDrawInterface* /*PDI*/ = nullptr) const {}
+		UE_DEPRECATED(5.8, "Use Context version of this method instead.")
+		void DrawPhysMeshWired(FPrimitiveDrawInterface* /*PDI*/) const {}
+		UE_DEPRECATED(5.8, "Use Context version of this method instead.")
+		void DrawAnimMeshWired(FPrimitiveDrawInterface* /*PDI*/) const {}
+		UE_DEPRECATED(5.8, "Use Context version of this method instead.")
 		void DrawAnimNormals(FPrimitiveDrawInterface* /*PDI*/, const FReal /*Length*/) const {}
-		void DrawOpenEdges(FPrimitiveDrawInterface* /*PDI*/ = nullptr) const {}
+		UE_DEPRECATED(5.8, "Use Context version of this method instead.")
+		void DrawOpenEdges(FPrimitiveDrawInterface* /*PDI*/) const {}
+		UE_DEPRECATED(5.8, "Use Context version of this method instead.")
 		void DrawPointNormals(FPrimitiveDrawInterface* /*PDI*/, const FReal /*Length*/) const {}
-		void DrawPointVelocities(FPrimitiveDrawInterface* /*PDI*/ = nullptr) const {}
+		UE_DEPRECATED(5.8, "Use Context version of this method instead.")
+		void DrawPointVelocities(FPrimitiveDrawInterface* /*PDI*/) const {}
+		UE_DEPRECATED(5.8, "Use Context version of this method instead.")
 		void DrawCollision(FPrimitiveDrawInterface* /*PDI*/, bool /*bWireframe*/) const {}
-		void DrawBackstops(FPrimitiveDrawInterface* /*PDI*/ = nullptr) const {}
-		void DrawBackstopDistances(FPrimitiveDrawInterface* /*PDI*/ = nullptr) const {}
-		void DrawMaxDistances(FPrimitiveDrawInterface* /*PDI*/ = nullptr) const {}
-		void DrawAnimDrive(FPrimitiveDrawInterface* /*PDI*/ = nullptr) const {}
-		void DrawEdgeConstraint(FPrimitiveDrawInterface* /*PDI*/ = nullptr) const {}
-		void DrawBendingConstraint(FPrimitiveDrawInterface* /*PDI*/ = nullptr) const {}
-		void DrawLongRangeConstraint(FPrimitiveDrawInterface* /*PDI*/ = nullptr) const {}
+		UE_DEPRECATED(5.8, "Use Context version of this method instead.")
+		void DrawBackstops(FPrimitiveDrawInterface* /*PDI*/) const {}
+		UE_DEPRECATED(5.8, "Use Context version of this method instead.")
+		void DrawBackstopDistances(FPrimitiveDrawInterface* /*PDI*/) const {}
+		UE_DEPRECATED(5.8, "Use Context version of this method instead.")
+		void DrawMaxDistances(FPrimitiveDrawInterface* /*PDI*/) const {}
+		UE_DEPRECATED(5.8, "Use Context version of this method instead.")
+		void DrawAnimDrive(FPrimitiveDrawInterface* /*PDI*/) const {}
+		UE_DEPRECATED(5.8, "Use Context version of this method instead.")
+		void DrawEdgeConstraint(FPrimitiveDrawInterface* /*PDI*/) const {}
+		UE_DEPRECATED(5.8, "Use Context version of this method instead.")
+		void DrawBendingConstraint(FPrimitiveDrawInterface* /*PDI*/) const {}
+		UE_DEPRECATED(5.8, "Use Context version of this method instead.")
+		void DrawLongRangeConstraint(FPrimitiveDrawInterface* /*PDI*/) const {}
+		UE_DEPRECATED(5.8, "Use Context version of this method instead.")
 		void DrawWindAndPressureForces(FPrimitiveDrawInterface* /*PDI*/, const FReal /*ForceLengthScale*/) const {}
+		UE_DEPRECATED(5.8, "Use Context version of this method instead.")
 		void DrawWindVelocity(FPrimitiveDrawInterface* /*PDI*/, const FReal /*LengthScale*/) const {}
+		UE_DEPRECATED(5.8, "Use Context version of this method instead.")
 		void DrawLocalSpace(FPrimitiveDrawInterface* /*PDI*/, FReal /*LocalSpaceLengthScale*/) const {}
+		UE_DEPRECATED(5.8, "Use Context version of this method instead.")
 		void DrawVelocityScale(FPrimitiveDrawInterface* /*PDI*/, const FReal /*LengthScale*/) const {}
-		void DrawSelfCollision(FPrimitiveDrawInterface* /*PDI*/ = nullptr) const {}
-		void DrawSelfIntersection(FPrimitiveDrawInterface* /*PDI*/ = nullptr) const {}
-		void DrawSelfCollisionThickness(FPrimitiveDrawInterface* /*PDI*/ = nullptr) const {}
-		void DrawKinematicColliderWired(FPrimitiveDrawInterface* /*PDI*/ = nullptr) const {}
-		void DrawBounds(FPrimitiveDrawInterface* /*PDI*/ = nullptr) const {}
-		void DrawGravity(FPrimitiveDrawInterface* /*PDI*/ = nullptr) const {}
-		void DrawFictitiousAngularForces(FPrimitiveDrawInterface* /*PDI*/ = nullptr) const {}
-		void DrawMultiResConstraint(FPrimitiveDrawInterface* /*PDI*/ = nullptr) const {}
-		void DrawClothClothConstraints(FPrimitiveDrawInterface* /*PDI*/ = nullptr) const {}
-		void DrawExtremlyDeformedEdges(FPrimitiveDrawInterface* /*PDI*/ = nullptr) const {}
-		void DrawAccessoryMesh(FPrimitiveDrawInterface* /*PDI*/ = nullptr, const FName& AccessoryMeshName = NAME_None) const {}
-		void DrawAccessoryMeshNormals(FPrimitiveDrawInterface* /*PDI*/ = nullptr, const FName& AccessoryMeshName = NAME_None, const FReal Length = 20.) const;
+		UE_DEPRECATED(5.8, "Use Context version of this method instead.")
+		void DrawSelfCollision(FPrimitiveDrawInterface* /*PDI*/) const {}
+		UE_DEPRECATED(5.8, "Use Context version of this method instead.")
+		void DrawSelfIntersection(FPrimitiveDrawInterface* /*PDI*/) const {}
+		UE_DEPRECATED(5.8, "Use Context version of this method instead.")
+		void DrawSelfCollisionThickness(FPrimitiveDrawInterface* /*PDI*/) const {}
+		UE_DEPRECATED(5.8, "Use Context version of this method instead.")
+		void DrawKinematicColliderWired(FPrimitiveDrawInterface* /*PDI*/) const {}
+		UE_DEPRECATED(5.8, "Use Context version of this method instead.")
+		void DrawBounds(FPrimitiveDrawInterface* /*PDI*/) const {}
+		UE_DEPRECATED(5.8, "Use Context version of this method instead.")
+		void DrawGravity(FPrimitiveDrawInterface* /*PDI*/) const {}
+		UE_DEPRECATED(5.8, "Use Context version of this method instead.")
+		void DrawFictitiousAngularForces(FPrimitiveDrawInterface* /*PDI*/) const {}
+		UE_DEPRECATED(5.8, "Use Context version of this method instead.")
+		void DrawMultiResConstraint(FPrimitiveDrawInterface* /*PDI*/) const {}
+		UE_DEPRECATED(5.8, "Use Context version of this method instead.")
+		void DrawClothClothConstraints(FPrimitiveDrawInterface* /*PDI*/) const {}
+		UE_DEPRECATED(5.8, "Use Context version of this method instead.")
+		void DrawExtremlyDeformedEdges(FPrimitiveDrawInterface* /*PDI*/) const {}
+		UE_DEPRECATED(5.8, "Use Context version of this method instead.")
+		void DrawAccessoryMesh(FPrimitiveDrawInterface* /*PDI*/, const FName& /*AccessoryMeshName*/ = NAME_None) const {}
+		UE_DEPRECATED(5.8, "Use Context version of this method instead.")
+		void DrawAccessoryMeshNormals(FPrimitiveDrawInterface* /*PDI*/, const FName& /*AccessoryMeshName*/ = NAME_None, const FReal /*Length*/ = 20.) const;
 #endif  // #if CHAOS_DEBUG_DRAW
 
-		void DrawAnimNormals(FPrimitiveDrawInterface* PDI = nullptr) const
+		UE_DEPRECATED(5.8, "Use Context version of this method instead.")
+		void DrawAnimNormals(FPrimitiveDrawInterface* PDI) const
 		{
-			constexpr FReal DefaultLengtth = 20.;
-			DrawAnimNormals(PDI, DefaultLengtth);
+			constexpr FReal DefaultLength = 20.;
+			return DrawAnimNormals(FDrawContext(PDI), DefaultLength);
 		}
-		void DrawPointNormals(FPrimitiveDrawInterface* PDI = nullptr) const
+		UE_DEPRECATED(5.8, "Use Context version of this method instead.")
+		void DrawPointNormals(FPrimitiveDrawInterface* PDI) const
 		{
-			constexpr FReal DefaultLengtth = 20.;
-			DrawPointNormals(PDI, DefaultLengtth);
+			constexpr FReal DefaultLength = 20.;
+			return DrawAnimNormals(FDrawContext(PDI), DefaultLength);
 		}
-		void DrawCollision(FPrimitiveDrawInterface* PDI = nullptr) const
+		UE_DEPRECATED(5.8, "Use Context version of this method instead.")
+		void DrawCollision(FPrimitiveDrawInterface* PDI) const
 		{
 			constexpr bool bDefaultWireframe = false;
-			DrawCollision(PDI, bDefaultWireframe);
+			return DrawCollision(FDrawContext(PDI), bDefaultWireframe);
 		}
-		void DrawWindAndPressureForces(FPrimitiveDrawInterface* PDI = nullptr) const
+		UE_DEPRECATED(5.8, "Use Context version of this method instead.")
+		void DrawWindAndPressureForces(FPrimitiveDrawInterface* PDI) const
 		{
 			constexpr FReal DefaultForceLengthScale = 10.;
-			DrawWindAndPressureForces(PDI, DefaultForceLengthScale);
+			return DrawWindAndPressureForces(FDrawContext(PDI), DefaultForceLengthScale);
 		}
-		void DrawWindVelocity(FPrimitiveDrawInterface* PDI = nullptr) const
+		UE_DEPRECATED(5.8, "Use Context version of this method instead.")
+		void DrawWindVelocity(FPrimitiveDrawInterface* PDI) const
 		{
 			constexpr FReal DefaultLengthScale = .1;
-			DrawWindVelocity(PDI, DefaultLengthScale);
+			return DrawWindVelocity(FDrawContext(PDI), DefaultLengthScale);
 		}
-		void DrawVelocityScale(FPrimitiveDrawInterface* PDI = nullptr) const
+		UE_DEPRECATED(5.8, "Use Context version of this method instead.")
+		void DrawVelocityScale(FPrimitiveDrawInterface* PDI) const
 		{
 			constexpr FReal DefaultLengthScale = .1;
-			DrawVelocityScale(PDI, DefaultLengthScale);
+			return DrawVelocityScale(FDrawContext(PDI), DefaultLengthScale);
 		}
-		void DrawLocalSpace(FPrimitiveDrawInterface* PDI = nullptr) const
+		UE_DEPRECATED(5.8, "Use Context version of this method instead.")
+		void DrawLocalSpace(FPrimitiveDrawInterface* PDI) const
 		{
 			constexpr FReal DefaultLengthScale = 1.;
-			DrawLocalSpace(PDI, DefaultLengthScale);
+			return DrawLocalSpace(FDrawContext(PDI), DefaultLengthScale);
 		}
 
 #if WITH_EDITOR && CHAOS_DEBUG_DRAW
 		// Editor only functions
-		CHAOSCLOTH_API void DrawPhysMeshShaded(FPrimitiveDrawInterface* PDI) const;
-		CHAOSCLOTH_API void DrawWeightMap(FPrimitiveDrawInterface* PDI = nullptr) const;
-		CHAOSCLOTH_API void DrawSelfCollisionLayers(FPrimitiveDrawInterface* PDI = nullptr) const;
-		CHAOSCLOTH_API void DrawInpaintWeightsMatched(FPrimitiveDrawInterface* PDI = nullptr) const;
-		CHAOSCLOTH_API void DrawKinematicColliderShaded(FPrimitiveDrawInterface* PDI) const;
-		CHAOSCLOTH_API void DrawWeightMapWithName(FPrimitiveDrawInterface* PDI, const FString& Name) const;
+		CHAOSCLOTH_API void DrawPhysMeshShaded(const FDrawContext& Context = FDrawContext()) const;
+		CHAOSCLOTH_API void DrawSelfCollisionLayers(const FDrawContext& Context = FDrawContext()) const;
+		CHAOSCLOTH_API void DrawInpaintWeightsMatched(const FDrawContext& Context = FDrawContext()) const;
+		CHAOSCLOTH_API void DrawKinematicColliderShaded(const FDrawContext& Context = FDrawContext()) const;
+		/** When Name is empty, the string stored by the cvar p.ChaosClothVisualization.WeightMapName is used instead */
+		CHAOSCLOTH_API void DrawWeightMapWithName(const FDrawContext& Context = FDrawContext(), const FString& Name = "") const;
 		/** Will draw current morph target (if any) if no name or an invalid name is specified */
-		CHAOSCLOTH_API void DrawSimMorphTarget(FPrimitiveDrawInterface* PDI = nullptr, const FString& Name = "") const;
+		CHAOSCLOTH_API void DrawSimMorphTarget(const FDrawContext& Context = FDrawContext(), const FString& Name = "") const;
 		CHAOSCLOTH_API TArray<FString> GetAllWeightMapNames() const;
 		CHAOSCLOTH_API TArray<FString> GetAllMorphTargetNames() const;
+		CHAOSCLOTH_API TArray<FString> GetAllVertexSetNames() const;
 		CHAOSCLOTH_API TArray<FName> GetAllAccessoryMeshNames() const;
+
+		// Editor only functions
+		UE_DEPRECATED(5.8, "Use Context version of this method instead.")
+		void DrawPhysMeshShaded(FPrimitiveDrawInterface* PDI) const
+		{
+			FDrawContext Context;
+			return DrawPhysMeshShaded(FDrawContext(PDI));
+		}
+		UE_DEPRECATED(5.8, "Use DrawWeightMapWithName instead.")
+		void DrawWeightMap(FPrimitiveDrawInterface* PDI) const
+		{
+			return DrawWeightMapWithName(FDrawContext(PDI), FString());
+		}
+		UE_DEPRECATED(5.8, "Use Context version of this method instead.")
+		void DrawSelfCollisionLayers(FPrimitiveDrawInterface* PDI) const
+		{
+			return DrawSelfCollisionLayers(FDrawContext(PDI));
+		}
+		UE_DEPRECATED(5.8, "Use Context version of this method instead.")
+		void DrawInpaintWeightsMatched(FPrimitiveDrawInterface* PDI) const
+		{
+			return DrawInpaintWeightsMatched(FDrawContext(PDI));
+		}
+		UE_DEPRECATED(5.8, "Use Context version of this method instead.")
+		void DrawKinematicColliderShaded(FPrimitiveDrawInterface* PDI) const
+		{
+			return DrawKinematicColliderShaded(FDrawContext(PDI));
+		}
+		UE_DEPRECATED(5.8, "Use Context version of this method instead.")
+		void DrawWeightMapWithName(FPrimitiveDrawInterface* PDI, const FString& Name) const
+		{
+			return DrawWeightMapWithName(FDrawContext(PDI), Name);
+		}
+		/** Will draw current morph target (if any) if no name or an invalid name is specified */
+		UE_DEPRECATED(5.8, "Use Context version of this method instead.")
+		void DrawSimMorphTarget(FPrimitiveDrawInterface* PDI, const FString& Name = "") const
+		{
+			return DrawSimMorphTarget(FDrawContext(PDI), Name);
+		}
 #else  // #if WITH_EDITOR && CHAOS_DEBUG_DRAW
+		void DrawPhysMeshShaded(const FDrawContext& Context = FDrawContext()) const {}
+		void DrawKinematicColliderShaded(const FDrawContext& Context = FDrawContext()) const {}
+
+		UE_DEPRECATED(5.8, "Use Context version of this method instead.")
 		void DrawPhysMeshShaded(FPrimitiveDrawInterface* /*PDI*/) const {}
+		UE_DEPRECATED(5.8, "Use Context version of this method instead.")
 		void DrawKinematicColliderShaded(FPrimitiveDrawInterface* /*PDI*/) const {}
 #endif  // #if WITH_EDITOR && CHAOS_DEBUG_DRAW
 
diff --git a/Engine/Plugins/ChaosCloth/Source/ChaosCloth/Public/ChaosCloth/ChaosClothingSimulation.h b/Engine/Plugins/ChaosCloth/Source/ChaosCloth/Public/ChaosCloth/ChaosClothingSimulation.h
--- a/Engine/Plugins/ChaosCloth/Source/ChaosCloth/Public/ChaosCloth/ChaosClothingSimulation.h
+++ b/Engine/Plugins/ChaosCloth/Source/ChaosCloth/Public/ChaosCloth/ChaosClothingSimulation.h
@@ -92,39 +92,68 @@ PRAGMA_ENABLE_DEPRECATION_WARNINGS
 
 #if WITH_EDITOR
 		// Editor only debug draw function
-		void DebugDrawPhysMeshShaded(FPrimitiveDrawInterface* PDI) const { Visualization.DrawPhysMeshShaded(PDI); }
+		UE_DEPRECATED(5.8, "Call debug draw methods directly from the GetClothVisualization() instead.")
+		void DebugDrawPhysMeshShaded(FPrimitiveDrawInterface* PDI) const { Visualization.DrawPhysMeshShaded(FClothVisualizationNoGC::FDrawContext(PDI)); }
 #endif  // #if WITH_EDITOR
 
 #if CHAOS_DEBUG_DRAW
 		// Editor & runtime debug draw functions
-		void DebugDrawParticleIndices(FCanvas* Canvas = nullptr, const FSceneView* SceneView = nullptr) const { Visualization.DrawParticleIndices(Canvas, SceneView); }
-		void DebugDrawElementIndices(FCanvas* Canvas = nullptr, const FSceneView* SceneView = nullptr) const { Visualization.DrawElementIndices(Canvas, SceneView); }
-		void DebugDrawMaxDistanceValues(FCanvas* Canvas = nullptr, const FSceneView* SceneView = nullptr) const { Visualization.DrawMaxDistanceValues(Canvas, SceneView); }
-		void DebugDrawLocalSpaceBoneNames(FCanvas* Canvas = nullptr, const FSceneView* SceneView = nullptr) const { Visualization.DrawLocalSpaceBoneNames(Canvas, SceneView); }
-		void DebugDrawPhysMeshWired(FPrimitiveDrawInterface* PDI = nullptr) const { Visualization.DrawPhysMeshWired(PDI); }
-		void DebugDrawAnimMeshWired(FPrimitiveDrawInterface* PDI = nullptr) const { Visualization.DrawAnimMeshWired(PDI); }
-		void DebugDrawAnimNormals(FPrimitiveDrawInterface* PDI = nullptr) const { Visualization.DrawAnimNormals(PDI); }
-		void DebugDrawAnimVelocities(FPrimitiveDrawInterface* PDI = nullptr) const { Visualization.DrawAnimVelocities(PDI); }
-		void DebugDrawPointNormals(FPrimitiveDrawInterface* PDI = nullptr) const { Visualization.DrawPointNormals(PDI); }
-		void DebugDrawPointVelocities(FPrimitiveDrawInterface* PDI = nullptr) const { Visualization.DrawPointVelocities(PDI); }
-		void DebugDrawCollision(FPrimitiveDrawInterface* PDI = nullptr) const { Visualization.DrawCollision(PDI); }
-		void DebugDrawBackstops(FPrimitiveDrawInterface* PDI = nullptr) const { Visualization.DrawBackstops(PDI); }
-		void DebugDrawBackstopDistances(FPrimitiveDrawInterface* PDI = nullptr) const { Visualization.DrawBackstopDistances(PDI); }
-		void DebugDrawMaxDistances(FPrimitiveDrawInterface* PDI = nullptr) const { Visualization.DrawMaxDistances(PDI); }
-		void DebugDrawAnimDrive(FPrimitiveDrawInterface* PDI = nullptr) const { Visualization.DrawAnimDrive(PDI); }
-		void DebugDrawEdgeConstraint(FPrimitiveDrawInterface* PDI = nullptr) const { Visualization.DrawEdgeConstraint(PDI); }
-		void DebugDrawBendingConstraint(FPrimitiveDrawInterface* PDI = nullptr) const { Visualization.DrawBendingConstraint(PDI); }
-		void DebugDrawLongRangeConstraint(FPrimitiveDrawInterface* PDI = nullptr) const { Visualization.DrawLongRangeConstraint(PDI); }
-		void DebugDrawWindAndPressureForces(FPrimitiveDrawInterface* PDI = nullptr) const { Visualization.DrawWindAndPressureForces(PDI); }
-		void DebugDrawWindVelocity(FPrimitiveDrawInterface* PDI = nullptr) const { Visualization.DrawWindVelocity(PDI); }
-		void DebugDrawLocalSpace(FPrimitiveDrawInterface* PDI = nullptr) const { Visualization.DrawLocalSpace(PDI); }
-		void DebugDrawSelfCollision(FPrimitiveDrawInterface* PDI = nullptr) const { Visualization.DrawSelfCollision(PDI); }
-		void DebugDrawSelfIntersection(FPrimitiveDrawInterface* PDI = nullptr) const { Visualization.DrawSelfIntersection(PDI); }
-		void DebugDrawBounds(FPrimitiveDrawInterface* PDI = nullptr) const { Visualization.DrawBounds(PDI); }
-		void DebugDrawGravity(FPrimitiveDrawInterface* PDI = nullptr) const { Visualization.DrawGravity(PDI); }
-		void DebugDrawTeleportReset(FPrimitiveDrawInterface* PDI = nullptr) const { Visualization.DrawTeleportReset(PDI); }
-		void DebugDrawExtremlyDeformedEdges(FPrimitiveDrawInterface* PDI = nullptr) const { Visualization.DrawExtremlyDeformedEdges(PDI); }
-		void DebugDrawVelocityScale(FPrimitiveDrawInterface* PDI = nullptr) const { Visualization.DrawVelocityScale(PDI); }
+		UE_DEPRECATED(5.8, "Call debug draw methods directly from the GetClothVisualization() instead.")
+		void DebugDrawParticleIndices(FCanvas* Canvas = nullptr, const FSceneView* SceneView = nullptr) const { Visualization.DrawParticleIndices(FClothVisualizationNoGC::FDrawTextsContext(Canvas, SceneView)); }
+		UE_DEPRECATED(5.8, "Call debug draw methods directly from the GetClothVisualization() instead.")
+		void DebugDrawElementIndices(FCanvas* Canvas = nullptr, const FSceneView* SceneView = nullptr) const { Visualization.DrawElementIndices(FClothVisualizationNoGC::FDrawTextsContext(Canvas, SceneView)); }
+		UE_DEPRECATED(5.8, "Call debug draw methods directly from the GetClothVisualization() instead.")
+		void DebugDrawMaxDistanceValues(FCanvas* Canvas = nullptr, const FSceneView* SceneView = nullptr) const { Visualization.DrawMaxDistanceValues(FClothVisualizationNoGC::FDrawTextsContext(Canvas, SceneView)); }
+		UE_DEPRECATED(5.8, "Call debug draw methods directly from the GetClothVisualization() instead.")
+		void DebugDrawLocalSpaceBoneNames(FCanvas* Canvas = nullptr, const FSceneView* SceneView = nullptr) const { Visualization.DrawLocalSpaceBoneNames(FClothVisualizationNoGC::FDrawTextsContext(Canvas, SceneView)); }
+		UE_DEPRECATED(5.8, "Call debug draw methods directly from the GetClothVisualization() instead.")
+		void DebugDrawPhysMeshWired(FPrimitiveDrawInterface* PDI = nullptr) const { Visualization.DrawPhysMeshWired(FClothVisualizationNoGC::FDrawContext(PDI)); }
+		UE_DEPRECATED(5.8, "Call debug draw methods directly from the GetClothVisualization() instead.")
+		void DebugDrawAnimMeshWired(FPrimitiveDrawInterface* PDI = nullptr) const { Visualization.DrawAnimMeshWired(FClothVisualizationNoGC::FDrawContext(PDI)); }
+		UE_DEPRECATED(5.8, "Call debug draw methods directly from the GetClothVisualization() instead.")
+		void DebugDrawAnimNormals(FPrimitiveDrawInterface* PDI = nullptr) const { Visualization.DrawAnimNormals(FClothVisualizationNoGC::FDrawContext(PDI)); }
+		UE_DEPRECATED(5.8, "Call debug draw methods directly from the GetClothVisualization() instead.")
+		void DebugDrawAnimVelocities(FPrimitiveDrawInterface* PDI = nullptr) const { Visualization.DrawAnimVelocities(FClothVisualizationNoGC::FDrawContext(PDI)); }
+		UE_DEPRECATED(5.8, "Call debug draw methods directly from the GetClothVisualization() instead.")
+		void DebugDrawPointNormals(FPrimitiveDrawInterface* PDI = nullptr) const { Visualization.DrawPointNormals(FClothVisualizationNoGC::FDrawContext(PDI)); }
+		UE_DEPRECATED(5.8, "Call debug draw methods directly from the GetClothVisualization() instead.")
+		void DebugDrawPointVelocities(FPrimitiveDrawInterface* PDI = nullptr) const { Visualization.DrawPointVelocities(FClothVisualizationNoGC::FDrawContext(PDI)); }
+		UE_DEPRECATED(5.8, "Call debug draw methods directly from the GetClothVisualization() instead.")
+		void DebugDrawCollision(FPrimitiveDrawInterface* PDI = nullptr) const { Visualization.DrawCollision(FClothVisualizationNoGC::FDrawContext(PDI)); }
+		UE_DEPRECATED(5.8, "Call debug draw methods directly from the GetClothVisualization() instead.")
+		void DebugDrawBackstops(FPrimitiveDrawInterface* PDI = nullptr) const { Visualization.DrawBackstops(FClothVisualizationNoGC::FDrawContext(PDI)); }
+		UE_DEPRECATED(5.8, "Call debug draw methods directly from the GetClothVisualization() instead.")
+		void DebugDrawBackstopDistances(FPrimitiveDrawInterface* PDI = nullptr) const { Visualization.DrawBackstopDistances(FClothVisualizationNoGC::FDrawContext(PDI)); }
+		UE_DEPRECATED(5.8, "Call debug draw methods directly from the GetClothVisualization() instead.")
+		void DebugDrawMaxDistances(FPrimitiveDrawInterface* PDI = nullptr) const { Visualization.DrawMaxDistances(FClothVisualizationNoGC::FDrawContext(PDI)); }
+		UE_DEPRECATED(5.8, "Call debug draw methods directly from the GetClothVisualization() instead.")
+		void DebugDrawAnimDrive(FPrimitiveDrawInterface* PDI = nullptr) const { Visualization.DrawAnimDrive(FClothVisualizationNoGC::FDrawContext(PDI)); }
+		UE_DEPRECATED(5.8, "Call debug draw methods directly from the GetClothVisualization() instead.")
+		void DebugDrawEdgeConstraint(FPrimitiveDrawInterface* PDI = nullptr) const { Visualization.DrawEdgeConstraint(FClothVisualizationNoGC::FDrawContext(PDI)); }
+		UE_DEPRECATED(5.8, "Call debug draw methods directly from the GetClothVisualization() instead.")
+		void DebugDrawBendingConstraint(FPrimitiveDrawInterface* PDI = nullptr) const { Visualization.DrawBendingConstraint(FClothVisualizationNoGC::FDrawContext(PDI)); }
+		UE_DEPRECATED(5.8, "Call debug draw methods directly from the GetClothVisualization() instead.")
+		void DebugDrawLongRangeConstraint(FPrimitiveDrawInterface* PDI = nullptr) const { Visualization.DrawLongRangeConstraint(FClothVisualizationNoGC::FDrawContext(PDI)); }
+		UE_DEPRECATED(5.8, "Call debug draw methods directly from the GetClothVisualization() instead.")
+		void DebugDrawWindAndPressureForces(FPrimitiveDrawInterface* PDI = nullptr) const { Visualization.DrawWindAndPressureForces(FClothVisualizationNoGC::FDrawContext(PDI)); }
+		UE_DEPRECATED(5.8, "Call debug draw methods directly from the GetClothVisualization() instead.")
+		void DebugDrawWindVelocity(FPrimitiveDrawInterface* PDI = nullptr) const { Visualization.DrawWindVelocity(FClothVisualizationNoGC::FDrawContext(PDI)); }
+		UE_DEPRECATED(5.8, "Call debug draw methods directly from the GetClothVisualization() instead.")
+		void DebugDrawLocalSpace(FPrimitiveDrawInterface* PDI = nullptr) const { Visualization.DrawLocalSpace(FClothVisualizationNoGC::FDrawContext(PDI)); }
+		UE_DEPRECATED(5.8, "Call debug draw methods directly from the GetClothVisualization() instead.")
+		void DebugDrawSelfCollision(FPrimitiveDrawInterface* PDI = nullptr) const { Visualization.DrawSelfCollision(FClothVisualizationNoGC::FDrawContext(PDI)); }
+		UE_DEPRECATED(5.8, "Call debug draw methods directly from the GetClothVisualization() instead.")
+		void DebugDrawSelfIntersection(FPrimitiveDrawInterface* PDI = nullptr) const { Visualization.DrawSelfIntersection(FClothVisualizationNoGC::FDrawContext(PDI)); }
+		UE_DEPRECATED(5.8, "Call debug draw methods directly from the GetClothVisualization() instead.")
+		void DebugDrawBounds(FPrimitiveDrawInterface* PDI = nullptr) const { Visualization.DrawBounds(FClothVisualizationNoGC::FDrawContext(PDI)); }
+		UE_DEPRECATED(5.8, "Call debug draw methods directly from the GetClothVisualization() instead.")
+		void DebugDrawGravity(FPrimitiveDrawInterface* PDI = nullptr) const { Visualization.DrawGravity(FClothVisualizationNoGC::FDrawContext(PDI)); }
+		UE_DEPRECATED(5.8, "Call debug draw methods directly from the GetClothVisualization() instead.")
+		void DebugDrawTeleportReset(FPrimitiveDrawInterface* PDI = nullptr) const { Visualization.DrawTeleportReset(FClothVisualizationNoGC::FDrawContext(PDI)); }
+		UE_DEPRECATED(5.8, "Call debug draw methods directly from the GetClothVisualization() instead.")
+		void DebugDrawExtremlyDeformedEdges(FPrimitiveDrawInterface* PDI = nullptr) const { Visualization.DrawExtremlyDeformedEdges(FClothVisualizationNoGC::FDrawContext(PDI)); }
+		UE_DEPRECATED(5.8, "Call debug draw methods directly from the GetClothVisualization() instead.")
+		void DebugDrawVelocityScale(FPrimitiveDrawInterface* PDI = nullptr) const { Visualization.DrawVelocityScale(FClothVisualizationNoGC::FDrawContext(PDI)); }
 #endif  // #if CHAOS_DEBUG_DRAW
 
 		/** Return the visualization object for this simulation. */
diff --git a/Engine/Plugins/ChaosCloth/Source/ChaosCloth/Public/ChaosCloth/ChaosClothingSimulationCloth.h b/Engine/Plugins/ChaosCloth/Source/ChaosCloth/Public/ChaosCloth/ChaosClothingSimulationCloth.h
--- a/Engine/Plugins/ChaosCloth/Source/ChaosCloth/Public/ChaosCloth/ChaosClothingSimulationCloth.h
+++ b/Engine/Plugins/ChaosCloth/Source/ChaosCloth/Public/ChaosCloth/ChaosClothingSimulationCloth.h
@@ -135,6 +135,11 @@ namespace Chaos
 		{
 			return GetFaceIntMapByProperty(Solver, FName(Property));
 		}
+		CHAOSCLOTH_API TSet<FString> GetAllVertexSetNames() const;
+		// Return the vertex set of the specified name if available on the current LOD.
+		CHAOSCLOTH_API const TSet<int32>* GetVertexSetByName(const FClothingSimulationSolver* Solver, const FString& Name) const;
+		// Return the vertex set  of the specified property name if it exists and is available on the current LOD.
+		CHAOSCLOTH_API const TSet<int32>* GetVertexSetByProperty(const FClothingSimulationSolver* Solver, const FName& Property) const;
 		CHAOSCLOTH_API FString GetPropertyString(const FClothingSimulationSolver* Solver, const FName& Name) const;
 		// Return the current LOD tethers.
 		CHAOSCLOTH_API const TArray<TConstArrayView<TTuple<int32, int32, float>>>& GetTethers(const FClothingSimulationSolver* Solver) const;
diff --git a/Engine/Plugins/ChaosCloth/Source/ChaosClothEditor/Private/ChaosClothEditor/ChaosSimulationEditorExtender.cpp b/Engine/Plugins/ChaosCloth/Source/ChaosClothEditor/Private/ChaosClothEditor/ChaosSimulationEditorExtender.cpp
--- a/Engine/Plugins/ChaosCloth/Source/ChaosClothEditor/Private/ChaosClothEditor/ChaosSimulationEditorExtender.cpp
+++ b/Engine/Plugins/ChaosCloth/Source/ChaosClothEditor/Private/ChaosClothEditor/ChaosSimulationEditorExtender.cpp
@@ -15,6 +15,10 @@
 namespace Chaos
 {
 
+
+DECLARE_DELEGATE_TwoParams(FClothVisualizationDebugDraw, const ::Chaos::FClothVisualizationNoGC&, const ::Chaos::FClothVisualizationNoGC::FDrawContext&);
+DECLARE_DELEGATE_TwoParams(FClothVisualizationDebugDrawTexts, const ::Chaos::FClothVisualizationNoGC&, const ::Chaos::FClothVisualizationNoGC::FDrawTextsContext&);
+
 struct FVisualizationOption
 {
 	// Actual option entries
@@ -22,10 +26,8 @@ struct FVisualizationOption
 	static const uint32 Count;
 
 	// Chaos debug draw function
-	typedef void (FClothVisualizationNoGC::*FDebugDrawFunction)(FPrimitiveDrawInterface*) const;
-	typedef void (FClothVisualizationNoGC::*FDebugDrawTextsFunction)(FCanvas*, const FSceneView*) const;
-	FDebugDrawFunction DebugDrawFunction;
-	FDebugDrawTextsFunction DebugDrawTextsFunction;
+	FClothVisualizationDebugDraw DebugDrawFunction;
+	FClothVisualizationDebugDrawTexts DebugDrawTextsFunction;
 
 	FText DisplayName;         // Text for menu entries.
 	FText ToolTip;             // Text for menu tooltips.
@@ -35,7 +37,7 @@ struct FVisualizationOption
 	// Console override
 	IConsoleVariable* const ConsoleVariable;
 
-	FVisualizationOption(FDebugDrawFunction InDebugDrawFunction, const TCHAR* ConsoleName, const FText& InDisplayName, const FText& InToolTip, bool bInDisablesSimulation = false, bool bInHidesClothSections = false)
+	FVisualizationOption(const FClothVisualizationDebugDraw& InDebugDrawFunction, const TCHAR* ConsoleName, const FText& InDisplayName, const FText& InToolTip, bool bInDisablesSimulation = false, bool bInHidesClothSections = false)
 		: DebugDrawFunction(InDebugDrawFunction)
 		, DebugDrawTextsFunction(nullptr)
 		, DisplayName(InDisplayName)
@@ -45,7 +47,7 @@ struct FVisualizationOption
 		, ConsoleVariable(IConsoleManager::Get().RegisterConsoleVariable(ConsoleName, false, *InToolTip.ToString(), 0))
 	{}
 
-	FVisualizationOption(FDebugDrawTextsFunction InDebugDrawTextsFunction, const TCHAR* ConsoleName, const FText& InDisplayName, const FText& InToolTip, bool bInDisablesSimulation = false, bool bInHidesClothSections = false)
+	FVisualizationOption(const FClothVisualizationDebugDrawTexts& InDebugDrawTextsFunction, const TCHAR* ConsoleName, const FText& InDisplayName, const FText& InToolTip, bool bInDisablesSimulation = false, bool bInHidesClothSections = false)
 		: DebugDrawFunction(nullptr)
 		, DebugDrawTextsFunction(InDebugDrawTextsFunction)
 		, DisplayName(InDisplayName)
@@ -66,132 +68,214 @@ struct FVisualizationOption
 const FVisualizationOption FVisualizationOption::OptionData[] = 
 {
 	FVisualizationOption(
-		&FClothVisualizationNoGC::DrawPhysMeshShaded,
+		FClothVisualizationDebugDraw::CreateLambda([](const ::Chaos::FClothVisualizationNoGC& Visualization, const ::Chaos::FClothVisualizationNoGC::FDrawContext& Context)
+			{
+				Visualization.DrawPhysMeshShaded(Context);
+			}),
 		TEXT("p.ChaosClothEditor.DebugDrawPhysMeshShaded"),
 		LOCTEXT("ChaosVisName_PhysMesh", "Physical Mesh (Flat Shaded)"),
 		LOCTEXT("ChaosVisName_PhysMeshShaded_ToolTip", "Draws the current physical result as a doubled sided flat shaded mesh"),
 		/*bDisablesSimulation =*/false,
 		/*bHidesClothSections=*/true),
 	FVisualizationOption(
-		&FClothVisualizationNoGC::DrawPhysMeshWired,
+		FClothVisualizationDebugDraw::CreateLambda([](const ::Chaos::FClothVisualizationNoGC& Visualization, const ::Chaos::FClothVisualizationNoGC::FDrawContext& Context)
+			{
+				Visualization.DrawPhysMeshWired(Context);
+			}),
 		TEXT("p.ChaosClothEditor.DebugDrawPhysMeshWired"),
 		LOCTEXT("ChaosVisName_PhysMeshWire", "Physical Mesh (Wireframe)"),
 		LOCTEXT("ChaosVisName_PhysMeshWired_ToolTip", "Draws the current physical mesh result in wireframe")),
 	FVisualizationOption(
-		&FClothVisualizationNoGC::DrawAnimMeshWired,
+		FClothVisualizationDebugDraw::CreateLambda([](const ::Chaos::FClothVisualizationNoGC& Visualization, const ::Chaos::FClothVisualizationNoGC::FDrawContext& Context)
+			{
+				Visualization.DrawAnimMeshWired(Context);
+			}),
 		TEXT("p.ChaosClothEditor.DebugDrawAnimMeshWired"),
 		LOCTEXT("ChaosVisName_AnimMeshWire", "Animated Mesh (Wireframe)"),
 		LOCTEXT("ChaosVisName_AnimMeshWired_ToolTip", "Draws the current animated mesh input in wireframe")),
 	FVisualizationOption(
-		&FClothVisualizationNoGC::DrawParticleIndices,
+		FClothVisualizationDebugDrawTexts::CreateLambda([](const ::Chaos::FClothVisualizationNoGC& Visualization, const ::Chaos::FClothVisualizationNoGC::FDrawTextsContext& Context)
+		{
+			Visualization.DrawParticleIndices(Context);
+		}),
 		TEXT("p.ChaosClothEditor.DebugDrawParticleIndices"),
 		LOCTEXT("ChaosVisName_ParticleIndices", "Particle Indices"),
 		LOCTEXT("ChaosVisName_ParticleIndices_ToolTip", "Draws the particle indices as instantiated by the solver")),
 	FVisualizationOption(
-		&FClothVisualizationNoGC::DrawElementIndices,
+		FClothVisualizationDebugDrawTexts::CreateLambda([](const ::Chaos::FClothVisualizationNoGC& Visualization, const ::Chaos::FClothVisualizationNoGC::FDrawTextsContext& Context)
+		{
+			Visualization.DrawElementIndices(Context);
+		}),
 		TEXT("p.ChaosClothEditor.DebugDrawElementIndices"),
 		LOCTEXT("ChaosVisName_ElementIndices", "Element Indices"),
 		LOCTEXT("ChaosVisName_ElementIndices_ToolTip", "Draws the element's (triangle or other) indices as instantiated by the solver")),
 	FVisualizationOption(
-		&FClothVisualizationNoGC::DrawPointNormals,
+		FClothVisualizationDebugDraw::CreateLambda([](const ::Chaos::FClothVisualizationNoGC& Visualization, const ::Chaos::FClothVisualizationNoGC::FDrawContext& Context)
+			{
+				Visualization.DrawPointNormals(Context);
+			}),
 		TEXT("p.ChaosClothEditor.DebugDrawPointNormals"),
 		LOCTEXT("ChaosVisName_PointNormals", "Physical Mesh Normals"),
 		LOCTEXT("ChaosVisName_PointNormals_ToolTip", "Draws the current point normals for the simulation mesh")),
 	FVisualizationOption(
-		&FClothVisualizationNoGC::DrawPointVelocities,
+		FClothVisualizationDebugDraw::CreateLambda([](const ::Chaos::FClothVisualizationNoGC& Visualization, const ::Chaos::FClothVisualizationNoGC::FDrawContext& Context)
+			{
+				Visualization.DrawPointVelocities(Context);
+			}),
 		TEXT("p.ChaosClothEditor.DebugDrawPointVelocities"),
 		LOCTEXT("ChaosVisName_PointVelocities", "Point Velocities"),
 		LOCTEXT("ChaosVisName_PointVelocities_ToolTip", "Draws the current point velocities for the simulation mesh")),
 	FVisualizationOption(
-		&FClothVisualizationNoGC::DrawAnimNormals,
+		FClothVisualizationDebugDraw::CreateLambda([](const ::Chaos::FClothVisualizationNoGC& Visualization, const ::Chaos::FClothVisualizationNoGC::FDrawContext& Context)
+			{
+				Visualization.DrawAnimNormals(Context);
+			}),
 		TEXT("p.ChaosClothEditor.DebugDrawAnimNormals"),
 		LOCTEXT("ChaosVisName_AnimNormals", "Animated Mesh Normals"),
 		LOCTEXT("ChaosVisName_AnimNormals_ToolTip", "Draws the current point normals for the animated mesh")),
 	FVisualizationOption(
-		&FClothVisualizationNoGC::DrawCollision,
+		FClothVisualizationDebugDraw::CreateLambda([](const ::Chaos::FClothVisualizationNoGC& Visualization, const ::Chaos::FClothVisualizationNoGC::FDrawContext& Context)
+			{
+				Visualization.DrawCollision(Context);
+			}),
 		TEXT("p.ChaosClothEditor.DebugDrawCollision"),
 		LOCTEXT("ChaosVisName_Collision", "Collisions"),
 		LOCTEXT("ChaosVisName_Collision_ToolTip", "Draws the collision bodies the simulation is currently using")),
 	FVisualizationOption(
-		&FClothVisualizationNoGC::DrawBackstops,
+		FClothVisualizationDebugDraw::CreateLambda([](const ::Chaos::FClothVisualizationNoGC& Visualization, const ::Chaos::FClothVisualizationNoGC::FDrawContext& Context)
+			{
+				Visualization.DrawBackstops(Context);
+			}),
 		TEXT("p.ChaosClothEditor.DebugDrawBackstops"),
 		LOCTEXT("ChaosVisName_Backstop", "Backstops"),
 		LOCTEXT("ChaosVisName_Backstop_ToolTip", "Draws the backstop radius and position for each simulation particle")),
 	FVisualizationOption(
-		&FClothVisualizationNoGC::DrawBackstopDistances,
+		FClothVisualizationDebugDraw::CreateLambda([](const ::Chaos::FClothVisualizationNoGC& Visualization, const ::Chaos::FClothVisualizationNoGC::FDrawContext& Context)
+			{
+				Visualization.DrawBackstopDistances(Context);
+			}),
 		TEXT("p.ChaosClothEditor.DebugDrawBackstopDistances"),
 		LOCTEXT("ChaosVisName_BackstopDistance", "Backstop Distances"),
 		LOCTEXT("ChaosVisName_BackstopDistance_ToolTip", "Draws the backstop distance offset for each simulation particle"),
 		/*bDisablesSimulation =*/true),
 	FVisualizationOption(
-		&FClothVisualizationNoGC::DrawMaxDistances,
+		FClothVisualizationDebugDraw::CreateLambda([](const ::Chaos::FClothVisualizationNoGC& Visualization, const ::Chaos::FClothVisualizationNoGC::FDrawContext& Context)
+			{
+				Visualization.DrawMaxDistances(Context);
+			}),
 		TEXT("p.ChaosClothEditor.DebugDrawMaxDistances"),
 		LOCTEXT("ChaosVisName_MaxDistance", "Max Distances"),
 		LOCTEXT("ChaosVisName_MaxDistance_ToolTip", "Draws the current max distances for the sim particles as a line along its normal"),
 		/*bDisablesSimulation =*/true),
 	FVisualizationOption(
-		&FClothVisualizationNoGC::DrawMaxDistanceValues,
+		FClothVisualizationDebugDrawTexts::CreateLambda([](const ::Chaos::FClothVisualizationNoGC& Visualization, const ::Chaos::FClothVisualizationNoGC::FDrawTextsContext& Context)
+			{
+				Visualization.DrawMaxDistanceValues(Context);
+			}),
 		TEXT("p.ChaosClothEditor.DebugDrawMaxDistanceValues"),
 		LOCTEXT("ChaosVisName_MaxDistanceValue", "Max Distances As Numbers"),
 		LOCTEXT("ChaosVisName_MaxDistanceValue_ToolTip", "Draws the current max distances as numbers")),
 	FVisualizationOption(
-		&FClothVisualizationNoGC::DrawAnimDrive,
+		FClothVisualizationDebugDraw::CreateLambda([](const ::Chaos::FClothVisualizationNoGC& Visualization, const ::Chaos::FClothVisualizationNoGC::FDrawContext& Context)
+			{
+				Visualization.DrawAnimDrive(Context);
+			}),
 		TEXT("p.ChaosClothEditor.DebugDrawAnimDrive"),
 		LOCTEXT("ChaosVisName_AnimDrive", "Anim Drive"),
 		LOCTEXT("ChaosVisName_AnimDrive_Tooltip", "Draws the current skinned reference mesh for the simulation which anim drive will attempt to reach if enabled")),
 	FVisualizationOption(
-		&FClothVisualizationNoGC::DrawEdgeConstraint,
+		FClothVisualizationDebugDraw::CreateLambda([](const ::Chaos::FClothVisualizationNoGC& Visualization, const ::Chaos::FClothVisualizationNoGC::FDrawContext& Context)
+			{
+				Visualization.DrawEdgeConstraint(Context);
+			}),
 		TEXT("p.ChaosClothEditor.DebugDrawEdgeConstraint"),
 		LOCTEXT("ChaosVisName_EdgeConstraint", "Edge Constraint"),
 		LOCTEXT("ChaosVisName_EdgeConstraint_Tooltip", "Draws the edge spring constraints")),
 	FVisualizationOption(
-		&FClothVisualizationNoGC::DrawBendingConstraint,
+		FClothVisualizationDebugDraw::CreateLambda([](const ::Chaos::FClothVisualizationNoGC& Visualization, const ::Chaos::FClothVisualizationNoGC::FDrawContext& Context)
+			{
+				Visualization.DrawBendingConstraint(Context);
+			}),
 		TEXT("p.ChaosClothEditor.DebugDrawBendingConstraint"),
 		LOCTEXT("ChaosVisName_BendingConstraint", "Bending Constraint"),
 		LOCTEXT("ChaosVisName_BendingConstraint_Tooltip", "Draws the bending spring constraints")),
 	FVisualizationOption(
-		&FClothVisualizationNoGC::DrawLongRangeConstraint,
+		FClothVisualizationDebugDraw::CreateLambda([](const ::Chaos::FClothVisualizationNoGC& Visualization, const ::Chaos::FClothVisualizationNoGC::FDrawContext& Context)
+			{
+				Visualization.DrawLongRangeConstraint(Context);
+			}),
 		TEXT("p.ChaosClothEditor.DebugDrawLongRangeConstraint"),
 		LOCTEXT("ChaosVisName_LongRangeConstraint", "Long Range Constraint"),
 		LOCTEXT("ChaosVisName_LongRangeConstraint_Tooltip", "Draws the long range attachment constraint distances")),
 	FVisualizationOption(
-		&FClothVisualizationNoGC::DrawWindAndPressureForces,
+		FClothVisualizationDebugDraw::CreateLambda([](const ::Chaos::FClothVisualizationNoGC& Visualization, const ::Chaos::FClothVisualizationNoGC::FDrawContext& Context)
+			{
+				Visualization.DrawWindAndPressureForces(Context);
+			}),
 		TEXT("p.ChaosClothEditor.DebugDrawWindAndPressureForces"),
 		LOCTEXT("ChaosVisName_WindAndPressureForces", "Wind Aerodynamic And Pressure Forces"),
 		LOCTEXT("ChaosVisName_WindAndPressure_Tooltip", "Draws the Wind drag and lift and pressure forces")),
 	FVisualizationOption(
-		&FClothVisualizationNoGC::DrawWindVelocity,
+		FClothVisualizationDebugDraw::CreateLambda([](const ::Chaos::FClothVisualizationNoGC& Visualization, const ::Chaos::FClothVisualizationNoGC::FDrawContext& Context)
+			{
+				Visualization.DrawWindVelocity(Context);
+			}),
 		TEXT("p.ChaosClothEditor.DebugDrawWindVelocity"),
 		LOCTEXT("ChaosVisName_WindVelocity", "Wind Velocity"),
 		LOCTEXT("ChaosVisName_WindVelocity_Tooltip", "Draws the Wind Velocity vector")),
 	FVisualizationOption(
-		&FClothVisualizationNoGC::DrawLocalSpace,
+		FClothVisualizationDebugDraw::CreateLambda([](const ::Chaos::FClothVisualizationNoGC& Visualization, const ::Chaos::FClothVisualizationNoGC::FDrawContext& Context)
+			{
+				Visualization.DrawLocalSpace(Context);
+			}),
 		TEXT("p.ChaosClothEditor.DebugDrawLocalSpace"),
 		LOCTEXT("ChaosVisName_LocalSpace", "Local Space Reference Bone"),
 		LOCTEXT("ChaosVisName_LocalSpace_Tooltip", "Draws the local space reference bone")),
 	FVisualizationOption(
-		&FClothVisualizationNoGC::DrawLocalSpaceBoneNames,
+		FClothVisualizationDebugDrawTexts::CreateLambda([](const ::Chaos::FClothVisualizationNoGC& Visualization, const ::Chaos::FClothVisualizationNoGC::FDrawTextsContext& Context)
+			{
+				Visualization.DrawLocalSpaceBoneNames(Context);
+			}),
 		TEXT("p.ChaosClothEditor.DebugDrawLocalSpaceBoneNames"),
 		LOCTEXT("ChaosVisName_LocalSpaceBoneNames", "Local Space Reference Bone Names"),
 		LOCTEXT("ChaosVisName_LocalSpaceBoneNames_Tooltip", "Draws the local space reference bone names")),
 	FVisualizationOption(
-		&FClothVisualizationNoGC::DrawVelocityScale,
+		FClothVisualizationDebugDraw::CreateLambda([](const ::Chaos::FClothVisualizationNoGC& Visualization, const ::Chaos::FClothVisualizationNoGC::FDrawContext& Context)
+			{
+				Visualization.DrawVelocityScale(Context);
+			}),
 		TEXT("p.ChaosClothEditor.DebugDrawVelocityScale"),
 		LOCTEXT("ChaosVisName_VelocityScale", "Velocity Scale"),
 		LOCTEXT("ChaosVisName_VelocityScale_Tooltip", "Draws the velocity scale")),
 	FVisualizationOption(
-		&FClothVisualizationNoGC::DrawSelfCollision,
+		FClothVisualizationDebugDraw::CreateLambda([](const ::Chaos::FClothVisualizationNoGC& Visualization, const ::Chaos::FClothVisualizationNoGC::FDrawContext& Context)
+			{
+				Visualization.DrawSelfCollision(Context);
+			}),
 		TEXT("p.ChaosClothEditor.DebugDrawSelfCollision"),
 		LOCTEXT("ChaosVisName_SelfCollision", "Self Collision"),
 		LOCTEXT("ChaosVisName_SelfCollision_Tooltip", "Draws the self collision thickness/debugging information")),
 	FVisualizationOption(
-		&FClothVisualizationNoGC::DrawSelfIntersection,
+		FClothVisualizationDebugDraw::CreateLambda([](const ::Chaos::FClothVisualizationNoGC& Visualization, const ::Chaos::FClothVisualizationNoGC::FDrawContext& Context)
+			{
+				Visualization.DrawSelfIntersection(Context);
+			}),
 		TEXT("p.ChaosClothEditor.DebugDrawSelfIntersection"),
 		LOCTEXT("ChaosVisName_SelfIntersection", "Self Intersection"),
 		LOCTEXT("ChaosVisName_SelfIntersection_Tooltip", "Draws the self intersection contour/region information"))
 };
 const uint32 FVisualizationOption::Count = sizeof(OptionData) / sizeof(FVisualizationOption);
 
+namespace ClothEditorCVar
+{
+	namespace Private
+	{
+		// Add these as visualization options at some point?
+		static TAutoConsoleVariable<int32> DebugDrawSingleVertex(TEXT("p.ChaosClothEditor.DebugDrawSingleVertex"), -1, TEXT("Filter Chaos Cloth debug draw to a single vertex. Default = -1 is no filtering."), ECVF_Cheat);
+		static TAutoConsoleVariable<FString> DebugDrawFilterSet(TEXT("p.ChaosClothEditor.DebugDrawFilterSet"), "", TEXT("Filter Chaos Cloth debug draw to a named sim vertex set. Default ='' is no filtering."), ECVF_Cheat);
+	}
+}
+
 FSimulationEditorExtender::FSimulationEditorExtender()
 	: Flags(false, FVisualizationOption::Count)
 {
@@ -256,12 +340,15 @@ void FSimulationEditorExtender::DebugDrawSimulation(const IClothingSimulationInt
 
 	if (const FClothVisualizationNoGC* const ClothVisualization = GetClothVisualization(Simulation))
 	{
+		FClothVisualizationNoGC::FDrawContext Context(PDI);
+		Context.Filter.SingleVertex = ClothEditorCVar::Private::DebugDrawSingleVertex.GetValueOnAnyThread();
+		Context.Filter.VertexSet = ClothEditorCVar::Private::DebugDrawFilterSet.GetValueOnAnyThread();
 		for (int32 OptionIndex = 0; OptionIndex < FVisualizationOption::Count; ++OptionIndex)
 		{
-			if (FVisualizationOption::OptionData[OptionIndex].DebugDrawFunction &&
+			if (FVisualizationOption::OptionData[OptionIndex].DebugDrawFunction.IsBound() &&
 				(FVisualizationOption::OptionData[OptionIndex].IsConsoleVariableEnabled() || Flags[OptionIndex]))
 			{
-				(ClothVisualization->*(FVisualizationOption::OptionData[OptionIndex].DebugDrawFunction))(PDI);
+				FVisualizationOption::OptionData[OptionIndex].DebugDrawFunction.Execute(*ClothVisualization, Context);
 			}
 		}
 	}
@@ -276,12 +363,15 @@ void FSimulationEditorExtender::DebugDrawSimulationTexts(const IClothingSimulati
 
 	if (const FClothVisualizationNoGC* const ClothVisualization = GetClothVisualization(Simulation))
 	{
+		FClothVisualizationNoGC::FDrawTextsContext Context(Canvas, SceneView);
+		Context.Filter.SingleVertex = ClothEditorCVar::Private::DebugDrawSingleVertex.GetValueOnAnyThread();
+		Context.Filter.VertexSet = ClothEditorCVar::Private::DebugDrawFilterSet.GetValueOnAnyThread();
 		for (int32 OptionIndex = 0; OptionIndex < FVisualizationOption::Count; ++OptionIndex)
 		{
-			if (FVisualizationOption::OptionData[OptionIndex].DebugDrawTextsFunction &&
+			if (FVisualizationOption::OptionData[OptionIndex].DebugDrawTextsFunction.IsBound() &&
 				(FVisualizationOption::OptionData[OptionIndex].IsConsoleVariableEnabled() || Flags[OptionIndex]))
 			{
-				(ClothVisualization->*(FVisualizationOption::OptionData[OptionIndex].DebugDrawTextsFunction))(Canvas, SceneView);
+				FVisualizationOption::OptionData[OptionIndex].DebugDrawTextsFunction.Execute(*ClothVisualization, Context);
 			}
 		}
 	}
diff --git a/Engine/Plugins/ChaosClothAsset/Source/ChaosClothAsset/Private/ChaosClothAsset/CollectionClothFacade.cpp b/Engine/Plugins/ChaosClothAsset/Source/ChaosClothAsset/Private/ChaosClothAsset/CollectionClothFacade.cpp
--- a/Engine/Plugins/ChaosClothAsset/Source/ChaosClothAsset/Private/ChaosClothAsset/CollectionClothFacade.cpp
+++ b/Engine/Plugins/ChaosClothAsset/Source/ChaosClothAsset/Private/ChaosClothAsset/CollectionClothFacade.cpp
@@ -1668,7 +1668,6 @@ namespace UE::Chaos::ClothAsset
 	void FCollectionClothFacade::SetDefaults()
 	{
 		GetClothCollection()->SetNumElements(1, ClothCollectionGroup::Lods);
-		SetCopyAllUserAttributesToSimModel(false);
 
 		if (IsValid(EClothCollectionExtendedSchemas::Solvers))
 		{
diff --git a/Engine/Plugins/ChaosClothAsset/Source/ChaosClothAssetEngine/Private/ChaosClothAsset/ClothSimulationProxy.cpp b/Engine/Plugins/ChaosClothAsset/Source/ChaosClothAssetEngine/Private/ChaosClothAsset/ClothSimulationProxy.cpp
--- a/Engine/Plugins/ChaosClothAsset/Source/ChaosClothAssetEngine/Private/ChaosClothAsset/ClothSimulationProxy.cpp
+++ b/Engine/Plugins/ChaosClothAsset/Source/ChaosClothAssetEngine/Private/ChaosClothAsset/ClothSimulationProxy.cpp
@@ -672,6 +672,8 @@ namespace UE::Chaos::ClothAsset
 		const UWorld* const World = ClothComponentAdapter.GetOwnerComponent().GetWorld();
 		if (World && World->WorldType != EWorldType::EditorPreview)  // Avoid drawing any editor preview debug draw inside the PIE window
 		{
+			static const TConsoleVariableData<int32>* const DebugDrawSingleVertexCVar = IConsoleManager::Get().FindTConsoleVariableDataInt(TEXT("p.ChaosCloth.DebugDrawSingleVertex"));
+			static const TConsoleVariableData<FString>* const DebugDrawFilterSetCVar = IConsoleManager::Get().FindTConsoleVariableDataString(TEXT("p.ChaosCloth.DebugDrawFilterSet"));
 			static const TConsoleVariableData<bool>* const DebugDrawLocalSpaceCVar = IConsoleManager::Get().FindTConsoleVariableDataBool(TEXT("p.ChaosCloth.DebugDrawLocalSpace"));
 			static const TConsoleVariableData<bool>* const DebugDrawBoundsCVar = IConsoleManager::Get().FindTConsoleVariableDataBool(TEXT("p.ChaosCloth.DebugDrawBounds"));
 			static const TConsoleVariableData<bool>* const DebugDrawGravityCVar = IConsoleManager::Get().FindTConsoleVariableDataBool(TEXT("p.ChaosCloth.DebugDrawGravity"));
@@ -700,45 +702,41 @@ namespace UE::Chaos::ClothAsset
 			static const TConsoleVariableData<bool>* const DebugDrawExtremlyDeformedEdgesCVar = IConsoleManager::Get().FindTConsoleVariableDataBool(TEXT("p.ChaosCloth.DebugDrawExtremlyDeformedEdges"));
 			static const TConsoleVariableData<bool>* const DebugDrawVelocityScaleCVar = IConsoleManager::Get().FindTConsoleVariableDataBool(TEXT("p.ChaosCloth.DebugDrawVelocityScale"));
 
-			if (DebugDrawLocalSpaceCVar && DebugDrawLocalSpaceCVar->GetValueOnAnyThread()) { Visualization->DrawLocalSpace(); }
-			if (DebugDrawBoundsCVar && DebugDrawBoundsCVar->GetValueOnAnyThread()) { Visualization->DrawBounds(); }
-			if (DebugDrawGravityCVar && DebugDrawGravityCVar->GetValueOnAnyThread()) { Visualization->DrawGravity(); }
-			if (DebugDrawPhysMeshWiredCVar && DebugDrawPhysMeshWiredCVar->GetValueOnAnyThread()) { Visualization->DrawPhysMeshWired(); }
-			if (DebugDrawAnimMeshWiredCVar && DebugDrawAnimMeshWiredCVar->GetValueOnAnyThread()) { Visualization->DrawAnimMeshWired(); }
-			if (DebugDrawPointVelocitiesCVar && DebugDrawPointVelocitiesCVar->GetValueOnAnyThread()) { Visualization->DrawPointVelocities(); }
-			if (DebugDrawAnimNormalsCVar && DebugDrawAnimNormalsCVar->GetValueOnAnyThread()) { Visualization->DrawAnimNormals(); }
-			if (DebugDrawPointNormalsCVar && DebugDrawPointNormalsCVar->GetValueOnAnyThread()) { Visualization->DrawPointNormals(); }
-			if (DebugDrawCollisionCVar && DebugDrawCollisionCVar->GetValueOnAnyThread()) { Visualization->DrawCollision(); }
-			if (DebugDrawBackstopsCVar && DebugDrawBackstopsCVar->GetValueOnAnyThread()) { Visualization->DrawBackstops(); }
-			if (DebugDrawBackstopDistancesCVar && DebugDrawBackstopDistancesCVar->GetValueOnAnyThread()) { Visualization->DrawBackstopDistances(); }
-			if (DebugDrawMaxDistancesCVar && DebugDrawMaxDistancesCVar->GetValueOnAnyThread()) { Visualization->DrawMaxDistances(); }
-			if (DebugDrawMaxDistanceValuesCVar && DebugDrawMaxDistanceValuesCVar->GetValueOnAnyThread()) { Visualization->DrawMaxDistanceValues(); }
-			if (DebugDrawAnimDriveCVar && DebugDrawAnimDriveCVar->GetValueOnAnyThread()) { Visualization->DrawAnimDrive(); }
-			if (DebugDrawEdgeConstraintCVar && DebugDrawEdgeConstraintCVar->GetValueOnAnyThread()) { Visualization->DrawEdgeConstraint(); }
-			if (DebugDrawBendingConstraintCVar && DebugDrawBendingConstraintCVar->GetValueOnAnyThread()) { Visualization->DrawBendingConstraint(); }
-			if (DebugDrawLongRangeConstraintCVar && DebugDrawLongRangeConstraintCVar->GetValueOnAnyThread()) { Visualization->DrawLongRangeConstraint(); }
-			if (DebugDrawWindForcesCVar && DebugDrawWindForcesCVar->GetValueOnAnyThread()) { Visualization->DrawWindAndPressureForces(); }
-			if (DebugDrawWindVelocityCVar && DebugDrawWindVelocityCVar->GetValueOnAnyThread()) { Visualization->DrawWindVelocity(); }
-			if (DebugDrawSelfCollisionCVar && DebugDrawSelfCollisionCVar->GetValueOnAnyThread()) { Visualization->DrawSelfCollision(); }
-			if (DebugDrawSelfIntersectionCVar && DebugDrawSelfIntersectionCVar->GetValueOnAnyThread()) { Visualization->DrawSelfIntersection(); }
-			if (DebugDrawParticleIndicesCVar && DebugDrawParticleIndicesCVar->GetValueOnAnyThread()) { Visualization->DrawParticleIndices(); }
-			if (DebugDrawElementIndicesCVar && DebugDrawElementIndicesCVar->GetValueOnAnyThread()) { Visualization->DrawElementIndices(); }
-			if (DebugDrawClothClothConstraintsCVar && DebugDrawClothClothConstraintsCVar->GetValueOnAnyThread())
-			{
-				Visualization->DrawClothClothConstraints();
-			}
-			if (DebugDrawTeleportResetCVar && DebugDrawTeleportResetCVar->GetValueOnAnyThread())
-			{
-				Visualization->DrawTeleportReset();
-			}
-			if (DebugDrawExtremlyDeformedEdgesCVar && DebugDrawExtremlyDeformedEdgesCVar->GetValueOnAnyThread())
-			{
-				Visualization->DrawExtremlyDeformedEdges();
-			}
-			if (DebugDrawVelocityScaleCVar && DebugDrawVelocityScaleCVar->GetValueOnAnyThread())
-			{
-				Visualization->DrawVelocityScale();
-			}
+			FClothVisualizationNoGC::FDrawContext DrawContext;
+			DrawContext.Filter.SingleVertex = DebugDrawSingleVertexCVar ? DebugDrawSingleVertexCVar->GetValueOnAnyThread() : INDEX_NONE;
+			DrawContext.Filter.VertexSet = DebugDrawFilterSetCVar ? DebugDrawFilterSetCVar->GetValueOnAnyThread() : FString();
+			FClothVisualizationNoGC::FDrawTextsContext DrawTextsContext;
+			DrawTextsContext.Filter = DrawContext.Filter;
+
+			if (DebugDrawParticleIndicesCVar && DebugDrawParticleIndicesCVar->GetValueOnAnyThread()) { Visualization->DrawParticleIndices(DrawTextsContext); }
+			if (DebugDrawElementIndicesCVar && DebugDrawElementIndicesCVar->GetValueOnAnyThread()) { Visualization->DrawElementIndices(DrawTextsContext); }
+			if (DebugDrawLocalSpaceCVar && DebugDrawLocalSpaceCVar->GetValueOnAnyThread()) { Visualization->DrawLocalSpace(DrawContext); }
+			if (DebugDrawBoundsCVar && DebugDrawBoundsCVar->GetValueOnAnyThread()) { Visualization->DrawBounds(DrawContext); }
+			if (DebugDrawGravityCVar && DebugDrawGravityCVar->GetValueOnAnyThread()) { Visualization->DrawGravity(DrawContext); }
+			if (DebugDrawPhysMeshWiredCVar && DebugDrawPhysMeshWiredCVar->GetValueOnAnyThread()) { Visualization->DrawPhysMeshWired(DrawContext); }
+			if (DebugDrawAnimMeshWiredCVar && DebugDrawAnimMeshWiredCVar->GetValueOnAnyThread()) { Visualization->DrawAnimMeshWired(DrawContext); }
+			if (DebugDrawPointVelocitiesCVar && DebugDrawPointVelocitiesCVar->GetValueOnAnyThread()) { Visualization->DrawPointVelocities(DrawContext); }
+			if (DebugDrawAnimNormalsCVar && DebugDrawAnimNormalsCVar->GetValueOnAnyThread()) { Visualization->DrawAnimNormals(DrawContext); }
+			if (DebugDrawPointNormalsCVar && DebugDrawPointNormalsCVar->GetValueOnAnyThread()) { Visualization->DrawPointNormals(DrawContext); }
+			if (DebugDrawCollisionCVar && DebugDrawCollisionCVar->GetValueOnAnyThread()) { Visualization->DrawCollision(DrawContext); }
+			if (DebugDrawBackstopsCVar && DebugDrawBackstopsCVar->GetValueOnAnyThread()) { Visualization->DrawBackstops(DrawContext); }
+			if (DebugDrawBackstopDistancesCVar && DebugDrawBackstopDistancesCVar->GetValueOnAnyThread()) { Visualization->DrawBackstopDistances(DrawContext); }
+			if (DebugDrawMaxDistancesCVar && DebugDrawMaxDistancesCVar->GetValueOnAnyThread()) { Visualization->DrawMaxDistances(DrawContext); }
+			if (DebugDrawMaxDistanceValuesCVar && DebugDrawMaxDistanceValuesCVar->GetValueOnAnyThread()) { Visualization->DrawMaxDistanceValues(DrawTextsContext); }
+			if (DebugDrawAnimDriveCVar && DebugDrawAnimDriveCVar->GetValueOnAnyThread()) { Visualization->DrawAnimDrive(DrawContext); }
+			if (DebugDrawEdgeConstraintCVar && DebugDrawEdgeConstraintCVar->GetValueOnAnyThread()) { Visualization->DrawEdgeConstraint(DrawContext); }
+			if (DebugDrawBendingConstraintCVar && DebugDrawBendingConstraintCVar->GetValueOnAnyThread()) { Visualization->DrawBendingConstraint(DrawContext); }
+			if (DebugDrawLongRangeConstraintCVar && DebugDrawLongRangeConstraintCVar->GetValueOnAnyThread()) { Visualization->DrawLongRangeConstraint(DrawContext); }
+			if (DebugDrawWindForcesCVar && DebugDrawWindForcesCVar->GetValueOnAnyThread()) { Visualization->DrawWindAndPressureForces(DrawContext); }
+			if (DebugDrawWindVelocityCVar && DebugDrawWindVelocityCVar->GetValueOnAnyThread()) { Visualization->DrawWindVelocity(DrawContext); }
+			if (DebugDrawSelfCollisionCVar && DebugDrawSelfCollisionCVar->GetValueOnAnyThread()) { Visualization->DrawSelfCollision(DrawContext); }
+			if (DebugDrawSelfIntersectionCVar && DebugDrawSelfIntersectionCVar->GetValueOnAnyThread()) { Visualization->DrawSelfIntersection(DrawContext); }
+			if (DebugDrawParticleIndicesCVar && DebugDrawParticleIndicesCVar->GetValueOnAnyThread()) { Visualization->DrawParticleIndices(DrawTextsContext); }
+			if (DebugDrawElementIndicesCVar && DebugDrawElementIndicesCVar->GetValueOnAnyThread()) { Visualization->DrawElementIndices(DrawTextsContext); }
+			if (DebugDrawClothClothConstraintsCVar && DebugDrawClothClothConstraintsCVar->GetValueOnAnyThread()) { Visualization->DrawClothClothConstraints(DrawContext); }
+			if (DebugDrawTeleportResetCVar && DebugDrawTeleportResetCVar->GetValueOnAnyThread()) { Visualization->DrawTeleportReset(DrawContext); }
+			if (DebugDrawExtremlyDeformedEdgesCVar && DebugDrawExtremlyDeformedEdgesCVar->GetValueOnAnyThread()) { Visualization->DrawExtremlyDeformedEdges(DrawContext); }
+			if (DebugDrawVelocityScaleCVar && DebugDrawVelocityScaleCVar->GetValueOnAnyThread()) { Visualization->DrawVelocityScale(DrawContext); }
 		}
 #endif  // #if CHAOS_DEBUG_DRAW
 	}
