Cleanup old unused research code from Chaos Cloth, namely Gauss-Seidel and multi-res constraints.

#Synced-CL 52021766

diff --git a/Engine/Plugins/ChaosCloth/Source/ChaosCloth/Private/ChaosCloth/ChaosClothConstraints.cpp b/Engine/Plugins/ChaosCloth/Source/ChaosCloth/Private/ChaosCloth/ChaosClothConstraints.cpp
--- a/Engine/Plugins/ChaosCloth/Source/ChaosCloth/Private/ChaosCloth/ChaosClothConstraints.cpp
+++ b/Engine/Plugins/ChaosCloth/Source/ChaosCloth/Private/ChaosCloth/ChaosClothConstraints.cpp
@@ -1,7 +1,6 @@
 // Copyright Epic Games, Inc. All Rights Reserved.
 #include "ChaosCloth/ChaosClothConstraints.h"
 #include "ChaosCloth/ChaosClothingAccessoryMeshData.h"
-#include "ChaosCloth/ChaosClothingAccessoryMeshData.h"
 #include "ChaosCloth/ChaosClothingSimulationConfig.h"
 #include "ChaosCloth/ChaosWeightMapTarget.h"
 #include "ChaosCloth/ChaosClothingPatternData.h"
@@ -17,7 +16,6 @@
 #include "Chaos/XPBDAxialSpringConstraints.h"
 #include "Chaos/PBDVolumeConstraint.h"
 #include "Chaos/XPBDLongRangeConstraints.h"
-#include "Chaos/SoftsMultiResConstraints.h"
 #include "Chaos/PBDSphericalConstraint.h"
 #include "Chaos/PBDAnimDriveConstraint.h"
 #include "Chaos/PBDShapeConstraints.h"
@@ -51,11 +49,6 @@ extern TAutoConsoleVariable<bool> CVarLegacyDisablesAccurateWind;
 extern TAutoConsoleVariable<float> CVarGravityMultiplier;
 }
 
-bool bEnableGS = false;
-bool bDoPlanarCollisionConstraints = false;
-Softs::FSolverReal PlanarCollisionTol = 0.f;
-Softs::FSolverReal PlanarCollisionStiffness = 1e10;
-
 namespace ClothConstraints::Private
 {
 	static bool bEnableBuoyancy = false;  // #todo(dmp): once we decide what we want to do for enabling buoyancy by default or controlling it some way, we can enable it for shipping
@@ -63,65 +56,6 @@ namespace ClothConstraints::Private
 }
 
 #if !UE_BUILD_SHIPPING
-    static FAutoConsoleVariableRef CVarClothbEnableGS(TEXT("p.Chaos.Cloth.EnableGaussSeidel"), bEnableGS, TEXT("Use Gauss Seidel constraints instead of XPBD [def: false]"));
-
-	bool bDisplayResidual = false;
-
-	static FAutoConsoleVariableRef CVarClothbDisplayResidual(TEXT("p.Chaos.Cloth.DisplayResidual"), bDisplayResidual, TEXT("Diplay residual norms for the first 10 iters [def: false]"));
-
-	int32 MaxResidualIters = 10;
-
-	static FAutoConsoleVariableRef CVarClothMaxResidualIters(TEXT("p.Chaos.Cloth.MaxResidualIters"), MaxResidualIters, TEXT("Max number of iterations to diaplay residuals [def: 10]"));
-
-	bool bWriteResidual2File = false;
-
-	static FAutoConsoleVariableRef CVarClothbWriteResidual2File(TEXT("p.Chaos.Cloth.WriteResidual2File"), bWriteResidual2File, TEXT("Write residual to file [def: false]"));
-
-	bool bReplaceBiasElementsWithCorotatedCod = false;
-
-	static FAutoConsoleVariableRef CVarClothReplaceBiasElementsWithCorotatedCod(TEXT("p.Chaos.Cloth.ReplaceBiasElementsWithCorotatedCod"), bReplaceBiasElementsWithCorotatedCod, TEXT("Replace existing aniso bias element constraint with gauss seidel corotated codimensional [def: false]"));
-
-	Softs::FSolverReal YoungsModulus = 10000.f;
-
-	static FAutoConsoleVariableRef CVarClothYoungsModulus(TEXT("p.Chaos.Cloth.YoungsModulus"), YoungsModulus, TEXT("Youngs modulus [def: 1e4]"));
-
-	bool bClothDoQuasistatics = false;
-
-	static FAutoConsoleVariableRef CVarClothDoQuasistatics(TEXT("p.Chaos.Cloth.DoQuasistatics"), bClothDoQuasistatics, TEXT("Do cloth quasistatics [def: false]"));
-
-	bool bEnableCG = false;
-
-	static FAutoConsoleVariableRef CVarClothEnableCG(TEXT("p.Chaos.Cloth.EnableCG"), bEnableCG, TEXT("Use conjugate gradient instead of nonlinear gauss seidel [def: false]"));
-
-	bool bMakeSandwich = false;
-
-	static FAutoConsoleVariableRef CVarClothMakeSandwich(TEXT("p.Chaos.Cloth.MakeSandwich"), bMakeSandwich, TEXT("MakeSandwich solver of xpbd - Newton/GS [def: false]"));
-
-	int32 UpperBreadIters = 3;
-
-	static FAutoConsoleVariableRef CVarClothUpperBreadIters(TEXT("p.Chaos.Cloth.UpperBreadIters"), UpperBreadIters, TEXT("Upper Bread Iters for the sandwich solver [def: 3]"));
-
-	int32 MiddleBreadIters = 3;
-
-	static FAutoConsoleVariableRef CVarClothMiddleBreadIters(TEXT("p.Chaos.Cloth.MiddleBreadIters"), MiddleBreadIters, TEXT("Middle Bread Iters for the sandwich solver [def: 3]"));
-
-	bool bWriteFinalResiduals = false;
-
-	static FAutoConsoleVariableRef CVarClothWriteFinalResiduals(TEXT("p.Chaos.Cloth.WriteFinalResiduals"), bWriteFinalResiduals, TEXT("Write final residuals at each timestep to a file [def: false]"));
-
-	bool bUseSOR = true;
-
-	static FAutoConsoleVariableRef CVarClothUseSOR(TEXT("p.Chaos.Cloth.UseSOR"), bUseSOR, TEXT("Use SOR acceleration for Gauss Seidel [def: true]"));
-
-	Softs::FSolverReal SOROmega = 1.7f;
-
-	static FAutoConsoleVariableRef CVarClothSOROmega(TEXT("p.Chaos.Cloth.SOROmega"), SOROmega, TEXT("SOR omega coefficient for acceleration [def: 1.7]"));
-
-	static FAutoConsoleVariableRef CVarClothDoPlanarCollisionConstraints(TEXT("p.Chaos.Cloth.DoPlanarCollisionConstraints"), bDoPlanarCollisionConstraints, TEXT("Do body collision at the beginning of timestep only and then do xpbd planar constraints [def: true]"));
-
-	static FAutoConsoleVariableRef CVarClothPlanarCollisionTol(TEXT("p.Chaos.Cloth.PlanarCollisionTol"), PlanarCollisionTol, TEXT("Collision tolerance for XPBD planar constraint [def: 0.]"));
-
-	static FAutoConsoleVariableRef CVarClothPlanarCollisionStiffness(TEXT("p.Chaos.Cloth.PlanarCollisionStiffness"), PlanarCollisionStiffness, TEXT("Collision stiffness for XPBD planar constraint [def: 1e10.]"));
 	
 	static FAutoConsoleVariableRef CVarClothEnableBuoyancy(TEXT("p.Chaos.Cloth.EnableBuoyancy"), ClothConstraints::Private::bEnableBuoyancy, TEXT("Query water bodies for drag and buoyancy effects with cloth [def: false]"));
 
@@ -481,312 +415,6 @@ void FClothConstraints::OnCollisionRangeRemoved(int32 CollisionRangeId)
 	}
 }
 
-//only counts the number of rules that GS/CG currently cupports
-void FClothConstraints::GetGSNumRules()
-{
-#if !UE_BUILD_SHIPPING
-	check(PBDEvolution);
-
-	NumConstraintInits = 1;
-	NumConstraintRules = 0;
-	if (XStretchBiasConstraints)
-	{
-		NumConstraintInits += 1;
-	}
-	if (XBendingElementConstraints)
-	{
-		NumConstraintInits += 1;
-	}
-	if (XEdgeConstraints)
-	{
-		NumConstraintInits += 1;
-	}
-	if (XAnisoBendingElementConstraints)
-	{
-		NumConstraintInits += 1;
-	}
-
-	if (bMakeSandwich)
-	{
-		NumConstraintRules = 2;
-	}
-	else
-	{
-		NumConstraintRules = 1;
-	}
-#endif
-}
-
-
-void FClothConstraints::CreateGSRules()
-{
-#if !UE_BUILD_SHIPPING
-	check(PBDEvolution);
-	check(ConstraintInitOffset == INDEX_NONE)
-	GetGSNumRules();
-
-	ConstraintInitOffset = PBDEvolution->AddConstraintInitRange(NumConstraintInits, false);
-
-	check(ConstraintRuleOffset == INDEX_NONE)
-
-	ConstraintRuleOffset = PBDEvolution->AddConstraintRuleRange(NumConstraintRules, false);
-
-	check(PostprocessingConstraintRuleOffset == INDEX_NONE)
-
-	int32 PostProcessingConstraintRule = 0;
-	if (bWriteFinalResiduals)
-	{
-		PostProcessingConstraintRule = 1;
-	}
-
-	if (PostProcessingConstraintRule)
-	{
-		PostprocessingConstraintRuleOffset = PBDEvolution->AddConstraintPostprocessingsRange(PostProcessingConstraintRule, false);
-	}
-
-	TFunction<void(Softs::FSolverParticles&, const Softs::FSolverReal)>* const ConstraintInits = PBDEvolution->ConstraintInits().GetData() + ConstraintInitOffset;
-	TFunction<void(Softs::FSolverParticles&, const Softs::FSolverReal)>* const ConstraintRules = PBDEvolution->ConstraintRules().GetData() + ConstraintRuleOffset;
-	TFunction<void(Softs::FSolverParticles&, const Softs::FSolverReal)>* const PostprocessingRules = PBDEvolution->ConstraintPostprocessings().GetData() + PostprocessingConstraintRuleOffset;
-
-	int32 ConstraintInitIndex = 0;
-	int32 ConstraintRuleIndex = 0;
-
-	int32 ConstraintPostProcessingIndex = 0;
-
-	GSMainConstraint = MakeShared<Chaos::Softs::FGaussSeidelMainConstraints>(PBDEvolution->Particles(), bClothDoQuasistatics, bUseSOR, SOROmega, 100);
-
-	if (bClothDoQuasistatics)
-	{
-		PBDEvolution->SetQuasistatics(bClothDoQuasistatics);
-		const Softs::FSolverVec3& GravityTerm = PBDEvolution->GetGravity(0);
-		GSMainConstraint->ExternalAcceleration[0] = 0.f;
-		GSMainConstraint->ExternalAcceleration[1] = 0.f;
-		GSMainConstraint->ExternalAcceleration[2] = GravityTerm[2];
-	}
-
-	ConstraintInits[ConstraintInitIndex++] =
-		[this](Softs::FSolverParticles& InParticles, const Softs::FSolverReal Dt)
-	{
-		this->GSMainConstraint->Init(Dt, InParticles);
-	};
-
-	if (XStretchBiasConstraints)
-	{
-		ConstraintInits[ConstraintInitIndex++] =
-			[this](Softs::FSolverParticles& /*Particles*/, const Softs::FSolverReal Dt)
-		{
-			XStretchBiasConstraints->Init();
-			XStretchBiasConstraints->ApplyProperties(Dt, Evolution->GetIterations());
-		};
-	}
-
-	if (XBendingElementConstraints)
-	{
-		ConstraintInits[ConstraintInitIndex++] =
-			[this](Softs::FSolverParticles& InParticles, const Softs::FSolverReal Dt)
-		{
-			this->XBendingElementConstraints->Init(InParticles);
-			this->XBendingElementConstraints->ApplyProperties(Dt, Evolution->GetIterations());
-		};
-	}
-
-	if (XAnisoBendingElementConstraints)
-	{
-		ConstraintInits[ConstraintInitIndex++] =
-			[this](Softs::FSolverParticles& InParticles, const Softs::FSolverReal Dt)
-		{
-			this->XAnisoBendingElementConstraints->Init(InParticles);
-			this->XAnisoBendingElementConstraints->ApplyProperties(Dt, Evolution->GetIterations());
-		};
-	}
-
-	if (XEdgeConstraints)
-	{
-		ConstraintInits[ConstraintInitIndex++] =
-			[this](Softs::FSolverParticles& InParticles, const Softs::FSolverReal Dt)
-		{
-			this->XEdgeConstraints->ApplyProperties(Dt, Evolution->GetIterations());
-		};
-	}
-
-	if (bMakeSandwich)
-	{
-		if (XStretchBiasConstraints)
-		{
-			ConstraintRules[ConstraintRuleIndex++] =
-				[this, PreIterXPBDNums = UpperBreadIters, bbDisplayResidual = bDisplayResidual](Softs::FSolverParticles& Particles, const Softs::FSolverReal Dt)
-			{
-				for (int32 i = 0; i < PreIterXPBDNums; i++)
-				{
-					XStretchBiasConstraints->Apply(Particles, Dt);
-					if (this->XBendingElementConstraints)
-					{
-						this->XBendingElementConstraints->Apply(Particles, Dt);
-					}
-					if (bbDisplayResidual)
-					{
-						TArray<Chaos::Softs::FSolverVec3> Residual = this->GSMainConstraint->ComputeNewtonResiduals(Particles, Dt);
-					}
-				}
-			};
-		}
-
-
-		//middle layer of sandwich:
-		if (bEnableCG)
-		{
-			ConstraintRules[ConstraintRuleIndex++] =
-				[this, GSIterNums = MiddleBreadIters](Softs::FSolverParticles& InParticles, const Softs::FSolverReal Dt)
-			{
-				for (int32 i = 0; i < GSIterNums; i++)
-				{
-					this->GSMainConstraint->ApplyCG(InParticles, Dt);
-				}
-			};
-		}
-		else
-		{
-			ConstraintRules[ConstraintRuleIndex++] =
-				[this, GSIterNums = MiddleBreadIters](Softs::FSolverParticles& InParticles, const Softs::FSolverReal Dt)
-			{
-				for (int32 i = 0; i < GSIterNums; i++)
-				{
-					this->GSMainConstraint->Apply(InParticles, Dt, MaxResidualIters, bWriteResidual2File);
-				}
-			};
-		}
-
-	}
-	else
-	{
-		if (bEnableCG)
-		{
-			ConstraintRules[ConstraintRuleIndex++] =
-				[this](Softs::FSolverParticles& InParticles, const Softs::FSolverReal Dt)
-			{
-				this->GSMainConstraint->ApplyCG(InParticles, Dt);
-			};
-		}
-		else
-		{
-			ConstraintRules[ConstraintRuleIndex++] =
-				[this](Softs::FSolverParticles& InParticles, const Softs::FSolverReal Dt)
-			{
-				this->GSMainConstraint->Apply(InParticles, Dt, MaxResidualIters, bWriteResidual2File);
-			};
-		}
-	}
-
-
-
-	if (XStretchBiasConstraints)
-	{
-		if (bReplaceBiasElementsWithCorotatedCod)
-		{
-			TArray<TArray<int32>> IncidentElements, IncidentElementsLocal;
-
-			const TArray<TVec3<int32>> CodimensionalMesh = XStretchBiasConstraints->GetConstraints();
-
-			GSCorotatedCodimensionalConstraint = MakeShared<Chaos::Softs::FGaussSeidelCorotatedTriangleConstraints>(PBDEvolution->Particles(), CodimensionalMesh, false, YoungsModulus);
-
-			GSMainConstraint->AddStaticConstraints(XStretchBiasConstraints->GetConstraintsArray(), IncidentElements, IncidentElementsLocal);
-			const int32 StaticIndex = GSMainConstraint->AddStaticConstraintResidualAndHessianRange(1);
-
-			GSMainConstraint->StaticConstraintResidualAndHessian()[StaticIndex] = [this](const Softs::FSolverParticles& Particles, const int32 ElementIndex, const int32 ElementIndexLocal, const Softs::FSolverReal Dt, TVec3<Softs::FSolverReal>& ParticleResidual, Chaos::PMatrix<Softs::FSolverReal, 3, 3>& ParticleHessian)
-			{
-				this->GSCorotatedCodimensionalConstraint->AddHyperelasticResidualAndHessian(Particles, ElementIndex, ElementIndexLocal, Dt, ParticleResidual, ParticleHessian);
-			};
-		}
-		else
-		{
-			TArray<TArray<int32>> IncidentElements, IncidentElementsLocal;
-			GSMainConstraint->AddStaticConstraints(XStretchBiasConstraints->GetConstraintsArray(), IncidentElements, IncidentElementsLocal);
-			const int32 StaticIndex = GSMainConstraint->AddStaticConstraintResidualAndHessianRange(1);
-
-			GSMainConstraint->StaticConstraintResidualAndHessian()[StaticIndex] = [this](const Softs::FSolverParticles& Particles, const int32 ElementIndex, const int32 ElementIndexLocal, const Softs::FSolverReal Dt, TVec3<Softs::FSolverReal>& ParticleResidual, Chaos::PMatrix<Softs::FSolverReal, 3, 3>& ParticleHessian)
-			{
-				this->XStretchBiasConstraints->AddStretchBiasElementResidualAndHessian(Particles, ElementIndex, ElementIndexLocal, Dt, ParticleResidual, ParticleHessian);
-			};
-
-			if (bEnableCG)
-			{
-				const int32 ForceDifferentialRange = GSMainConstraint->AddAddInternalForceDifferentialsRange(1);
-
-				GSMainConstraint->InternalForceDifferentials()[ForceDifferentialRange] = [this](const Softs::FSolverParticles& Particles, const TArray<Softs::FSolverVec3>& Deltax, TArray<Softs::FSolverVec3>& ndf)
-				{
-					this->XStretchBiasConstraints->AddInternalForceDifferential(Particles, Deltax, ndf);
-				};
-			}
-
-		}
-
-		XStretchBiasConstraints->InitializeDmInvAndMeasures(PBDEvolution->Particles());
-
-	}
-
-	if (XAnisoBendingElementConstraints)
-	{
-		TArray<TArray<int32>> IncidentElements, IncidentElementsLocal;
-		GSMainConstraint->AddStaticConstraints(XAnisoBendingElementConstraints->GetConstraintsArray(), IncidentElements, IncidentElementsLocal);
-		const int32 StaticIndex = GSMainConstraint->AddStaticConstraintResidualAndHessianRange(1);
-
-		GSMainConstraint->StaticConstraintResidualAndHessian()[StaticIndex] = [this](const Softs::FSolverParticles& Particles, const int32 ElementIndex, const int32 ElementIndexLocal, const Softs::FSolverReal Dt, TVec3<Softs::FSolverReal>& ParticleResidual, Chaos::PMatrix<Softs::FSolverReal, 3, 3>& ParticleHessian)
-		{
-			this->XAnisoBendingElementConstraints->AddAnisotropicBendingResidualAndHessian(Particles, ElementIndex, ElementIndexLocal, Dt, ParticleResidual, ParticleHessian);
-		};
-
-		if (bEnableCG)
-		{
-			const int32 ForceDifferentialRange = GSMainConstraint->AddAddInternalForceDifferentialsRange(1);
-
-			GSMainConstraint->InternalForceDifferentials()[ForceDifferentialRange] = [this](const Softs::FSolverParticles& Particles, const TArray<Softs::FSolverVec3>& Deltax, TArray<Softs::FSolverVec3>& ndf)
-			{
-				this->XAnisoBendingElementConstraints->AddInternalForceDifferential(Particles, Deltax, ndf);
-			};
-		}
-	}
-
-	if (XBendingElementConstraints)
-	{
-
-		TArray<TArray<int32>> IncidentElements, IncidentElementsLocal;
-		GSMainConstraint->AddStaticConstraints(XBendingElementConstraints->GetConstraintsArray(), IncidentElements, IncidentElementsLocal);
-		const int32 StaticIndex = GSMainConstraint->AddStaticConstraintResidualAndHessianRange(1);
-
-		GSMainConstraint->StaticConstraintResidualAndHessian()[StaticIndex] = [this](const Softs::FSolverParticles& Particles, const int32 ElementIndex, const int32 ElementIndexLocal, const Softs::FSolverReal Dt, TVec3<Softs::FSolverReal>& ParticleResidual, Chaos::PMatrix<Softs::FSolverReal, 3, 3>& ParticleHessian)
-		{
-			this->XBendingElementConstraints->AddBendingResidualAndHessian(Particles, ElementIndex, ElementIndexLocal, Dt, ParticleResidual, ParticleHessian);
-		};
-
-
-		if (bEnableCG)
-		{
-			const int32 ForceDifferentialRange = GSMainConstraint->AddAddInternalForceDifferentialsRange(1);
-
-			GSMainConstraint->InternalForceDifferentials()[ForceDifferentialRange] = [this](const Softs::FSolverParticles& Particles, const TArray<Softs::FSolverVec3>& Deltax, TArray<Softs::FSolverVec3>& ndf)
-			{
-				this->XBendingElementConstraints->AddInternalForceDifferential(Particles, Deltax, ndf);
-			};
-		}
-	}
-
-	GSMainConstraint->InitStaticColor(PBDEvolution->Particles());
-
-	if (bWriteFinalResiduals)
-	{
-		FString file = FPaths::ProjectDir();
-		file.Append(TEXT("/DebugOutput/NewtonResidual.txt"));
-		FFileHelper::SaveStringToFile(FString(TEXT("Newton Norm\r\n")), *file);
-		PostprocessingRules[ConstraintPostProcessingIndex++] =
-			[this](Softs::FSolverParticles& InParticles, const Softs::FSolverReal Dt)
-		{
-			this->GSMainConstraint->ComputeNewtonResiduals(InParticles, Dt, true, nullptr);
-		};
-	}
-
-#endif
-}
-
 void FClothConstraints::AddRules(
 	const Softs::FCollectionPropertyConstFacade& ConfigProperties,
 	const FTriangleMesh& TriangleMesh,
@@ -797,9 +425,6 @@ void FClothConstraints::AddRules(
 	const TMap<FString, TConstArrayView<int32>>& FaceIntMaps,
 	const TArray<TConstArrayView<TTuple<int32, int32, FRealSingle>>>& Tethers,
 	Softs::FSolverReal MeshScale, bool bEnabled,
-	const FTriangleMesh* MultiResCoarseLODMesh,
-	const int32 MultiResCoarseLODParticleRangeId,
-	const TSharedPtr<Softs::FMultiResConstraints>& FineLODMultiResConstraint,
 	const TArray<const FClothComplexColliders*>& ComplexColliders,
 	const TSharedPtr<const FManagedArrayCollection>& ManagedArrayCollection,
 	const TMap<FName, FClothingAccessoryMeshData>* AccessoryMeshes)
@@ -856,28 +481,18 @@ void FClothConstraints::AddRules(
 		// Body collisions
 		CreateCollisionConstraint(ConfigProperties, MeshScale, WeightMaps, ComplexColliders);
 
-		//Multires Springs
-		CreateMultiresConstraint(ConfigProperties, WeightMaps, TriangleMesh, MultiResCoarseLODMesh, MultiResCoarseLODParticleRangeId);
-
 		// ClothCloth Springs
 		CreateClothClothConstraints(ConfigProperties, ManagedArrayCollection);
 	}
 
 	// Commit rules to solver
-	if (!bEnableGS)
+	if (Evolution)
 	{
-		if (Evolution)
-		{
-			CreateForceBasedRules(FineLODMultiResConstraint);
-		}
-		else
-		{
-			CreatePBDRules();
-		}
-	} 
+		CreateForceBasedRules();
+	}
 	else
 	{
-		CreateGSRules();
+		CreatePBDRules();
 	}
 
 	if (PBDEvolution)
@@ -1599,89 +1214,6 @@ void FClothConstraints::CreateCollisionConstraint(
 	}
 }
 
-void FClothConstraints::CreateMultiresConstraint(
-	const Softs::FCollectionPropertyConstFacade& ConfigProperties,
-	const TMap<FString, TConstArrayView<FRealSingle>>& WeightMaps,
-	const FTriangleMesh& TriangleMesh,
-	const FTriangleMesh* MultiResCoarseLODMesh,
-	const int32 MultiResCoarseLODParticleRangeId)
-{
-	if (Evolution && MultiResCoarseLODMesh && MultiResCoarseLODParticleRangeId != INDEX_NONE)
-	{
-		if (Softs::FMultiResConstraints::IsEnabled(ConfigProperties))
-		{
-			TArray<TVec4<Softs::FSolverReal>> CoarseToFinePositionBaryCoordsAndDist;
-			TArray<TVec3<int32>> CoarseToFineSourceMeshVertIndices;
-			const Softs::FSolverParticlesRange& FineParticles = Evolution->GetSoftBodyParticles(ParticleRangeId);
-			{
-				TRACE_CPUPROFILER_EVENT_SCOPE(BuildMultiResTransitionData);
-				// TODO: Cache this data in the asset.
-				const FPointWeightMap* const MaxDistances = nullptr;  // No need to update the vertex contribution on the transition maps
-				constexpr bool bUseSmoothTransitions = false;         // Smooth transitions are only used at rendering for now and not during LOD transitions
-				constexpr bool bUseMultipleInfluences = false;        // Multiple influences must not be used for LOD transitions
-				constexpr float SkinningKernelRadius = 0.f;           // KernelRadius is only required when using multiple influences
-				static_assert(std::is_same<float, Softs::FSolverReal>::value);
-				static_assert(sizeof(Softs::FSolverVec3) == sizeof(FVector3f));
-
-				TConstArrayView<FVector3f> FinePositions((const FVector3f*)FineParticles.XArray().GetData(), FineParticles.Size());
-				TArray<uint32> FineIndices;
-				FineIndices.Reserve(3 * TriangleMesh.GetNumElements());
-				for (const TVec3<int32>& Element : TriangleMesh.GetElements())
-				{
-					FineIndices.Add(Element[0]);
-					FineIndices.Add(Element[1]);
-					FineIndices.Add(Element[2]);
-				}
-
-				const ClothingMeshUtils::ClothMeshDesc FineLodDesc(FinePositions, TConstArrayView<uint32>(FineIndices));
-
-				const Softs::FSolverParticlesRange& CoarseParticles = Evolution->GetSoftBodyParticles(MultiResCoarseLODParticleRangeId);
-				TConstArrayView<FVector3f> CoarsePositions((const FVector3f*)CoarseParticles.XArray().GetData(), CoarseParticles.Size());
-				TConstArrayView<Softs::FSolverReal> CoarseInvM = CoarseParticles.GetInvM();
-				TArray<uint32> CoarseIndices;
-				CoarseIndices.Reserve(3 * MultiResCoarseLODMesh->GetNumElements());
-				for (const TVec3<int32>& Element : MultiResCoarseLODMesh->GetElements())
-				{
-					if (CoarseInvM[Element[0]] != 0.f || CoarseInvM[Element[1]] != 0.f || CoarseInvM[Element[2]] != 0.f)
-					{
-						CoarseIndices.Add(Element[0]);
-						CoarseIndices.Add(Element[1]);
-						CoarseIndices.Add(Element[2]);
-					}
-				}
-
-				if (CoarseIndices.IsEmpty())
-				{
-					return;
-				}
-
-				const ClothingMeshUtils::ClothMeshDesc CoarseLodDesc(CoarsePositions, TConstArrayView<uint32>(CoarseIndices));
-				TArray<FMeshToMeshVertData> TransitionData;
-				ClothingMeshUtils::GenerateMeshToMeshVertData(TransitionData, FineLodDesc, CoarseLodDesc, MaxDistances, bUseSmoothTransitions, bUseMultipleInfluences, SkinningKernelRadius);
-
-				CoarseToFinePositionBaryCoordsAndDist.Reserve(TransitionData.Num());
-				CoarseToFineSourceMeshVertIndices.Reserve(TransitionData.Num());
-				for (const FMeshToMeshVertData& Data : TransitionData)
-				{
-					CoarseToFinePositionBaryCoordsAndDist.Emplace(Data.PositionBaryCoordsAndDist.X, Data.PositionBaryCoordsAndDist.Y, Data.PositionBaryCoordsAndDist.Z, Data.PositionBaryCoordsAndDist.W);
-					CoarseToFineSourceMeshVertIndices.Emplace(Data.SourceMeshVertIndices[0], Data.SourceMeshVertIndices[1], Data.SourceMeshVertIndices[2]);
-				}
-			}
-
-			MultiResConstraints = MakeShared<Softs::FMultiResConstraints>(
-				FineParticles,
-				MultiResCoarseLODParticleRangeId,
-				*MultiResCoarseLODMesh,
-				MoveTemp(CoarseToFinePositionBaryCoordsAndDist),
-				MoveTemp(CoarseToFineSourceMeshVertIndices),
-				WeightMaps,
-				ConfigProperties);
-			++NumConstraintRules;
-			++NumConstraintInits;
-		}
-	}
-}
-
 void FClothConstraints::CreateClothClothConstraints(
 	const Softs::FCollectionPropertyConstFacade& ConfigProperties,
 	const TSharedPtr<const FManagedArrayCollection>& ManagedArrayCollection)
@@ -1776,13 +1308,8 @@ void FClothConstraints::CreateClothClothConstraints(
 	}
 }
 
-void FClothConstraints::CreateForceBasedRules(const TSharedPtr<Softs::FMultiResConstraints>& FineLODMultiResConstraint)
+void FClothConstraints::CreateForceBasedRules()
 {
-	if (FineLODMultiResConstraint)
-	{
-		++NumPostprocessingConstraintRules;
-	}
-
 	FRuleCreator RuleCreator(this);
 
 	if (ExternalForces)
@@ -1840,23 +1367,6 @@ void FClothConstraints::CreateForceBasedRules(const TSharedPtr<Softs::FMultiResC
 		// TODO: Linear System
 	}
 
-	if (MultiResConstraints)
-	{
-		constexpr bool bInitParticles = false;
-		RuleCreator.AddXPBDParallelInitRule<bInitParticles>(MultiResConstraints.Get());
-		constexpr bool bPostCollisions = false;
-		RuleCreator.AddPerIterationPBDConstraintRule_Apply<bPostCollisions>(MultiResConstraints.Get());
-	}
-
-	if (FineLODMultiResConstraint)
-	{
-		RuleCreator.AddPostSubstepConstraintRule(
-			[FineLODMultiResConstraint, this](Softs::FSolverParticlesRange& Particles, const Softs::FSolverReal Dt, const Softs::ESolverMode SolverMode)
-		{
-			FineLODMultiResConstraint->UpdateFineTargets(Particles);
-		});
-	}
-
 	if (XStretchBiasConstraints)
 	{
 		constexpr bool bInitParticles = false;
@@ -2114,16 +1624,6 @@ void FClothConstraints::CreatePBDRules()
 {
 	check(PBDEvolution);
 	check(ConstraintInitOffset == INDEX_NONE)
-#if !UE_BUILD_SHIPPING
-	if (bDisplayResidual)
-	{
-		NumConstraintRules++;
-	}
-	if (bWriteFinalResiduals || bDisplayResidual)
-	{
-		NumConstraintInits++;
-	}
-#endif
 	if (NumConstraintInits)
 	{
 		ConstraintInitOffset = PBDEvolution->AddConstraintInitRange(NumConstraintInits, false);
@@ -2140,12 +1640,6 @@ void FClothConstraints::CreatePBDRules()
 		PostCollisionConstraintRuleOffset = PBDEvolution->AddPostCollisionConstraintRuleRange(NumPostCollisionConstraintRules, false);
 	}
 	check(PostprocessingConstraintRuleOffset == INDEX_NONE);
-#if !UE_BUILD_SHIPPING	
-	if (bWriteFinalResiduals)
-	{
-		NumPostprocessingConstraintRules += 1;
-	}
-#endif
 	if (NumPostprocessingConstraintRules)
 	{
 		PostprocessingConstraintRuleOffset = PBDEvolution->AddConstraintPostprocessingsRange(NumPostprocessingConstraintRules, false);
@@ -2160,83 +1654,6 @@ void FClothConstraints::CreatePBDRules()
 	int32 ConstraintRuleIndex = 0;
 	int32 PostCollisionConstraintRuleIndex = 0;
 	int32 PostprocessingConstraintRuleIndex = 0;
-	#if !UE_BUILD_SHIPPING	
-	if (bDisplayResidual || bWriteFinalResiduals)
-	{
-		GSMainConstraint = MakeShared<Chaos::Softs::FGaussSeidelMainConstraints>(PBDEvolution->Particles(), false, false, (Softs::FSolverReal)1.2, 100);
-
-		ConstraintInits[ConstraintInitIndex++] =
-			[this](Softs::FSolverParticles& InParticles, const Softs::FSolverReal Dt)
-		{
-			this->GSMainConstraint->Init(Dt, InParticles);
-		};
-		if (bDisplayResidual)
-		{
-
-			ConstraintRules[ConstraintRuleIndex++] =
-				[this](Softs::FSolverParticles& InParticles, const Softs::FSolverReal Dt)
-			{
-				if (this->GSMainConstraint->DebugResidual && this->GSMainConstraint->PassedIters < MaxResidualIters)
-				{
-					UE_LOG(LogTemp, Warning, TEXT("Current Iteration As From Evolution: %d"), this->Evolution->GetIterations());
-					this->GSMainConstraint->ComputeNewtonResiduals(InParticles, Dt, bWriteResidual2File);
-					this->GSMainConstraint->PassedIters += 1;
-				}
-			};
-		}
-
-		if (XStretchBiasConstraints)
-		{
-			TArray<TArray<int32>> IncidentElements, IncidentElementsLocal;
-			GSMainConstraint->AddStaticConstraints(XStretchBiasConstraints->GetConstraintsArray(), IncidentElements, IncidentElementsLocal);
-			int32 StaticIndex = GSMainConstraint->AddStaticConstraintResidualAndHessianRange(1);
-
-			GSMainConstraint->StaticConstraintResidualAndHessian()[StaticIndex] = [this](const Softs::FSolverParticles& Particles, const int32 ElementIndex, const int32 ElementIndexLocal, const Softs::FSolverReal Dt, TVec3<Softs::FSolverReal>& ParticleResidual, Chaos::PMatrix<Softs::FSolverReal, 3, 3>& ParticleHessian)
-			{
-				this->XStretchBiasConstraints->AddStretchBiasElementResidualAndHessian(Particles, ElementIndex, ElementIndexLocal, Dt, ParticleResidual, ParticleHessian);
-			};
-		}
-
-		if (XAnisoBendingElementConstraints)
-		{
-			TArray<TArray<int32>> IncidentElements, IncidentElementsLocal;
-			GSMainConstraint->AddStaticConstraints(XAnisoBendingElementConstraints->GetConstraintsArray(), IncidentElements, IncidentElementsLocal);
-			int32 StaticIndex = GSMainConstraint->AddStaticConstraintResidualAndHessianRange(1);
-
-			GSMainConstraint->StaticConstraintResidualAndHessian()[StaticIndex] = [this](const Softs::FSolverParticles& Particles, const int32 ElementIndex, const int32 ElementIndexLocal, const Softs::FSolverReal Dt, TVec3<Softs::FSolverReal>& ParticleResidual, Chaos::PMatrix<Softs::FSolverReal, 3, 3>& ParticleHessian)
-			{
-				this->XAnisoBendingElementConstraints->AddAnisotropicBendingResidualAndHessian(Particles, ElementIndex, ElementIndexLocal, Dt, ParticleResidual, ParticleHessian);
-			};
-		}
-
-
-		if (XBendingElementConstraints)
-		{
-			TArray<TArray<int32>> IncidentElements, IncidentElementsLocal;
-			GSMainConstraint->AddStaticConstraints(XBendingElementConstraints->GetConstraintsArray(), IncidentElements, IncidentElementsLocal);
-			int32 StaticIndex = GSMainConstraint->AddStaticConstraintResidualAndHessianRange(1);
-
-			GSMainConstraint->StaticConstraintResidualAndHessian()[StaticIndex] = [this](const Softs::FSolverParticles& Particles, const int32 ElementIndex, const int32 ElementIndexLocal, const Softs::FSolverReal Dt, TVec3<Softs::FSolverReal>& ParticleResidual, Chaos::PMatrix<Softs::FSolverReal, 3, 3>& ParticleHessian)
-			{
-				this->XBendingElementConstraints->AddBendingResidualAndHessian(Particles, ElementIndex, ElementIndexLocal, Dt, ParticleResidual, ParticleHessian);
-			};
-		}
-
-		GSMainConstraint->InitStaticColor(PBDEvolution->Particles());
-
-		if (bWriteFinalResiduals)
-		{
-			FString file = FPaths::ProjectDir();
-			file.Append(TEXT("/DebugOutput/NewtonResidual.txt"));
-			FFileHelper::SaveStringToFile(FString(TEXT("Newton Norm\r\n")), *file);
-			PostprocessingConstraintRules[PostprocessingConstraintRuleIndex++] =
-				[this](Softs::FSolverParticles& InParticles, const Softs::FSolverReal Dt)
-			{
-				this->GSMainConstraint->ComputeNewtonResiduals(InParticles, Dt, true, nullptr);
-			};
-		}
-	}
-#endif
 
 	if (XStretchBiasConstraints)
 	{
@@ -2647,10 +2064,6 @@ void FClothConstraints::Update(
 	{
 		RepulsionConstraints->SetProperties(ConfigProperties);
 	}
-	if (MultiResConstraints)
-	{
-		MultiResConstraints->SetProperties(ConfigProperties, WeightMaps);
-	}
 
 	bool bUsePointBasedWindModel = false;
 	if (ExternalForces)
diff --git a/Engine/Plugins/ChaosCloth/Source/ChaosCloth/Private/ChaosCloth/ChaosClothVisualization.cpp b/Engine/Plugins/ChaosCloth/Source/ChaosCloth/Private/ChaosCloth/ChaosClothVisualization.cpp
--- a/Engine/Plugins/ChaosCloth/Source/ChaosCloth/Private/ChaosCloth/ChaosClothVisualization.cpp
+++ b/Engine/Plugins/ChaosCloth/Source/ChaosCloth/Private/ChaosCloth/ChaosClothVisualization.cpp
@@ -23,13 +23,13 @@
 #include "Chaos/PBDFlatWeightMap.h"
 #include "Chaos/PBDLongRangeConstraints.h"
 #include "Chaos/PBDSelfCollisionSphereConstraints.h"
+#include "Chaos/PBDSoftBodyCollisionConstraint.h"
 #include "Chaos/PBDSphericalConstraint.h"
 #include "Chaos/PBDSpringConstraints.h"
 #include "Chaos/PBDTriangleMeshCollisions.h"
 #include "Chaos/PBDUnilateralTetConstraints.h"
 #include "Chaos/SkinnedTriangleMesh.h"
 #include "Chaos/SoftsExternalForces.h"
-#include "Chaos/SoftsMultiResConstraints.h"
 #include "Chaos/XPBDBendingConstraints.h"
 #include "Chaos/XPBDSpringConstraints.h"
 #include "Chaos/XPBDAnisotropicBendingConstraints.h"
@@ -1719,71 +1719,6 @@ namespace Chaos
 		}
 	}
 
-	void FClothVisualizationNoGC::DrawMultiResConstraint(const FDrawContext& Context) const
-	{
-		if (!Solver)
-		{
-			return;
-		}
-
-		for (const FClothingSimulationCloth* const Cloth : Solver->GetCloths())
-		{
-			const int32 ParticleRangeId = Cloth->GetParticleRangeId(Solver);
-			if (ParticleRangeId == INDEX_NONE)
-			{
-				continue;
-			}
-			const FClothConstraints& ClothConstraints = Solver->GetClothConstraints(ParticleRangeId);
-			if (const Softs::FMultiResConstraints* const MultiResConstraints = ClothConstraints.GetMultiResConstraints().Get())
-			{
-
-				const int32 CoarseParticleRangeId = MultiResConstraints->GetCoarseSoftBodyId();
-				const FTriangleMesh& CoarseMesh = MultiResConstraints->GetCoarseMesh();
-				TConstArrayView<Softs::FSolverVec3> CoarsePositions = Solver->GetParticleXsView(CoarseParticleRangeId);
-				TConstArrayView<Softs::FSolverReal> CoarseInvMasses = Solver->GetParticleInvMassesView(CoarseParticleRangeId);
-
-				// Draw wired coarse mesh
-				static const FLinearColor DynamicColor = FColor::White;
-				static const FLinearColor KinematicColor = FColor::Purple;
-				for (const TVec3<int32>& Element : CoarseMesh.GetElements())
-				{
-					const FVector Pos0 = GetWorldPosition(*Solver, CoarsePositions[Element.X]);
-					const FVector Pos1 = GetWorldPosition(*Solver, CoarsePositions[Element.Y]);
-					const FVector Pos2 = GetWorldPosition(*Solver, CoarsePositions[Element.Z]);
-					const bool bIsKinematic0 = (CoarseInvMasses[Element.X] == (Softs::FSolverReal)0.);
-					const bool bIsKinematic1 = (CoarseInvMasses[Element.Y] == (Softs::FSolverReal)0.);
-					const bool bIsKinematic2 = (CoarseInvMasses[Element.Z] == (Softs::FSolverReal)0.);
-					if (bIsKinematic0 && bIsKinematic1 && bIsKinematic2)
-					{
-						continue;
-					}
-
-					DrawLine(Context.PDI, Pos0, Pos1, bIsKinematic0 && bIsKinematic1 ? KinematicColor : DynamicColor);
-					DrawLine(Context.PDI, Pos1, Pos2, bIsKinematic1 && bIsKinematic2 ? KinematicColor : DynamicColor);
-					DrawLine(Context.PDI, Pos2, Pos0, bIsKinematic2 && bIsKinematic0 ? KinematicColor : DynamicColor);
-				}
-
-				// Draw springs to targets
-				static const FLinearColor Red(0.3f, 0.f, 0.f);
-				static const FLinearColor Brown(0.1f, 0.05f, 0.f);
-				const TConstArrayView<Softs::FSolverVec3> Positions = Solver->GetParticleXsView(ParticleRangeId);
-				const TConstArrayView<Softs::FSolverReal> InvMasses = Solver->GetParticleInvMassesView(ParticleRangeId);
-				const TArray<Softs::FSolverVec3>& TargetPositions = MultiResConstraints->GetFineTargetPositions();
-				for (int32 Index = 0; Index < TargetPositions.Num(); ++Index)
-				{
-					if (InvMasses[Index] != (FSolverReal)0.f && MultiResConstraints->IsConstraintActive(Index))
-					{
-						const FVector P1 = GetWorldPosition(*Solver, Positions[Index]);
-						const FVector P2 = GetWorldPosition(*Solver, TargetPositions[Index]);
-
-						DrawPoint(Context.PDI, P2, Red, nullptr, 2.f);
-						DrawLine(Context.PDI, P1, P2, Brown);
-					}
-				}
-			}
-		}
-	}
-
 	void FClothVisualizationNoGC::DrawAnimNormals(const FDrawContext& Context, const FReal Length) const
 	{
 		if (!Solver)
diff --git a/Engine/Plugins/ChaosCloth/Source/ChaosCloth/Private/ChaosCloth/ChaosClothingSimulationCloth.cpp b/Engine/Plugins/ChaosCloth/Source/ChaosCloth/Private/ChaosCloth/ChaosClothingSimulationCloth.cpp
--- a/Engine/Plugins/ChaosCloth/Source/ChaosCloth/Private/ChaosCloth/ChaosClothingSimulationCloth.cpp
+++ b/Engine/Plugins/ChaosCloth/Source/ChaosCloth/Private/ChaosCloth/ChaosClothingSimulationCloth.cpp
@@ -47,7 +47,6 @@ namespace ClothingSimulationClothDefault
 	constexpr float MaxVelocity = TNumericLimits<float>::Max();
 	constexpr float MaxAcceleration = TNumericLimits<float>::Max();
 	constexpr float FictitiousAngularScale = Softs::FExternalForces::DefaultFictitiousAngularScale; // 1.f;
-	constexpr int32 MultiResCoarseLODIndex = INDEX_NONE;
 }
 
 namespace ClothingSimulationClothConsoleVariables
@@ -116,7 +115,6 @@ struct FClothingSimulationCloth::FLODData
 		int32 LODIndex;
 		int32 ParticleRangeId;
 		FTriangleMesh OffsetTriangleMesh; // Only used if using PBD solver
-		int32 MultiResCoarseLODIndex = INDEX_NONE;
 
 		TMap<FName, FClothingAccessoryMeshData> AccessoryMeshData;
 	};
@@ -161,11 +159,6 @@ struct FClothingSimulationCloth::FLODData
 
 	void UpdateNormals(FClothingSimulationSolver* Solver) const;
 
-
-	// These are only used in Add, so no need to cache them
-	UE_CHAOS_DECLARE_INDEXLESS_PROPERTYCOLLECTION_NAME(MultiResCoarseLODIndex, int32);
-	UE_CHAOS_DECLARE_INDEXLESS_PROPERTYCOLLECTION_NAME(IsCoarseMultiResLOD, bool);
-
 	// This can update every frame, so should be cached
 	UE_CHAOS_DECLARE_PROPERTYCOLLECTION_NAME(ActiveMorphTarget, float);
 
@@ -367,45 +360,6 @@ void FClothingSimulationCloth::FLODData::Add(FClothingSimulationSolver* Solver,
 	check(Cloth->Config);
 	const Softs::FCollectionPropertyFacade& ConfigProperties = Cloth->Config->GetProperties(InLODIndex);
 
-	// Gather multires constraint data.
-	TSharedPtr<Softs::FMultiResConstraints> FineLODMultiResConstraint;
-	const FTriangleMesh* CoarseLODTriangleMesh = nullptr;
-	int32 CoarseLODParticleRangeId = INDEX_NONE;
-	if (!Solver->IsLegacySolver())
-	{
-		// Multi-Res isn't supported by legacy solver
-		if (InLODIndex == 0)
-		{
-			// Only allow LOD0 to be a fine LOD for now.
-			const int32 MultiResCoarseLODIndex = GetMultiResCoarseLODIndex(ConfigProperties, ClothingSimulationClothDefault::MultiResCoarseLODIndex);
-			if (MultiResCoarseLODIndex != INDEX_NONE && MultiResCoarseLODIndex != InLODIndex)
-			{
-				if (Cloth->Config->IsValidLOD(MultiResCoarseLODIndex) && Cloth->LODData.IsValidIndex(MultiResCoarseLODIndex))
-				{
-					// Check if coarse lod is setup correctly.
-					const Softs::FCollectionPropertyFacade& CoarseConfigProperties = Cloth->Config->GetProperties(MultiResCoarseLODIndex);
-					if (GetIsCoarseMultiResLOD(CoarseConfigProperties, false))
-					{
-						CoarseLODTriangleMesh = &Cloth->LODData[MultiResCoarseLODIndex]->NoOffsetTriangleMesh;
-						CoarseLODParticleRangeId = Cloth->LODData[MultiResCoarseLODIndex]->SolverData.FindChecked(Solver).ParticleRangeId;
-						SolverDatum.MultiResCoarseLODIndex = MultiResCoarseLODIndex;
-					}
-				}
-			}
-		}
-		else if (GetIsCoarseMultiResLOD(ConfigProperties, false))
-		{
-			// check that fine lod is setup correctly.
-			const Softs::FCollectionPropertyFacade& FineConfigProperties = Cloth->Config->GetProperties(0);
-			const int32 MultiResCoarseLODIndex = GetMultiResCoarseLODIndex(FineConfigProperties, ClothingSimulationClothDefault::MultiResCoarseLODIndex);
-			if (MultiResCoarseLODIndex == InLODIndex)
-			{
-				const int32 FineLODParticleRangeId = Cloth->LODData[0]->SolverData.FindChecked(Solver).ParticleRangeId;
-				FineLODMultiResConstraint = Solver->GetClothConstraints(FineLODParticleRangeId).GetMultiResConstraints();
-			}
-		}
-	}
-
 	// Setup solver constraints
 	FClothConstraints& ClothConstraints = Solver->GetClothConstraints(ParticleRangeId);
 
@@ -425,7 +379,7 @@ void FClothingSimulationCloth::FLODData::Add(FClothingSimulationSolver* Solver,
 
 	// Create constraints
 	const bool bEnabled = false;  // Set constraint disabled by default
-	ClothConstraints.AddRules(ConfigProperties, TriangleMesh, &PatternData, WeightMaps, VertexSets, FaceSets, FaceIntMaps, Tethers, MeshScale, bEnabled, CoarseLODTriangleMesh, CoarseLODParticleRangeId, FineLODMultiResConstraint, ComplexColliders, ManagedArrayCollection, &SolverDatum.AccessoryMeshData);
+	ClothConstraints.AddRules(ConfigProperties, TriangleMesh, &PatternData, WeightMaps, VertexSets, FaceSets, FaceIntMaps, Tethers, MeshScale, bEnabled, ComplexColliders, ManagedArrayCollection, &SolverDatum.AccessoryMeshData);
 
 	// Update LOD stats
 	const TConstArrayView<Softs::FSolverReal> InvMasses(Solver->GetParticleInvMasses(ParticleRangeId), NumParticles);
@@ -781,13 +735,12 @@ void FClothingSimulationCloth::Add(FClothingSimulationSolver* Solver)
 	int32& LODIndex = LODIndices.Add(Solver);
 	LODIndex = INDEX_NONE;
 
-	// Add all particles first and in reverse order. This is necessary so that any multires coarse lods soft bodies are added first, and all particle offsets are setup
-	// when adding the LOD constraints.
+	// Add all particles first and in reverse order. (was doing this for multi-res constraints, but now just leaving it)
 	for (int32 Index = LODData.Num() - 1; Index >= 0; --Index)
 	{
 		LODData[Index]->AddParticles(Solver, this, Index);
 	}
-	// Now add the LODs themselves. These need to go in normal order since the coarse lod needs the fine lod constraints.
+	// Now add the LODs themselves.
 	for (int32 Index = 0; Index < LODData.Num(); ++Index)
 	{
 		LODData[Index]->Add(Solver, this, Index);
@@ -1074,33 +1027,13 @@ void FClothingSimulationCloth::Update(FClothingSimulationSolver* Solver)
 
 	Mesh->Update(Solver, PrevLODIndex, LODIndex, PrevParticleRangeId, ParticleRangeId, CurrentLODMorphTargetIndex, CurrentLODMorphTargetWeight);
 
-	const int32 CoarseLODIndex = LODIndex != INDEX_NONE ? LODData[LODIndex]->SolverData.FindChecked(Solver).MultiResCoarseLODIndex : INDEX_NONE;
-	const int32 CoarseParticleRangeId = CoarseLODIndex != INDEX_NONE ? GetParticleRangeId(Solver, CoarseLODIndex) : INDEX_NONE;
-	if (CoarseParticleRangeId != INDEX_NONE)
-	{
-		LODData[CoarseLODIndex]->UpdateCachedProperties(*Mesh, CoarseLODIndex, Config->GetProperties(CoarseLODIndex));
-		LODData[CoarseLODIndex]->Enable(Solver, true);
-		// TODO: interpolate/ reset when LOD switching to enable multires
-		Mesh->Update(Solver, CoarseLODIndex, CoarseLODIndex, CoarseParticleRangeId, CoarseParticleRangeId, LODData[CoarseLODIndex]->ActiveMorphTarget, LODData[CoarseLODIndex]->ActiveMorphTargetWeight);
-		LODData[CoarseLODIndex]->UpdateAccessoryMeshes(Solver);
-	}
-
 	// LOD Switching
 	if (LODIndex != PrevLODIndex)
 	{
 		if (PrevLODIndex != INDEX_NONE)
 		{
-			if (PrevLODIndex != CoarseLODIndex)
-			{
-				// Disable previous LOD's particles
-				LODData[PrevLODIndex]->Enable(Solver, false);
-			}
-			const int32 PrevCoarseLODIndex = LODData[PrevLODIndex]->SolverData.FindChecked(Solver).MultiResCoarseLODIndex;
-			if (PrevCoarseLODIndex != INDEX_NONE && PrevCoarseLODIndex != CoarseLODIndex && PrevCoarseLODIndex != LODIndex)
-			{
-				// Disable previous coarse LOD's particles
-				LODData[PrevCoarseLODIndex]->Enable(Solver, false);
-			}
+			// Disable previous LOD's particles
+			LODData[PrevLODIndex]->Enable(Solver, false);
 		}
 		if (LODIndex != INDEX_NONE)
 		{
@@ -1172,10 +1105,6 @@ void FClothingSimulationCloth::Update(FClothingSimulationSolver* Solver)
 				Collider->ResetStartPose(Solver, this);
 			}
 
-			if (CoarseLODIndex != INDEX_NONE && CoarseParticleRangeId != INDEX_NONE)
-			{
-				LODData[CoarseLODIndex]->ResetStartPose(Solver);
-			}
 			UE_LOG(LogChaosCloth, Verbose, TEXT("Cloth in group Id %d Needs reset."), GroupId);
 		}
 		else if (bNeedsTeleport)
@@ -1241,10 +1170,6 @@ void FClothingSimulationCloth::Update(FClothingSimulationSolver* Solver)
 		if (!Solver->IsLegacySolver())
 		{
 			Solver->SetProperties(ParticleRangeId, *ConfigProperties, LODData[LODIndex]->WeightMaps);
-			if (CoarseLODIndex != INDEX_NONE && CoarseParticleRangeId != INDEX_NONE)
-			{
-				Solver->SetProperties(CoarseParticleRangeId, Config->GetProperties(CoarseLODIndex), LODData[CoarseLODIndex]->WeightMaps);
-			}
 		}
 		else
 		{
@@ -1287,10 +1212,6 @@ void FClothingSimulationCloth::Update(FClothingSimulationSolver* Solver)
 		if (!Config->IsLegacySingleLOD())
 		{
 			LODData[LODIndex]->Update(Solver, this);
-			if (CoarseLODIndex != INDEX_NONE)
-			{
-				LODData[CoarseLODIndex]->Update(Solver, this);
-			}
 		}
 	}
 
diff --git a/Engine/Plugins/ChaosCloth/Source/ChaosCloth/Public/ChaosCloth/ChaosClothConstraints.h b/Engine/Plugins/ChaosCloth/Source/ChaosCloth/Public/ChaosCloth/ChaosClothConstraints.h
--- a/Engine/Plugins/ChaosCloth/Source/ChaosCloth/Public/ChaosCloth/ChaosClothConstraints.h
+++ b/Engine/Plugins/ChaosCloth/Source/ChaosCloth/Public/ChaosCloth/ChaosClothConstraints.h
@@ -8,8 +8,6 @@
 #include "Chaos/PBDLongRangeConstraints.h"
 #include "Chaos/PBDCollisionSpringConstraintsBase.h"
 #include "Chaos/XPBDPlanarConstraints.h"
-#include "Chaos/Deformable/GaussSeidelMainConstraint.h"
-#include "Chaos/Deformable/GaussSeidelCorotatedCodimensionalConstraints.h"
 #include "Chaos/PBDExtremeDeformationConstraints.h"
 
 namespace Buoyancy
@@ -69,6 +67,22 @@ namespace Chaos
 		// ---- End of Solver interface ----
 
 		// ---- Cloth interface ----
+		void AddRules(
+			const Softs::FCollectionPropertyConstFacade& ConfigProperties,
+			const FTriangleMesh& TriangleMesh,
+			const FClothingPatternData* PatternData,
+			const TMap<FString, TConstArrayView<FRealSingle>>& WeightMaps,
+			const TMap<FString, const TSet<int32>*>& VertexSets,
+			const TMap<FString, const TSet<int32>*>& FaceSets,
+			const TMap<FString, TConstArrayView<int32>>& FaceIntMaps,
+			const TArray<TConstArrayView<TTuple<int32, int32, FRealSingle>>>& Tethers,
+			Softs::FSolverReal MeshScale,
+			bool bEnabled,
+			const TArray<const FClothComplexColliders*>& ComplexColliders = TArray<const FClothComplexColliders*>(),
+			const TSharedPtr<const FManagedArrayCollection>& ManagedArrayCollection = TSharedPtr<const FManagedArrayCollection>(),
+			const TMap<FName, FClothingAccessoryMeshData>* AccessoryMeshes = nullptr);
+
+		UE_DEPRECATED(5.8, "MultiRes constraint is no longer supported. Use other version of AddRules instead.")
 		void AddRules(
 			const Softs::FCollectionPropertyConstFacade& ConfigProperties,
 			const FTriangleMesh& TriangleMesh,
@@ -85,7 +99,10 @@ namespace Chaos
 			const TSharedPtr<Softs::FMultiResConstraints>& FineLODMultiResConstraint = TSharedPtr<Softs::FMultiResConstraints>(nullptr),
 			const TArray<const FClothComplexColliders*>& ComplexColliders = TArray<const FClothComplexColliders*>(),
 			const TSharedPtr<const FManagedArrayCollection>& ManagedArrayCollection = TSharedPtr<const FManagedArrayCollection>(),
-			const TMap<FName, FClothingAccessoryMeshData>* AccessoryMeshes = nullptr);
+			const TMap<FName, FClothingAccessoryMeshData>* AccessoryMeshes = nullptr)
+		{
+			AddRules(ConfigProperties, TriangleMesh, PatternData, WeightMaps, VertexSets, FaceSets, FaceIntMaps, Tethers, MeshScale, bEnabled, ComplexColliders, ManagedArrayCollection, AccessoryMeshes);
+		}
 
 		void Update(
 			const Softs::FCollectionPropertyConstFacade& ConfigProperties,
@@ -139,7 +156,8 @@ namespace Chaos
 		const TSharedPtr<Softs::FExternalForces>& GetExternalForces() const { return ExternalForces; }
 		const TSharedPtr<Softs::FPBDSoftBodyCollisionConstraint>& GetCollisionConstraint() const { return CollisionConstraint; }
 		const TSharedPtr<Softs::FPBDSkinnedTriangleMeshCollisionConstraints>& GetSkinnedTriangleCollisionsConstraint() const { return SkinnedTriangleCollisionsConstraint; }
-		const TSharedPtr<Softs::FMultiResConstraints>& GetMultiResConstraints() const { return MultiResConstraints; }
+		UE_DEPRECATED(5.8, "MultiResConstraints have been deprecated")
+		const TSharedPtr<Softs::FMultiResConstraints>& GetMultiResConstraints() const { static TSharedPtr<Softs::FMultiResConstraints> Dummy; return Dummy; }
 		const TSharedPtr<Softs::FXPBDVertexConstraints>& GetClothVertexSpringConstraints() const { return ClothVertexSpringConstraints; };
 		const TSharedPtr<Softs::FXPBDVertexFaceConstraints>& GetClothVertexFaceSpringConstraints() const { return ClothVertexFaceSpringConstraints; }
 		const TSharedPtr<Softs::FXPBDFaceConstraints>& GetClothFaceSpringConstraints() const { return ClothFaceSpringConstraints; }
@@ -214,20 +232,12 @@ namespace Chaos
 			Softs::FSolverReal MeshScale,
 			const TMap<FString, TConstArrayView<FRealSingle>>& WeightMaps,
 			const TArray<const FClothComplexColliders*>& ComplexColliders);
-		void CreateMultiresConstraint(
-			const Softs::FCollectionPropertyConstFacade& ConfigProperties,
-			const TMap<FString, TConstArrayView<FRealSingle>>& WeightMaps,
-			const FTriangleMesh& TriangleMesh,
-			const FTriangleMesh* MultiResCoarseLODMesh,
-			const int32 MultiResCoarseLODParticleRangeId);
 		void CreateClothClothConstraints(
 			const Softs::FCollectionPropertyConstFacade& ConfigProperties,
 			const TSharedPtr<const FManagedArrayCollection>& ManagedArrayCollection);
 
-		void CreateForceBasedRules(const TSharedPtr<Softs::FMultiResConstraints>& FineLODMultiResConstraint);
+		void CreateForceBasedRules();
 		void CreatePBDRules();
-		void CreateGSRules();
-		void GetGSNumRules();
 
 		TSharedPtr<Softs::FPBDEdgeSpringConstraints> EdgeConstraints;
 		TSharedPtr<Softs::FXPBDEdgeSpringConstraints> XEdgeConstraints;
@@ -248,9 +258,6 @@ namespace Chaos
 		TSharedPtr<Softs::FPBDCollisionSpringConstraints> SelfCollisionConstraints;
 		TSharedPtr<Softs::FPBDTriangleMeshIntersections> SelfIntersectionConstraints;
 		TSharedPtr<Softs::FPBDSelfCollisionSphereConstraints> SelfCollisionSphereConstraints;
-		TSharedPtr<Softs::FGaussSeidelMainConstraints> GSMainConstraint;
-		TSharedPtr<Softs::FGaussSeidelCorotatedTriangleConstraints> GSCorotatedCodimensionalConstraint;
-		TSharedPtr<Softs::FMultiResConstraints> MultiResConstraints;
 		TSharedPtr<Softs::FPBDExtremeDeformationConstraints> ExtremeDeformationConstraints;
 		//~ Begin Force-based solver only constraints
 		Softs::FSolverVec3 SolverWindVelocity; // Set from solver and added to wind from the config
diff --git a/Engine/Plugins/ChaosCloth/Source/ChaosCloth/Public/ChaosCloth/ChaosClothVisualization.h b/Engine/Plugins/ChaosCloth/Source/ChaosCloth/Public/ChaosCloth/ChaosClothVisualization.h
--- a/Engine/Plugins/ChaosCloth/Source/ChaosCloth/Public/ChaosCloth/ChaosClothVisualization.h
+++ b/Engine/Plugins/ChaosCloth/Source/ChaosCloth/Public/ChaosCloth/ChaosClothVisualization.h
@@ -88,7 +88,6 @@ namespace Chaos
 		CHAOSCLOTH_API void DrawBounds(const FDrawContext& Context = FDrawContext()) const;
 		CHAOSCLOTH_API void DrawGravity(const FDrawContext& Context = FDrawContext()) const;
 		CHAOSCLOTH_API void DrawFictitiousAngularForces(const FDrawContext& Context = FDrawContext()) const;
-		CHAOSCLOTH_API void DrawMultiResConstraint(const FDrawContext& Context = FDrawContext()) const;
 		CHAOSCLOTH_API void DrawClothClothConstraints(const FDrawContext& Context = FDrawContext()) const;
 		CHAOSCLOTH_API void DrawTeleportReset(const FDrawContext& Context = FDrawContext()) const;
 		CHAOSCLOTH_API void DrawExtremlyDeformedEdges(const FDrawContext& Context = FDrawContext()) const;
@@ -246,10 +245,9 @@ namespace Chaos
 		{
 			return DrawFictitiousAngularForces(FDrawContext(PDI));
 		}
-		UE_DEPRECATED(5.8, "Use Context version of this method instead.")
-		void DrawMultiResConstraint(FPrimitiveDrawInterface* PDI) const
+		UE_DEPRECATED(5.8, "MultiResConstraints have been deprecated.")
+		void DrawMultiResConstraint(FPrimitiveDrawInterface* PDI = nullptr) const
 		{
-			return DrawMultiResConstraint(FDrawContext(PDI));
 		}
 		UE_DEPRECATED(5.8, "Use Context version of this method instead.")
 		void DrawClothClothConstraints(FPrimitiveDrawInterface* PDI) const
@@ -308,7 +306,6 @@ namespace Chaos
 		void DrawBounds(const FDrawContext& /*Context*/ = FDrawContext()) const {}
 		void DrawGravity(const FDrawContext& /*Context*/ = FDrawContext()) const {}
 		void DrawFictitiousAngularForces(const FDrawContext& /*Context*/ = FDrawContext()) const {}
-		void DrawMultiResConstraint(const FDrawContext& /*Context*/ = FDrawContext()) const {}
 		void DrawClothClothConstraints(const FDrawContext& /*Context*/ = FDrawContext()) const {}
 		void DrawExtremlyDeformedEdges(const FDrawContext& /*Context*/ = FDrawContext()) const {}
 		void DrawAccessoryMesh(const FDrawContext& /*Context*/ = FDrawContext(), const FName& /*AccessoryMeshName*/ = NAME_None) const {}
@@ -372,8 +369,8 @@ namespace Chaos
 		void DrawGravity(FPrimitiveDrawInterface* /*PDI*/) const {}
 		UE_DEPRECATED(5.8, "Use Context version of this method instead.")
 		void DrawFictitiousAngularForces(FPrimitiveDrawInterface* /*PDI*/) const {}
-		UE_DEPRECATED(5.8, "Use Context version of this method instead.")
-		void DrawMultiResConstraint(FPrimitiveDrawInterface* /*PDI*/) const {}
+		UE_DEPRECATED(5.8, "MultiRes Constraints have been deprecated.")
+		void DrawMultiResConstraint(FPrimitiveDrawInterface* /*PDI*/ = nullptr) const {}
 		UE_DEPRECATED(5.8, "Use Context version of this method instead.")
 		void DrawClothClothConstraints(FPrimitiveDrawInterface* /*PDI*/) const {}
 		UE_DEPRECATED(5.8, "Use Context version of this method instead.")
diff --git a/Engine/Plugins/ChaosClothAssetDataflowNodes/Source/ChaosClothAssetDataflowNodes/Private/ChaosClothAsset/SimulationMultiResConfigNode.cpp b/Engine/Plugins/ChaosClothAssetDataflowNodes/Source/ChaosClothAssetDataflowNodes/Private/ChaosClothAsset/SimulationMultiResConfigNode.cpp
--- a/Engine/Plugins/ChaosClothAssetDataflowNodes/Source/ChaosClothAssetDataflowNodes/Private/ChaosClothAsset/SimulationMultiResConfigNode.cpp
+++ b/Engine/Plugins/ChaosClothAssetDataflowNodes/Source/ChaosClothAssetDataflowNodes/Private/ChaosClothAsset/SimulationMultiResConfigNode.cpp
@@ -19,6 +19,7 @@ FChaosClothAssetSimulationMultiResConfigNode::FChaosClothAssetSimulationMultiRes
 
 void FChaosClothAssetSimulationMultiResConfigNode::AddProperties(FPropertyHelper& PropertyHelper) const
 {
+	PRAGMA_DISABLE_DEPRECATION_WARNINGS  // SetProperty functions are templated and cause deprecation warnings with the now deprecated v1
 	if (bIsFineLOD)
 	{
 		PropertyHelper.SetProperty(this, &MultiResCoarseLODIndex);
@@ -30,5 +31,6 @@ void FChaosClothAssetSimulationMultiResConfigNode::AddProperties(FPropertyHelper
 	{
 		PropertyHelper.SetPropertyBool(this, &bIsCoarseMultiResLOD);
 	}
+	PRAGMA_ENABLE_DEPRECATION_WARNINGS
 }
 
diff --git a/Engine/Plugins/ChaosClothAssetDataflowNodes/Source/ChaosClothAssetDataflowNodes/Public/ChaosClothAsset/SimulationMultiResConfigNode.h b/Engine/Plugins/ChaosClothAssetDataflowNodes/Source/ChaosClothAssetDataflowNodes/Public/ChaosClothAsset/SimulationMultiResConfigNode.h
--- a/Engine/Plugins/ChaosClothAssetDataflowNodes/Source/ChaosClothAssetDataflowNodes/Public/ChaosClothAsset/SimulationMultiResConfigNode.h
+++ b/Engine/Plugins/ChaosClothAssetDataflowNodes/Source/ChaosClothAssetDataflowNodes/Public/ChaosClothAsset/SimulationMultiResConfigNode.h
@@ -6,8 +6,8 @@
 #include "SimulationMultiResConfigNode.generated.h"
 
 /** Experimental Solver multires configuration node. */
-USTRUCT(Meta = (DataflowCloth, Experimental))
-struct FChaosClothAssetSimulationMultiResConfigNode : public FChaosClothAssetSimulationBaseConfigNode
+USTRUCT(Meta = (DataflowCloth, Deprecated=5.8, Experimental))
+struct UE_DEPRECATED(5.8, "MultiRes Constraints have been deprecated.")FChaosClothAssetSimulationMultiResConfigNode : public FChaosClothAssetSimulationBaseConfigNode
 {
 	GENERATED_USTRUCT_BODY()
 	DATAFLOW_NODE_DEFINE_INTERNAL(FChaosClothAssetSimulationMultiResConfigNode, "SimulationMultiResConfig", "Cloth", "Cloth Simulation MultiRes Config")
diff --git a/Engine/Plugins/ChaosClothAssetEditorCore/Source/ChaosClothAssetEditor/Private/ChaosClothAsset/ClothEditorSimulationVisualization.cpp b/Engine/Plugins/ChaosClothAssetEditorCore/Source/ChaosClothAssetEditor/Private/ChaosClothAsset/ClothEditorSimulationVisualization.cpp
--- a/Engine/Plugins/ChaosClothAssetEditorCore/Source/ChaosClothAssetEditor/Private/ChaosClothAsset/ClothEditorSimulationVisualization.cpp
+++ b/Engine/Plugins/ChaosClothAssetEditorCore/Source/ChaosClothAssetEditor/Private/ChaosClothAsset/ClothEditorSimulationVisualization.cpp
@@ -464,13 +464,6 @@ const FVisualizationOption FVisualizationOption::OptionData[] =
 		}),
 		LOCTEXT("ChaosVisName_DrawTeleportReset", "Teleport/Reset"), 
 		LOCTEXT("ChaosVisName_DrawTeleportReset_Tooltip", "Draw teleport/reset status.")),
-	FVisualizationOption(
-		FClothVisualizationDebugDraw::CreateLambda([](const FClothEditorSimulationVisualization&, const ::Chaos::FClothVisualizationNoGC& Visualization, const ::Chaos::FClothVisualizationNoGC::FDrawContext& Context)
-		{
-			Visualization.DrawMultiResConstraint(Context);
-		}),
-		LOCTEXT("ChaosVisName_DrawMultiResConstraint", "Draw Multi Res Constraint"), 
-		LOCTEXT("ChaosVisName_DrawMultiResConstraint_Tooltip", "Draw multi res constraint coarse mesh and targets.")),
 	FVisualizationOption(
 		FClothVisualizationDebugDraw::CreateLambda([](const FClothEditorSimulationVisualization& EditorVisualization, const ::Chaos::FClothVisualizationNoGC& Visualization, const ::Chaos::FClothVisualizationNoGC::FDrawContext& Context)
 		{
